Initial Chooser implementation
In an experimental plugin based on Data Interfaces system

diff --git a/Engine/Plugins/Experimental/Chooser/Chooser.uplugin b/Engine/Plugins/Experimental/Chooser/Chooser.uplugin
new file mode 100644
--- /dev/null
+++ b/Engine/Plugins/Experimental/Chooser/Chooser.uplugin
@@ -0,0 +1,37 @@
+{
+	"FileVersion" : 3,
+	"Version" : 1,
+	"VersionName" : "0.1",
+	"FriendlyName" : "Chooser",
+	"Description" : "Chooser",
+	"Category" : "Animation",
+	"CreatedBy" : "Epic Games, Inc.",
+	"CreatedByURL" : "http://epicgames.com",
+	"DocsURL" : "",
+	"MarketplaceURL" : "",
+	"SupportURL" : "",
+	"EnabledByDefault" : false,
+	"CanContainContent" : false,
+	"IsBetaVersion" : true,
+	"IsExperimentalVersion" : true,
+	"Installed" : false,
+	"Plugins": [
+		{
+			"Name": "DataInterfaceGraph",
+			"Enabled": true
+		}
+	],
+	"Modules" :
+	[
+		{
+			"Name" : "Chooser",
+			"Type" : "Runtime",
+			"LoadingPhase" : "Default"
+		},
+		{
+			"Name" : "ChooserEditor",
+			"Type" : "Editor",
+			"LoadingPhase" : "Default"
+		}
+	]
+}
diff --git a/Engine/Plugins/Experimental/Chooser/Source/Chooser/Chooser.Build.cs b/Engine/Plugins/Experimental/Chooser/Source/Chooser/Chooser.Build.cs
new file mode 100644
--- /dev/null
+++ b/Engine/Plugins/Experimental/Chooser/Source/Chooser/Chooser.Build.cs
@@ -0,0 +1,47 @@
+// Copyright Epic Games, Inc. All Rights Reserved.
+
+namespace UnrealBuildTool.Rules
+{
+	public class Chooser : ModuleRules
+	{
+		public Chooser(ReadOnlyTargetRules Target) : base(Target)
+		{
+			PublicIncludePaths.AddRange(
+				new string[] {
+					// ... add public include paths required here ...
+				}
+			);
+			
+			PrivateIncludePaths.AddRange(
+				new string[] {
+					// ... add other private include paths required here ...
+				}
+			);
+			
+			PublicDependencyModuleNames.AddRange(
+				new string[]
+				{
+					// ... add other public dependencies that you statically link with here ...
+					"DataInterface",
+				}
+			);
+
+			PrivateDependencyModuleNames.AddRange(
+				new string[]
+				{
+					"Core",
+					"CoreUObject",
+					"Engine",
+					// ... add private dependencies that you statically link with here ...
+				}
+			);
+
+			DynamicallyLoadedModuleNames.AddRange(
+				new string[]
+				{
+					// ... add any modules that your module loads dynamically here ...
+				}
+			);
+		}
+	}
+}
\ No newline at end of file
diff --git a/Engine/Plugins/Experimental/Chooser/Source/Chooser/Internal/Chooser.h b/Engine/Plugins/Experimental/Chooser/Source/Chooser/Internal/Chooser.h
new file mode 100644
--- /dev/null
+++ b/Engine/Plugins/Experimental/Chooser/Source/Chooser/Internal/Chooser.h
@@ -0,0 +1,142 @@
+#pragma once
+
+#include "CoreMinimal.h"
+#include "UObject/Object.h"
+#include "IDataInterface.h"
+#include "DataInterface_Object.h"
+#include "Chooser.generated.h"
+
+UINTERFACE(NotBlueprintType, meta = (CannotImplementInterfaceInBlueprint))
+class CHOOSER_API UChooserColumn : public UInterface
+{
+	GENERATED_BODY()
+};
+
+class CHOOSER_API IChooserColumn 
+{
+	GENERATED_BODY()
+
+public:
+	virtual void Filter(const UE::DataInterface::FContext& Context, const TArray<uint32>& IndexListIn, TArray<uint32>& IndexListOut) {};
+	virtual FName GetDisplayName() { return "ChooserColumn"; }
+	virtual void SetDisplayName(FName Name) { }
+	virtual void SetNumRows(uint32 NumRows) {}
+	virtual void DeleteRows(const TArray<uint32> & RowIndices) {}
+};
+
+UCLASS()
+class CHOOSER_API UChooserColumnBool : public UObject, public IChooserColumn
+{
+	GENERATED_BODY()
+	public:
+	UChooserColumnBool() { }
+	
+	UPROPERTY(EditAnywhere, Category = "Editor")
+	FName DisplayName = "Bool Column";
+	
+	UPROPERTY(EditAnywhere, Meta=(DataInterfaceType="bool"), Category = "Input")
+	TScriptInterface<IDataInterface> Value;
+
+	UPROPERTY(EditAnywhere, Category=Runtime)
+	// array of results (cells for this column for each row in the table)
+	// should match the length of the Results array 
+	TArray<bool> RowValues;
+	
+	virtual void Filter(const UE::DataInterface::FContext& Context, const TArray<uint32>& IndexListIn, TArray<uint32>& IndexListOut) override;
+
+	// todo: macro boilerplate
+	virtual void SetNumRows(uint32 NumRows) override { RowValues.SetNum(NumRows); }
+	virtual FName GetDisplayName() override { return DisplayName; }
+	virtual void SetDisplayName(FName Name) override { DisplayName = Name; }
+	virtual void DeleteRows(const TArray<uint32> & RowIndices )
+	{
+		for(uint32 Index : RowIndices)
+		{
+			RowValues.RemoveAt(Index);
+		}
+	}
+};
+
+USTRUCT()
+struct FChooserFloatRangeRowData
+{
+	GENERATED_BODY()
+	
+	UPROPERTY(EditAnywhere, Category=Runtime)
+	float Min=0;
+	
+	UPROPERTY(EditAnywhere, Category=Runtime)
+	float Max=0;
+};
+
+UCLASS()
+class CHOOSER_API UChooserColumnFloatRange : public UObject, public IChooserColumn
+{
+	GENERATED_BODY()
+	public:
+	UChooserColumnFloatRange() { }
+	
+	UPROPERTY(EditAnywhere, Category = "Editor")
+	FName DisplayName = "Float Range Column";
+	
+	UPROPERTY(EditAnywhere, Meta=(DataInterfaceType="float"), Category = "Input")
+	TScriptInterface<IDataInterface> Value;
+
+	UPROPERTY(EditAnywhere, Category=Runtime)
+	// array of results (cells for this column for each row in the table)
+	// should match the length of the Results array 
+	TArray<FChooserFloatRangeRowData> RowValues;
+	
+	virtual void Filter(const UE::DataInterface::FContext& Context, const TArray<uint32>& IndexListIn, TArray<uint32>& IndexListOut) override;
+	virtual void SetNumRows(uint32 NumRows) { RowValues.SetNum(NumRows); }
+	virtual FName GetDisplayName() override { return DisplayName; }
+	virtual void SetDisplayName(FName Name) override { DisplayName = Name; }
+	virtual void DeleteRows(const TArray<uint32> & RowIndices )
+	{
+		for(uint32 Index : RowIndices)
+		{
+			RowValues.RemoveAt(Index);
+		}
+	}
+};
+
+UCLASS(MinimalAPI)
+class UChooserTable : public UObject
+{
+	GENERATED_UCLASS_BODY()
+public:
+	UChooserTable() {}
+	
+	UPROPERTY(EditAnywhere, Meta=(EditInline="true"), Category = "Editor")
+	FName ResultType = "Object"; // todo: drop-down, populated by all possible return type names from all implemented data interfaces
+	// todo: struct type for Object return types
+
+	UPROPERTY(EditAnywhere, Meta=(EditInline="true"), Category = "Runtime")
+	TArray<TScriptInterface<IChooserColumn>> Columns;
+
+	// array of results (rows of table)
+	UPROPERTY(EditAnywhere, Meta=(EditInline="true"), Category = "Runtime")
+	TArray<TScriptInterface<IDataInterface>> Results;
+};
+
+UCLASS()
+class CHOOSER_API UDataInterface_EvaluateChooser : public UObject, public IDataInterface
+{
+	GENERATED_BODY()
+
+	// virtual bool GetObject(const UE::DataInterface::FContext& DataContext) const final override;
+
+	/** Get the return type name, used for dynamic type checking */
+    virtual FName GetReturnTypeNameImpl() const { return Chooser ? Chooser->ResultType : NAME_None; } 
+
+    /** Get the return type struct, used for dynamic type checking. If the result is not a struct, this should return nullptr */
+    virtual const UScriptStruct* GetReturnTypeStructImpl() const { return nullptr; } // todo, get type struct from chooser once it's added
+    
+    /** Get the value for this interface. @return true if successful, false if unsuccessful. */
+    virtual bool GetDataImpl(const UE::DataInterface::FContext& Context) const;
+
+	public:
+	
+	UPROPERTY(EditAnywhere, Category="Parameters")
+	TObjectPtr<UChooserTable> Chooser;
+};
\ No newline at end of file
diff --git a/Engine/Plugins/Experimental/Chooser/Source/Chooser/Internal/DataInterface_Object.h b/Engine/Plugins/Experimental/Chooser/Source/Chooser/Internal/DataInterface_Object.h
new file mode 100644
--- /dev/null
+++ b/Engine/Plugins/Experimental/Chooser/Source/Chooser/Internal/DataInterface_Object.h
@@ -0,0 +1,24 @@
+// Copyright Epic Games, Inc. All Rights Reserved.
+
+#pragma once
+
+#include "CoreMinimal.h"
+#include "UObject/Object.h"
+#include "DataInterfaceContext.h"
+#include "DataInterfaceTypes.h"
+#include "IDataInterface.h"
+#include "DataInterface_Object.generated.h"
+
+UCLASS()
+class CHOOSER_API UDataInterface_Object_Asset : public UObject, public IDataInterface
+{
+	GENERATED_BODY()
+	DATA_INTERFACE_RETURN_TYPE(Object);
+	
+	// IDataInterface interface
+	virtual bool GetDataImpl(const UE::DataInterface::FContext& Context) const final override;
+	
+public: 
+	UPROPERTY(EditAnywhere, Category = "Parameters")
+	TObjectPtr<UObject> Asset;
+};
diff --git a/Engine/Plugins/Experimental/Chooser/Source/Chooser/Private/Chooser.cpp b/Engine/Plugins/Experimental/Chooser/Source/Chooser/Private/Chooser.cpp
new file mode 100644
--- /dev/null
+++ b/Engine/Plugins/Experimental/Chooser/Source/Chooser/Private/Chooser.cpp
@@ -0,0 +1,94 @@
+#include "Chooser.h"
+#include "DataInterfaceTypes.h"
+#include "DataInterface.h"
+
+UChooserTable::UChooserTable(const FObjectInitializer& Initializer)
+	:Super(Initializer)
+{
+
+}
+
+void UChooserColumnBool::Filter(const UE::DataInterface::FContext& Context, const TArray<uint32>& IndexListIn, TArray<uint32>& IndexListOut)
+{
+	bool Result = false;
+	if(UE::DataInterface::GetDataSafe(Value, Context, Result))
+	{
+		for(uint32 Index : IndexListIn)
+		{
+			if (RowValues.Num() > (int)Index)
+			{
+				if (Result == RowValues[Index])
+				{
+					IndexListOut.Push(Index);
+				}
+			}
+		}
+	}
+}
+
+void UChooserColumnFloatRange::Filter(const UE::DataInterface::FContext& Context, const TArray<uint32>& IndexListIn, TArray<uint32>& IndexListOut)
+{
+	float Result = 0.0f;
+	
+	if(UE::DataInterface::GetDataSafe(Value, Context, Result))
+	{
+		for(uint32 Index : IndexListIn)
+		{
+			if (RowValues.Num() > (int)Index)
+			{
+				const FChooserFloatRangeRowData& RowValue = RowValues[Index];
+				if (Result >= RowValue.Min && Result <= RowValue.Max)
+				{
+					IndexListOut.Push(Index);
+				}
+			}
+		}
+	}
+}
+
+bool UDataInterface_EvaluateChooser::GetDataImpl(const UE::DataInterface::FContext& DataContext) const
+{
+	if (Chooser == nullptr)
+	{
+		return false;
+	}
+
+	TArray<uint32> Indices1;
+	TArray<uint32> Indices2;
+
+	int RowCount = Chooser->Results.Num();
+	Indices1.SetNum(RowCount);
+	for(int i=0;i<RowCount;i++)
+	{
+		Indices1[i]=i;
+	}
+	TArray<uint32>* IndicesOut = &Indices1;
+	TArray<uint32>* IndicesIn = &Indices2;
+
+	for (auto Column: Chooser->Columns)
+	{
+		if (Column)
+		{
+			Swap(IndicesIn, IndicesOut);
+			IndicesOut->SetNum(0, false);
+			Column->Filter(DataContext, *IndicesIn, *IndicesOut);
+		}
+	}
+
+	// of the rows that passed all column filters, return the first one for which the result row succeeds (could fail eg for a nexted chooser where no rows passed)
+	for (uint32 SelectedIndex : *IndicesOut)
+	{
+		if (Chooser->Results.Num() > (int32)SelectedIndex)
+		{
+			const TScriptInterface<IDataInterface>& SelectedResult = Chooser->Results[SelectedIndex];
+
+			// todo: GetResultRaw returns const, GetData requires non-const
+			if(UE::DataInterface::GetDataSafe(SelectedResult, DataContext, const_cast<UE::DataInterface::FParam&>(DataContext.GetResultRaw())))
+			{
+				return true;
+			}
+		}
+	}
+
+	return false;
+}
diff --git a/Engine/Plugins/Experimental/Chooser/Source/Chooser/Private/ChooserModule.cpp b/Engine/Plugins/Experimental/Chooser/Source/Chooser/Private/ChooserModule.cpp
new file mode 100644
--- /dev/null
+++ b/Engine/Plugins/Experimental/Chooser/Source/Chooser/Private/ChooserModule.cpp
@@ -0,0 +1,5 @@
+// Copyright Epic Games, Inc. All Rights Reserved.
+
+#include "Modules/ModuleManager.h"
+
+IMPLEMENT_MODULE(FDefaultModuleImpl, Chooser);
\ No newline at end of file
diff --git a/Engine/Plugins/Experimental/Chooser/Source/Chooser/Private/DataInterface_Object.cpp b/Engine/Plugins/Experimental/Chooser/Source/Chooser/Private/DataInterface_Object.cpp
new file mode 100644
--- /dev/null
+++ b/Engine/Plugins/Experimental/Chooser/Source/Chooser/Private/DataInterface_Object.cpp
@@ -0,0 +1,9 @@
+﻿// Copyright Epic Games, Inc. All Rights Reserved.
+
+#include "DataInterface_Object.h"
+
+bool UDataInterface_Object_Asset::GetDataImpl(const UE::DataInterface::FContext& Context) const
+{
+	Context.SetResult(TObjectPtr<UObject>(Asset));
+	return true;
+}
\ No newline at end of file
diff --git a/Engine/Plugins/Experimental/Chooser/Source/ChooserEditor/ChooserEditor.Build.cs b/Engine/Plugins/Experimental/Chooser/Source/ChooserEditor/ChooserEditor.Build.cs
new file mode 100644
--- /dev/null
+++ b/Engine/Plugins/Experimental/Chooser/Source/ChooserEditor/ChooserEditor.Build.cs
@@ -0,0 +1,63 @@
+// Copyright Epic Games, Inc. All Rights Reserved.
+
+namespace UnrealBuildTool.Rules
+{
+	public class ChooserEditor : ModuleRules
+	{
+		public ChooserEditor(ReadOnlyTargetRules Target) : base(Target)
+		{
+			PublicIncludePaths.AddRange(
+				new string[] {
+					// ... add public include paths required here ...
+				}
+			);
+			
+			PrivateIncludePaths.AddRange(
+				new string[] {
+					// ... add other private include paths required here ...
+				}
+			);
+			
+			PublicDependencyModuleNames.AddRange(
+				new string[]
+				{
+					// ... add other public dependencies that you statically link with here ...
+				}
+			);
+
+			PrivateDependencyModuleNames.AddRange(
+				new string[]
+				{
+					"Core",
+					"CoreUObject",
+					"Engine",
+					"AssetTools",
+					"DataInterface",
+					"DataInterfaceGraphEditor", // for DataInterfaceClassFilter, and Widget Factory system (could be moved to a more generic DataInterfaceEditor module
+					"Chooser",
+					"UnrealEd",
+					"EditorWidgets",
+					"SlateCore",
+					"Slate",
+					"PropertyEditor",
+					"RigVM",
+					"RigVMDeveloper",
+					"ControlRig",
+					"ControlRigDeveloper",
+					"ControlRigEditor",
+					"GraphEditor",
+					"InputCore",
+					"EditorStyle",
+					// ... add private dependencies that you statically link with here ...
+				}
+			);
+
+			DynamicallyLoadedModuleNames.AddRange(
+				new string[]
+				{
+					// ... add any modules that your module loads dynamically here ...
+				}
+			);
+		}
+	}
+}
\ No newline at end of file
diff --git a/Engine/Plugins/Experimental/Chooser/Source/ChooserEditor/Private/AssetTypeActions_Chooser.cpp b/Engine/Plugins/Experimental/Chooser/Source/ChooserEditor/Private/AssetTypeActions_Chooser.cpp
new file mode 100644
--- /dev/null
+++ b/Engine/Plugins/Experimental/Chooser/Source/ChooserEditor/Private/AssetTypeActions_Chooser.cpp
@@ -0,0 +1,27 @@
+// Copyright Epic Games, Inc. All Rights Reserved.
+
+#include "AssetTypeActions_Chooser.h"
+#include "ChooserTableEditor.h"
+
+#define LOCTEXT_NAMESPACE "AssetTypeActions_ChooserTable" 
+
+namespace UE::ChooserEditor
+{
+
+void FAssetTypeActions_ChooserTable::OpenAssetEditor(const TArray<UObject*>& InObjects, TSharedPtr<class IToolkitHost> EditWithinLevelEditor)
+{
+	FChooserTableEditor::CreateEditor(EToolkitMode::Standalone, EditWithinLevelEditor, InObjects);
+}
+
+const TArray<FText>& FAssetTypeActions_ChooserTable::GetSubMenus() const
+{
+	static const TArray<FText> SubMenus
+	{
+		LOCTEXT("DataInterfacesSubMenu", "Data Interfaces")
+	};
+	return SubMenus;
+}
+
+}
+
+#undef LOCTEXT_NAMESPACE
\ No newline at end of file
diff --git a/Engine/Plugins/Experimental/Chooser/Source/ChooserEditor/Private/AssetTypeActions_Chooser.h b/Engine/Plugins/Experimental/Chooser/Source/ChooserEditor/Private/AssetTypeActions_Chooser.h
new file mode 100644
--- /dev/null
+++ b/Engine/Plugins/Experimental/Chooser/Source/ChooserEditor/Private/AssetTypeActions_Chooser.h
@@ -0,0 +1,24 @@
+// Copyright Epic Games, Inc. All Rights Reserved.
+
+#pragma once
+
+#include "AssetTypeActions_Base.h"
+#include "Chooser.h"
+
+namespace UE::ChooserEditor
+{
+
+class FAssetTypeActions_ChooserTable : public FAssetTypeActions_Base
+{
+public:
+	// IAssetTypeActions Implementation
+	virtual FText GetName() const override { return NSLOCTEXT("AssetTypeActions", "AssetTypeActions_ChooserTable", "Chooser Table"); }
+	virtual FColor GetTypeColor() const override { return FColor(128,128,64); }
+	virtual UClass* GetSupportedClass() const override { return UChooserTable::StaticClass(); }
+	virtual void OpenAssetEditor( const TArray<UObject*>& InObjects, TSharedPtr<class IToolkitHost> EditWithinLevelEditor = TSharedPtr<IToolkitHost>() ) override;
+	virtual bool CanFilter() override { return true; }
+	virtual const TArray<FText>& GetSubMenus() const override;
+	virtual uint32 GetCategories() override { return EAssetTypeCategories::Animation; }
+};
+
+}
diff --git a/Engine/Plugins/Experimental/Chooser/Source/ChooserEditor/Private/ChooserEditorModule.cpp b/Engine/Plugins/Experimental/Chooser/Source/ChooserEditor/Private/ChooserEditorModule.cpp
new file mode 100644
--- /dev/null
+++ b/Engine/Plugins/Experimental/Chooser/Source/ChooserEditor/Private/ChooserEditorModule.cpp
@@ -0,0 +1,30 @@
+// Copyright Epic Games, Inc. All Rights Reserved.
+
+#include "ChooserEditorModule.h"
+#include "IAssetTools.h"
+#include "AssetTypeActions_Chooser.h"
+#include "ChooserTableEditor.h"
+
+namespace UE::ChooserEditor
+{
+
+void FModule::StartupModule()
+{
+	IAssetTools& AssetTools = FModuleManager::LoadModuleChecked<FAssetToolsModule>("AssetTools").Get();
+	AssetTypeActions_ChooserTable = MakeShared<FAssetTypeActions_ChooserTable>();
+	AssetTools.RegisterAssetTypeActions(AssetTypeActions_ChooserTable.ToSharedRef());
+	FChooserTableEditor::RegisterWidgets();
+}
+
+void FModule::ShutdownModule()
+{
+	if(FModuleManager::Get().IsModuleLoaded("AssetTools"))
+	{
+		IAssetTools& AssetTools = FModuleManager::GetModuleChecked<FAssetToolsModule>("AssetTools").Get();
+		AssetTools.UnregisterAssetTypeActions(AssetTypeActions_ChooserTable.ToSharedRef());
+	}
+}
+
+}
+
+IMPLEMENT_MODULE(UE::ChooserEditor::FModule, ChooserEditor);
\ No newline at end of file
diff --git a/Engine/Plugins/Experimental/Chooser/Source/ChooserEditor/Private/ChooserEditorModule.h b/Engine/Plugins/Experimental/Chooser/Source/ChooserEditor/Private/ChooserEditorModule.h
new file mode 100644
--- /dev/null
+++ b/Engine/Plugins/Experimental/Chooser/Source/ChooserEditor/Private/ChooserEditorModule.h
@@ -0,0 +1,22 @@
+// Copyright Epic Games, Inc. All Rights Reserved.
+
+#pragma once
+
+#include "CoreMinimal.h"
+#include "Modules/ModuleManager.h"
+
+namespace UE::ChooserEditor
+{
+
+class FAssetTypeActions_ChooserTable;
+
+class FModule : public IModuleInterface
+{
+private:
+	virtual void StartupModule() override;
+	virtual void ShutdownModule() override;
+
+	TSharedPtr<FAssetTypeActions_ChooserTable> AssetTypeActions_ChooserTable;
+};
+
+}
\ No newline at end of file
diff --git a/Engine/Plugins/Experimental/Chooser/Source/ChooserEditor/Private/ChooserFactory.cpp b/Engine/Plugins/Experimental/Chooser/Source/ChooserEditor/Private/ChooserFactory.cpp
new file mode 100644
--- /dev/null
+++ b/Engine/Plugins/Experimental/Chooser/Source/ChooserEditor/Private/ChooserFactory.cpp
@@ -0,0 +1,21 @@
+﻿// Copyright Epic Games, Inc. All Rights Reserved.
+
+#include "ChooserFactory.h"
+#include "Chooser.h"
+
+UChooserTableFactory::UChooserTableFactory()
+{
+	bCreateNew = true;
+	bEditAfterNew = true;
+	SupportedClass = UChooserTable::StaticClass();
+}
+
+bool UChooserTableFactory::ConfigureProperties()
+{
+	return true;
+}
+
+UObject* UChooserTableFactory::FactoryCreateNew(UClass* Class, UObject* InParent, FName Name, EObjectFlags Flags, UObject* Context, FFeedbackContext* Warn, FName CallingContext)
+{
+	return NewObject<UChooserTable>(InParent, Class, Name, Flags | RF_Public | RF_Standalone | RF_Transactional | RF_LoadCompleted);
+}
\ No newline at end of file
diff --git a/Engine/Plugins/Experimental/Chooser/Source/ChooserEditor/Private/ChooserFactory.h b/Engine/Plugins/Experimental/Chooser/Source/ChooserEditor/Private/ChooserFactory.h
new file mode 100644
--- /dev/null
+++ b/Engine/Plugins/Experimental/Chooser/Source/ChooserEditor/Private/ChooserFactory.h
@@ -0,0 +1,22 @@
+﻿// Copyright Epic Games, Inc. All Rights Reserved.
+
+#pragma once
+
+#include "CoreMinimal.h"
+#include "Factories/Factory.h"
+#include "ChooserFactory.generated.h"
+
+UCLASS(MinimalAPI)
+class UChooserTableFactory : public UFactory
+{
+	GENERATED_BODY()
+
+	UChooserTableFactory();
+	
+	virtual bool ConfigureProperties() override;
+	virtual UObject* FactoryCreateNew(UClass* Class, UObject* InParent, FName Name, EObjectFlags Flags, UObject* Context, FFeedbackContext* Warn, FName CallingContext) override;
+	virtual UObject* FactoryCreateNew(UClass* Class, UObject* InParent, FName Name, EObjectFlags Flags, UObject* Context, FFeedbackContext* Warn) override
+	{
+		return FactoryCreateNew(Class, InParent, Name, Flags, Context, Warn, NAME_None);
+	}
+};
\ No newline at end of file
diff --git a/Engine/Plugins/Experimental/Chooser/Source/ChooserEditor/Private/ChooserTableEditor.cpp b/Engine/Plugins/Experimental/Chooser/Source/ChooserEditor/Private/ChooserTableEditor.cpp
new file mode 100644
--- /dev/null
+++ b/Engine/Plugins/Experimental/Chooser/Source/ChooserEditor/Private/ChooserTableEditor.cpp
@@ -0,0 +1,1141 @@
+// Copyright Epic Games, Inc. All Rights Reserved.
+
+#include "ChooserTableEditor.h"
+#include "Modules/ModuleManager.h"
+#include "Styling/AppStyle.h"
+#include "Widgets/Docking/SDockTab.h"
+#include "PropertyEditorModule.h"
+#include "IDetailsView.h"
+#include "Editor.h"
+#include "Widgets/Input/SHyperlink.h"
+#include "Widgets/Input/SNumericEntryBox.h"
+#include "Widgets/Views/SListView.h"
+#include "Widgets/Input/SComboButton.h"
+#include "SAssetDropTarget.h"
+#include "SClassViewer.h"
+#include "SourceCodeNavigation.h"
+#include "Chooser.h"
+#include "DetailCategoryBuilder.h"
+#include "DetailLayoutBuilder.h"
+#include "ClassViewerFilter.h"
+#include "IDetailGroup.h"
+#include "SAssetDropTarget.h"
+#include "Widgets/Layout/SWidgetSwitcher.h"
+#include "DragAndDrop/DecoratedDragDropOp.h"
+#include "DataInterfaceClassFilter.h"
+#include "DataInterfaceWidgetFactories.h"
+
+
+#define LOCTEXT_NAMESPACE "ChooserEditor"
+
+namespace UE::ChooserEditor
+{
+
+const FName FChooserTableEditor::ToolkitFName( TEXT( "GenericAssetEditor" ) );
+const FName FChooserTableEditor::PropertiesTabId( TEXT( "ChooserEditor_Properties" ) );
+const FName FChooserTableEditor::TableTabId( TEXT( "ChooserEditor_Table" ) );
+
+
+void ConvertToText(UObject* Object, FText& OutText)
+{
+	UClass* Class = Object->GetClass();
+	while (Class)
+	{
+		if (auto TextConverter = DataInterfaceGraphEditor::FDataInterfaceWidgetFactories::DataInterfaceTextConverter.Find(Class))
+		{
+			(*TextConverter)(Object, OutText);
+			break;
+		}
+		Class = Class->GetSuperClass();
+	}
+}
+
+// Filter class for things which implement IChooserColumn
+class FInterfaceClassFilter : public IClassViewerFilter
+{
+public:
+	FInterfaceClassFilter(UClass* InInterfaceType) : InterfaceType(InInterfaceType)  { };
+	virtual ~FInterfaceClassFilter() override {};
+	
+	virtual bool IsClassAllowed(const FClassViewerInitializationOptions& InInitOptions, const UClass* InClass, TSharedRef< class FClassViewerFilterFuncs > InFilterFuncs ) override
+	{
+		return (InClass->ImplementsInterface(InterfaceType) && InClass->HasAnyClassFlags(CLASS_Abstract) == false);
+	}
+
+	virtual bool IsUnloadedClassAllowed(const FClassViewerInitializationOptions& InInitOptions, const TSharedRef< const class IUnloadedBlueprintData > InUnloadedClassData, TSharedRef< class FClassViewerFilterFuncs > InFilterFuncs) override
+	{
+		return false;
+	}
+private:
+	UClass* InterfaceType;
+};
+
+void FChooserTableEditor::RegisterTabSpawners(const TSharedRef<class FTabManager>& InTabManager)
+{
+	WorkspaceMenuCategory = InTabManager->AddLocalWorkspaceMenuCategory(LOCTEXT("WorkspaceMenu_GenericAssetEditor", "Asset Editor"));
+
+	FAssetEditorToolkit::RegisterTabSpawners(InTabManager);
+
+	InTabManager->RegisterTabSpawner( PropertiesTabId, FOnSpawnTab::CreateSP(this, &FChooserTableEditor::SpawnPropertiesTab) )
+		.SetDisplayName( LOCTEXT("PropertiesTab", "Details") )
+		.SetGroup(WorkspaceMenuCategory.ToSharedRef())
+		.SetIcon(FSlateIcon("EditorStyle", "LevelEditor.Tabs.Details"));
+	
+	InTabManager->RegisterTabSpawner( TableTabId, FOnSpawnTab::CreateSP(this, &FChooserTableEditor::SpawnTableTab) )
+		.SetDisplayName( LOCTEXT("TableTab", "Chooser Table") )
+		.SetGroup(WorkspaceMenuCategory.ToSharedRef())
+		.SetIcon(FSlateIcon("EditorStyle", "LevelEditor.Tabs.Details"));
+}
+
+void FChooserTableEditor::UnregisterTabSpawners(const TSharedRef<class FTabManager>& InTabManager)
+{
+	FAssetEditorToolkit::UnregisterTabSpawners(InTabManager);
+
+	InTabManager->UnregisterTabSpawner( PropertiesTabId );
+}
+
+const FName FChooserTableEditor::ChooserEditorAppIdentifier( TEXT( "ChooserEditorApp" ) );
+
+FChooserTableEditor::~FChooserTableEditor()
+{
+	GEditor->GetEditorSubsystem<UImportSubsystem>()->OnAssetPostImport.RemoveAll(this);
+	FCoreUObjectDelegates::OnObjectsReplaced.RemoveAll(this);
+
+	DetailsView.Reset();
+}
+
+
+void FChooserTableEditor::InitEditor( const EToolkitMode::Type Mode, const TSharedPtr< class IToolkitHost >& InitToolkitHost, const TArray<UObject*>& ObjectsToEdit, FGetDetailsViewObjects GetDetailsViewObjects )
+{
+	EditingObjects = ObjectsToEdit;
+	GEditor->GetEditorSubsystem<UImportSubsystem>()->OnAssetPostImport.AddSP(this, &FChooserTableEditor::HandleAssetPostImport);
+	FCoreUObjectDelegates::OnObjectsReplaced.AddSP(this, &FChooserTableEditor::OnObjectsReplaced);
+
+	FPropertyEditorModule& PropertyEditorModule = FModuleManager::GetModuleChecked<FPropertyEditorModule>( "PropertyEditor" );
+	FDetailsViewArgs DetailsViewArgs;
+	DetailsViewArgs.NameAreaSettings = FDetailsViewArgs::HideNameArea;
+	DetailsView = PropertyEditorModule.CreateDetailView( DetailsViewArgs );
+	const TSharedRef<FTabManager::FLayout> StandaloneDefaultLayout = FTabManager::NewLayout( "Standalone_ChooserTableEditor_Layout_v1" )
+	->AddArea
+	(
+		FTabManager::NewPrimaryArea() ->SetOrientation(Orient_Vertical)
+		->Split
+		(
+			FTabManager::NewSplitter()
+			->Split
+			(
+				FTabManager::NewStack()
+				->SetSizeCoefficient(0.7)
+				->AddTab( TableTabId, ETabState::OpenedTab )
+			)
+			->Split
+			(
+				FTabManager::NewStack()
+				->SetSizeCoefficient(0.3)
+				->AddTab( PropertiesTabId, ETabState::OpenedTab )
+			)
+		)
+	);
+
+	const bool bCreateDefaultStandaloneMenu = true;
+	const bool bCreateDefaultToolbar = true;
+	FAssetEditorToolkit::InitAssetEditor( Mode, InitToolkitHost, FChooserTableEditor::ChooserEditorAppIdentifier, StandaloneDefaultLayout, bCreateDefaultStandaloneMenu, bCreateDefaultToolbar, ObjectsToEdit );
+
+	RegenerateMenusAndToolbars();
+
+	UChooserTable* Chooser = Cast<UChooserTable>(EditingObjects[0]);
+	
+	SelectRootProperties();
+}
+
+FName FChooserTableEditor::GetToolkitFName() const
+{
+	return ToolkitFName;
+}
+
+FText FChooserTableEditor::GetBaseToolkitName() const
+{
+	return LOCTEXT("AppLabel", "Chooser Table Editor");
+}
+
+FText FChooserTableEditor::GetToolkitName() const
+{
+	const TArray<UObject*>& EditingObjs = GetEditingObjects();
+
+	check( EditingObjs.Num() > 0 );
+
+	FFormatNamedArguments Args;
+	Args.Add( TEXT("ToolkitName"), GetBaseToolkitName() );
+
+	if( EditingObjs.Num() == 1 )
+	{
+		const UObject* EditingObject = EditingObjs[ 0 ];
+		return FText::FromString(EditingObject->GetName());
+	}
+	else
+	{
+		UClass* SharedBaseClass = nullptr;
+		for( int32 x = 0; x < EditingObjs.Num(); ++x )
+		{
+			UObject* Obj = EditingObjs[ x ];
+			check( Obj );
+
+			UClass* ObjClass = Cast<UClass>(Obj);
+			if (ObjClass == nullptr)
+			{
+				ObjClass = Obj->GetClass();
+			}
+			check( ObjClass );
+
+			// Initialize with the class of the first object we encounter.
+			if( SharedBaseClass == nullptr )
+			{
+				SharedBaseClass = ObjClass;
+			}
+
+			// If we've encountered an object that's not a subclass of the current best baseclass,
+			// climb up a step in the class hierarchy.
+			while( !ObjClass->IsChildOf( SharedBaseClass ) )
+			{
+				SharedBaseClass = SharedBaseClass->GetSuperClass();
+			}
+		}
+
+		check(SharedBaseClass);
+
+		Args.Add( TEXT("NumberOfObjects"), EditingObjs.Num() );
+		Args.Add( TEXT("ClassName"), FText::FromString( SharedBaseClass->GetName() ) );
+		return FText::Format( LOCTEXT("ToolkitTitle_EditingMultiple", "{NumberOfObjects} {ClassName} - {ToolkitName}"), Args );
+	}
+}
+
+FText FChooserTableEditor::GetToolkitToolTipText() const
+{
+	const TArray<UObject*>& EditingObjs = GetEditingObjects();
+
+	check( EditingObjs.Num() > 0 );
+
+	FFormatNamedArguments Args;
+	Args.Add( TEXT("ToolkitName"), GetBaseToolkitName() );
+
+	if( EditingObjs.Num() == 1 )
+	{
+		const UObject* EditingObject = EditingObjs[ 0 ];
+		return FAssetEditorToolkit::GetToolTipTextForObject(EditingObject);
+	}
+	else
+	{
+		UClass* SharedBaseClass = NULL;
+		for( int32 x = 0; x < EditingObjs.Num(); ++x )
+		{
+			UObject* Obj = EditingObjs[ x ];
+			check( Obj );
+
+			UClass* ObjClass = Cast<UClass>(Obj);
+			if (ObjClass == nullptr)
+			{
+				ObjClass = Obj->GetClass();
+			}
+			check( ObjClass );
+
+			// Initialize with the class of the first object we encounter.
+			if( SharedBaseClass == nullptr )
+			{
+				SharedBaseClass = ObjClass;
+			}
+
+			// If we've encountered an object that's not a subclass of the current best baseclass,
+			// climb up a step in the class hierarchy.
+			while( !ObjClass->IsChildOf( SharedBaseClass ) )
+			{
+				SharedBaseClass = SharedBaseClass->GetSuperClass();
+			}
+		}
+
+		check(SharedBaseClass);
+
+		Args.Add( TEXT("NumberOfObjects"), EditingObjs.Num() );
+		Args.Add( TEXT("ClassName"), FText::FromString( SharedBaseClass->GetName() ) );
+		return FText::Format( LOCTEXT("ToolkitTitle_EditingMultipleToolTip", "{NumberOfObjects} {ClassName} - {ToolkitName}"), Args );
+	}
+}
+
+FLinearColor FChooserTableEditor::GetWorldCentricTabColorScale() const
+{
+	return FLinearColor( 0.5f, 0.0f, 0.0f, 0.5f );
+}
+
+void FChooserTableEditor::SetPropertyVisibilityDelegate(FIsPropertyVisible InVisibilityDelegate)
+{
+	DetailsView->SetIsPropertyVisibleDelegate(InVisibilityDelegate);
+	DetailsView->ForceRefresh();
+}
+
+void FChooserTableEditor::SetPropertyEditingEnabledDelegate(FIsPropertyEditingEnabled InPropertyEditingDelegate)
+{
+	DetailsView->SetIsPropertyEditingEnabledDelegate(InPropertyEditingDelegate);
+	DetailsView->ForceRefresh();
+}
+
+TSharedRef<SDockTab> FChooserTableEditor::SpawnPropertiesTab( const FSpawnTabArgs& Args )
+{
+	check( Args.GetTabId() == PropertiesTabId );
+
+	return SNew(SDockTab)
+		.Label( LOCTEXT("GenericDetailsTitle", "Details") )
+		.TabColorScale( GetTabColorScale() )
+		.OnCanCloseTab_Lambda([]() { return false; })
+		[
+			DetailsView.ToSharedRef()
+		];
+}
+
+
+class FChooserRowDragDropOp : public FDecoratedDragDropOp
+{
+public:
+	DRAG_DROP_OPERATOR_TYPE(FWidgetTemplateDragDropOp, FDecoratedDragDropOp)
+
+	FChooserTableEditor* ChooserEditor;
+	uint32 RowIndex;
+
+	/** Constructs the drag drop operation */
+	static TSharedRef<FChooserRowDragDropOp> New(FChooserTableEditor* InEditor, uint32 InRowIndex)
+	{
+		TSharedRef<FChooserRowDragDropOp> Operation = MakeShareable(new FChooserRowDragDropOp());
+		Operation->ChooserEditor = InEditor;
+		Operation->RowIndex = InRowIndex;
+		Operation->DefaultHoverText = LOCTEXT("Chooser Row", "Chooser Row");
+		ConvertToText(InEditor->GetChooser()->Results[InRowIndex].GetObject(), Operation->DefaultHoverText);
+		Operation->CurrentHoverText = Operation->DefaultHoverText;
+			
+		Operation->Construct();
+	
+		return Operation;
+	};
+};
+
+class SChooserRowHandle : public SCompoundWidget
+{
+public:
+	SLATE_BEGIN_ARGS(SChooserRowHandle)
+	{}
+	SLATE_DEFAULT_SLOT(FArguments, Content)
+	SLATE_ARGUMENT(FChooserTableEditor*, ChooserEditor)
+	SLATE_ARGUMENT(uint32, RowIndex)
+	SLATE_END_ARGS()
+
+	void Construct(const FArguments& InArgs)
+	{
+		ChooserEditor = InArgs._ChooserEditor;
+		RowIndex = InArgs._RowIndex;
+
+		ChildSlot
+		[
+			SNew(SBox) .Padding(0.0f) .HAlign(HAlign_Center) .VAlign(VAlign_Center) .WidthOverride(16.0f)
+			[
+				SNew(SImage)
+				.Image(FCoreStyle::Get().GetBrush("VerticalBoxDragIndicatorShort"))
+			]
+		];
+	}
+
+	FReply OnMouseButtonDown(const FGeometry& MyGeometry, const FPointerEvent& MouseEvent) override
+	{
+		return FReply::Handled().DetectDrag(SharedThis(this), EKeys::LeftMouseButton);
+	};
+
+	FReply OnDragDetected(const FGeometry& MyGeometry, const FPointerEvent& MouseEvent) override
+	{
+		TSharedRef<FChooserRowDragDropOp> DragDropOp = FChooserRowDragDropOp::New(ChooserEditor, RowIndex);
+		return FReply::Handled().BeginDragDrop(DragDropOp);
+	}
+
+private:
+	FChooserTableEditor* ChooserEditor = nullptr;
+	uint32 RowIndex;
+};
+
+
+class SChooserTableRow : public SMultiColumnTableRow<TSharedPtr<FChooserTableEditor::FChooserTableRow>>
+{
+public:
+	SLATE_BEGIN_ARGS(SChooserTableRow) {}
+		/** The list item for this row */
+		SLATE_ARGUMENT(TSharedPtr<FChooserTableEditor::FChooserTableRow>, Entry)
+		SLATE_ARGUMENT(UChooserTable*, Chooser)
+		SLATE_ARGUMENT(FChooserTableEditor*, Editor)
+	SLATE_END_ARGS()
+
+
+	void Construct(const FArguments& Args, const TSharedRef<STableViewBase>& OwnerTableView)
+	{
+		RowIndex = Args._Entry;
+		Chooser = Args._Chooser;
+		Editor = Args._Editor;
+
+		SMultiColumnTableRow<TSharedPtr<FChooserTableEditor::FChooserTableRow>>::Construct(
+			FSuperRowType::FArguments(),
+			OwnerTableView
+		);
+	}
+
+	/** Overridden from SMultiColumnTableRow.  Generates a widget for this column of the list view. */
+	virtual TSharedRef<SWidget> GenerateWidgetForColumn(const FName& ColumnName) override
+	{
+		static FName Result = "Result";
+		static FName Handles = "Handles";
+		
+		if (RowIndex->RowIndex < Chooser->Results.Num())
+		{
+			if (ColumnName == Handles && RowIndex->RowIndex < Chooser->Results.Num())
+			{
+				// row drag handle
+				
+				return 		SNew(SChooserRowHandle).ChooserEditor(Editor).RowIndex(RowIndex->RowIndex);
+			}
+			else if (ColumnName == Result)
+			{
+		
+				UObject* RowValue = Chooser->Results[RowIndex->RowIndex].GetObject();
+				TSharedPtr<SWidget> ResultWidget = DataInterfaceGraphEditor::FDataInterfaceWidgetFactories::CreateDataInterfaceWidget(Chooser->ResultType, RowValue,
+				FOnClassPicked::CreateLambda([this, RowIndex=RowIndex->RowIndex](UClass* ChosenClass)
+				{
+					UObject* RowValue = NewObject<UObject>(Chooser, ChosenClass);
+					Chooser->Results[RowIndex] = RowValue;
+					DataInterfaceGraphEditor::FDataInterfaceWidgetFactories::CreateDataInterfaceWidget(Chooser->ResultType, RowValue, FOnClassPicked(), &CacheBorder);
+				}),
+				&CacheBorder
+				);
+				return ResultWidget.ToSharedRef();
+			}
+			else
+			{
+				const int ColumnIndex = ColumnName.GetNumber() - 1;
+				if (ColumnIndex < Chooser->Columns.Num() && ColumnIndex >=0)
+				{
+					auto& Column = Chooser->Columns[ColumnIndex];
+					TSharedPtr<SWidget> ColumnWidget;
+					UClass* ColumnClass = Column.GetObject()->GetClass();
+					while (ColumnClass && !ColumnWidget.IsValid())
+					{
+						if (auto Creator = FChooserTableEditor::ColumnWidgetCreators.Find(ColumnClass))
+						{
+							ColumnWidget = (*Creator)(Column.GetObject(), RowIndex->RowIndex);
+							break;
+						}
+						ColumnClass = ColumnClass->GetSuperClass();
+					}
+					return ColumnWidget.ToSharedRef();
+				}
+			}
+		}
+		else if (RowIndex->RowIndex == Chooser->Results.Num())
+        {
+			// on the row past the end, show an Add button in the result column
+			if (ColumnName == Result)
+			{
+				TSharedRef<SComboButton> CreateRowComboButton = SNew(SComboButton)
+									.ComboButtonStyle(FAppStyle::Get(), "SimpleComboButton")
+									.ButtonContent()
+									[
+										SNew(STextBlock).Text(LOCTEXT("Add", "+ Add Row"))
+									];
+
+				CreateRowComboButton->SetOnGetMenuContent(FOnGetContent::CreateLambda([this, CreateRowComboButton]()
+					{
+						FClassViewerInitializationOptions Options;
+						Options.ClassFilters.Add(MakeShared<UE::DataInterfaceGraphEditor::FDataInterfaceClassFilter>(Chooser->ResultType, nullptr));
+						
+						// Add class filter for columns here
+						TSharedRef<SWidget> Widget = FModuleManager::LoadModuleChecked<FClassViewerModule>("ClassViewer").CreateClassViewer(Options,
+							FOnClassPicked::CreateLambda([this,CreateRowComboButton](UClass* ChosenClass)
+							{
+								CreateRowComboButton->SetIsOpen(false);
+								UObject* ResultObject = NewObject<UObject>(Chooser, ChosenClass);
+								Chooser->Results.Add(ResultObject);
+								Editor->UpdateTableRows();
+							}));
+						return Widget;
+					}));
+
+				return CreateRowComboButton;
+			}
+		}
+		return SNullWidget::NullWidget;
+	}
+
+private:
+	TSharedPtr<FChooserTableEditor::FChooserTableRow> RowIndex;
+	UChooserTable* Chooser;
+	FChooserTableEditor* Editor;
+	TSharedPtr<SBorder> CacheBorder;
+};
+
+
+TSharedRef<ITableRow> FChooserTableEditor::GenerateTableRow(TSharedPtr<FChooserTableRow> InItem, const TSharedRef<STableViewBase>& OwnerTable)
+{
+	UChooserTable* Chooser = Cast<UChooserTable>(EditingObjects[0]);
+	
+	return SNew(SChooserTableRow, OwnerTable)
+		.Entry(InItem).Chooser(Chooser).Editor(this);
+}
+
+FReply FChooserTableEditor::SelectRootProperties()
+{
+	if( DetailsView.IsValid() )
+	{
+		// Make sure details window is pointing to our object
+		DetailsView->SetObjects( EditingObjects );
+	}
+
+	return FReply::Handled();
+}
+
+
+void FChooserTableEditor::UpdateTableColumns()
+{
+	UChooserTable* Chooser = Cast<UChooserTable>(EditingObjects[0]);
+	
+	HeaderRow->ClearColumns();
+	
+	HeaderRow->AddColumn(SHeaderRow::Column("Handles")
+					.ManualWidth(30)
+					.HeaderContent()
+					[					
+						SNew(SButton).OnClicked_Raw(this, &FChooserTableEditor::SelectRootProperties)
+					]);
+	
+	HeaderRow->AddColumn(SHeaderRow::Column("Result")
+					.DefaultLabel(LOCTEXT("ResultColumnName", "Result"))
+					.ManualWidth(200));
+
+	ColumnText.SetNum(0);
+
+	FName ColumnId("ChooserColumn", 1);
+	int NumColumns = Chooser->Columns.Num();	
+	for(int ColumnIndex = 0; ColumnIndex < NumColumns; ColumnIndex++)
+	{
+		auto& Column = Chooser->Columns[ColumnIndex];
+
+		ColumnText.Push(
+			SNew(SEditableText)
+					.Text_Lambda([&Column](){ return FText::FromName(Column->GetDisplayName()); })
+					.OnTextCommitted_Lambda([this, &Column](const FText& Text, ETextCommit::Type CommitType)
+					{
+						Column->SetDisplayName(ToCStr(Text.ToString()));
+						EditColumnName = false;
+					} )
+		);
+
+		HeaderRow->AddColumn(SHeaderRow::FColumn::FArguments()
+			.ColumnId(ColumnId)
+			.ManualWidth(200)
+			.OnGetMenuContent_Lambda([this, ColumnIndex]()
+			{
+				FMenuBuilder MenuBuilder(true,nullptr);
+				MenuBuilder.AddMenuEntry(LOCTEXT("Delete Column","Delete Column"),LOCTEXT("Delete Column ToolTip", "Remove this column and all its data from the table"),FSlateIcon(),
+					FUIAction(
+						FExecuteAction::CreateLambda([this,ColumnIndex]()
+						{
+							UChooserTable* Chooser = Cast<UChooserTable>(EditingObjects[0]);
+							Chooser->Columns.RemoveAt(ColumnIndex);
+							UpdateTableColumns();
+						})
+						)
+					);
+				return MenuBuilder.MakeWidget();
+			})
+			.HeaderContent()
+			[
+				SNew(SBorder)
+				.VAlign(VAlign_Center)
+				.Padding(3)
+				.BorderBackgroundColor_Lambda([this, ColumnId] ()
+				{
+					// unclear why this color is coming out much darker
+					return SelectedColumn == ColumnId ? FSlateColor(FColor(0x00, 0x70, 0xe0, 0xFF)) : FSlateColor(FLinearColor(0.05f,0.05f,0.05f));
+				})
+				.OnMouseButtonDown_Lambda([this, ColumnIndex, ColumnId](	const FGeometry&, const FPointerEvent& PointerEvent)
+				{
+					UChooserTable* Chooser = Cast<UChooserTable>(EditingObjects[0]);
+					TableView->ClearSelection();
+					
+					if (PointerEvent.GetEffectingButton() == EKeys::LeftMouseButton )
+					{
+						// on second click, enable editing, and select all text
+						bool ShouldEditColumn = (SelectedColumn == ColumnId);
+						if (ShouldEditColumn && !EditColumnName)
+						{
+							EditColumnName = ShouldEditColumn;
+							FSlateApplication::Get().SetKeyboardFocus(ColumnText[ColumnIndex]);  // This doesn't work, because it can't find the path to the widget? because it's not active yet
+							ColumnText[ColumnIndex]->SelectAllText();
+						}
+						EditColumnName = ShouldEditColumn;
+						
+						DetailsView->SetObject(Chooser->Columns[ColumnIndex].GetObject());
+						
+						SelectedColumn = ColumnId;
+						
+						return FReply::Handled();
+					}
+					else 
+					{
+						SelectedColumn = ColumnId;
+						
+						return FReply::Handled();
+					}
+				})
+				[
+					SNew(SWidgetSwitcher).WidgetIndex_Lambda([this, ColumnId]()
+					{
+						return EditColumnName && SelectedColumn == ColumnId ? 1 : 0;
+					})
+					+ SWidgetSwitcher::Slot()
+					[
+						SNew(STextBlock).Text_Lambda([&Column](){ return FText::FromName(Column->GetDisplayName()); })
+					]
+					+ SWidgetSwitcher::Slot()
+					[
+						ColumnText[ColumnIndex].ToSharedRef()
+					]
+				]
+				
+			]);
+		
+		ColumnId.SetNumber(ColumnId.GetNumber() + 1);
+	}
+
+	HeaderRow->AddColumn( SHeaderRow::FColumn::FArguments()
+		.ColumnId("Add")
+		.FillWidth(1.0)
+		.HeaderContent( )
+		[
+			SNew(SHorizontalBox)
+			+ SHorizontalBox::Slot().MaxWidth(120)
+			[
+				CreateColumnComboButton.ToSharedRef()
+			]
+		]
+		);
+
+}
+
+TSharedRef<SDockTab> FChooserTableEditor::SpawnTableTab( const FSpawnTabArgs& Args )
+{
+	check( Args.GetTabId() == TableTabId );
+
+	UChooserTable* Chooser = Cast<UChooserTable>(EditingObjects[0]);
+	
+	// + button to create new columns
+	
+	CreateColumnComboButton = SNew(SComboButton).OnGetMenuContent_Lambda([this]()
+	{
+		FClassViewerInitializationOptions Options;
+		Options.ClassFilters.Add(MakeShared<FInterfaceClassFilter>(UChooserColumn::StaticClass()));
+		
+		// Add class filter for columns here
+		TSharedRef<SWidget> Widget = FModuleManager::LoadModuleChecked<FClassViewerModule>("ClassViewer").CreateClassViewer(Options, FOnClassPicked::CreateLambda([this](UClass* ChosenClass)
+		{
+			CreateColumnComboButton->SetIsOpen(false);
+			UChooserTable* Chooser = Cast<UChooserTable>(EditingObjects[0]);
+			UObject* ColumnObject = NewObject<UObject>(Chooser, ChosenClass);
+			TScriptInterface<IChooserColumn> Column;
+			Column.SetObject(ColumnObject);
+			Column.SetInterface(CastChecked<IChooserColumn>(ColumnObject));
+			Chooser->Columns.Add(TScriptInterface<IChooserColumn>(Column));
+			UpdateTableColumns();
+			UpdateTableRows();
+			SelectedColumn = "ChooserColumn";
+			SelectedColumn.SetNumber(ColumnText.Num());
+			DetailsView->SetObject(ColumnObject);
+			EditColumnName = true;
+			FSlateApplication::Get().SetKeyboardFocus(ColumnText[ColumnText.Num() - 1]);
+			ColumnText[ColumnText.Num() - 1]->SelectAllText();
+		}));
+		return Widget;
+	})
+	.ComboButtonStyle(FAppStyle::Get(), "SimpleComboButton")
+	.ButtonContent()
+	[
+		SNew(STextBlock).Text(LOCTEXT("Add", "+ Add Column"))
+	];
+
+
+	HeaderRow = SNew(SHeaderRow);
+
+	UpdateTableRows();
+	UpdateTableColumns();
+
+	TableView = SNew(SListView<TSharedPtr<FChooserTableRow>>)
+    			.ListItemsSource(&TableRows)
+				.OnKeyDownHandler_Lambda([this](const FGeometry&, const FKeyEvent& Event)
+				{
+					if (Event.GetKey() == EKeys::Delete)
+					{
+						UChooserTable* Chooser = Cast<UChooserTable>(EditingObjects[0]);
+						// delete selected rows.
+						TArray<uint32> RowsToDelete;
+						for(auto& SelectedRow:SelectedRows)
+						{
+							RowsToDelete.Add(SelectedRow->Row);
+						}
+						// sort indices in reverse
+						RowsToDelete.Sort([](int32 A, int32 B){ return A>B; });
+						for(uint32 RowIndex : RowsToDelete)
+						{
+							Chooser->Results.RemoveAt(RowIndex);
+						}
+
+						for(auto& Column:Chooser->Columns)
+						{
+							Column.GetInterface()->DeleteRows(RowsToDelete);
+						}
+
+						UpdateTableRows();
+						
+						return FReply::Handled();
+					}
+					return FReply::Unhandled();
+				}
+				)
+				.OnSelectionChanged_Lambda([this](TSharedPtr<FChooserTableRow> SelectedItem,  ESelectInfo::Type SelectInfo)
+				{
+					SelectedColumn = "None";
+					if (SelectedItem)
+					{
+						SelectedRows.SetNum(0);
+						// hack test details
+						UChooserTable* Chooser = Cast<UChooserTable>(EditingObjects[0]);
+						// Get the list of objects to edit the details of
+						TObjectPtr<UChooserRowDetails> Selection = NewObject<UChooserRowDetails>();
+						Selection->Chooser = Chooser;
+						Selection->Row = SelectedItem->RowIndex;
+						Selection->SetFlags(RF_Transactional); // for undo?
+						SelectedRows.Add(Selection);
+											
+						TArray<UObject*> DetailsObjects;
+						for(auto& Item : SelectedRows)
+						{
+							DetailsObjects.Add(Item.Get());
+						}
+
+						if( DetailsView.IsValid() )
+						{
+							// Make sure details window is pointing to our object
+							DetailsView->SetObjects( DetailsObjects );
+						}
+					}
+				})
+    			.OnGenerateRow_Raw(this, &FChooserTableEditor::GenerateTableRow)
+				.HeaderRow(HeaderRow);
+	
+	return SNew(SDockTab)
+		.Label( LOCTEXT("ChooserTableTitle", "Chooser Table") )
+		.TabColorScale( GetTabColorScale() )
+		.OnCanCloseTab_Lambda([]() { return false; })
+		[
+			TableView.ToSharedRef()
+		];
+}
+
+void FChooserTableEditor::UpdateTableRows()
+{
+	UChooserTable* Chooser = Cast<UChooserTable>(EditingObjects[0]);
+	int32 OldNum = TableRows.Num();
+	int32 NewNum = Chooser->Results.Num();
+
+	// Sync the TableRows array which drives the ui table to match the number of results.
+	// Add 1 at the end, for the "Add Row" control
+	TableRows.SetNum(NewNum + 1);
+	
+	for(int32 i = OldNum; i < NewNum+1; i++)
+	{
+		TableRows[i] = MakeShared<FChooserTableRow>(i);
+	}
+
+	// Make sure each column has the same number of row datas as there are results
+	for(auto& Column : Chooser->Columns)
+	{
+		if (IChooserColumn* ColumnInterface = Column.GetInterface())
+		{
+			ColumnInterface->SetNumRows(NewNum);
+		}
+	}
+
+	if (TableView.IsValid())
+	{
+		TableView->RebuildList();
+	}
+}
+
+void FChooserTableEditor::HandleAssetPostImport(UFactory* InFactory, UObject* InObject)
+{
+	if (EditingObjects.Contains(InObject))
+	{
+		// The details panel likely needs to be refreshed if an asset was imported again
+		DetailsView->SetObjects(EditingObjects);
+	}
+}
+
+void FChooserTableEditor::OnObjectsReplaced(const TMap<UObject*, UObject*>& ReplacementMap)
+{
+	bool bChangedAny = false;
+
+	// Refresh our details view if one of the objects replaced was in the map. This gets called before the reinstance GC fixup, so we might as well fixup EditingObjects now too
+	for (int32 i = 0; i < EditingObjects.Num(); i++)
+	{
+		UObject* SourceObject = EditingObjects[i];
+		UObject* ReplacedObject = ReplacementMap.FindRef(SourceObject);
+
+		if (ReplacedObject && ReplacedObject != SourceObject)
+		{
+			EditingObjects[i] = ReplacedObject;
+			bChangedAny = true;
+		}
+	}
+
+	if (bChangedAny)
+	{
+		DetailsView->SetObjects(EditingObjects);
+	}
+}
+
+FString FChooserTableEditor::GetWorldCentricTabPrefix() const
+{
+	return LOCTEXT("WorldCentricTabPrefix", "Generic Asset ").ToString();
+}
+
+TSharedRef<FChooserTableEditor> FChooserTableEditor::CreateEditor( const EToolkitMode::Type Mode, const TSharedPtr< IToolkitHost >& InitToolkitHost, UObject* ObjectToEdit, FGetDetailsViewObjects GetDetailsViewObjects )
+{
+	TSharedRef< FChooserTableEditor > NewEditor( new FChooserTableEditor() );
+
+	TArray<UObject*> ObjectsToEdit;
+	ObjectsToEdit.Add( ObjectToEdit );
+	NewEditor->InitEditor( Mode, InitToolkitHost, ObjectsToEdit, GetDetailsViewObjects );
+
+	return NewEditor;
+}
+
+TSharedRef<FChooserTableEditor> FChooserTableEditor::CreateEditor( const EToolkitMode::Type Mode, const TSharedPtr< IToolkitHost >& InitToolkitHost, const TArray<UObject*>& ObjectsToEdit, FGetDetailsViewObjects GetDetailsViewObjects )
+{
+	TSharedRef< FChooserTableEditor > NewEditor( new FChooserTableEditor() );
+	NewEditor->InitEditor( Mode, InitToolkitHost, ObjectsToEdit, GetDetailsViewObjects );
+	return NewEditor;
+}
+
+void FChooserTableEditor::PostRegenerateMenusAndToolbars()
+{
+	// Find the common denominator class of the assets we're editing
+	TArray<UClass*> ClassList;
+	for (UObject* Obj : EditingObjects)
+	{
+		check(Obj);
+		ClassList.Add(Obj->GetClass());
+	}
+
+	UClass* CommonDenominatorClass = UClass::FindCommonBase(ClassList);
+	const bool bNotAllSame = (EditingObjects.Num() > 0) && (EditingObjects[0]->GetClass() != CommonDenominatorClass);
+
+	// Provide a hyperlink to view that native class
+	if (CommonDenominatorClass)
+	{
+		TWeakObjectPtr<UClass> WeakClassPtr(CommonDenominatorClass);
+		auto OnNavigateToClassCode = [WeakClassPtr]()
+		{
+			if (UClass* StrongClassPtr = WeakClassPtr.Get())
+			{
+				if (FSourceCodeNavigation::CanNavigateToClass(StrongClassPtr))
+				{
+					FSourceCodeNavigation::NavigateToClass(StrongClassPtr);
+				}
+			}
+		};
+
+		// build and attach the menu overlay
+		TSharedRef<SHorizontalBox> MenuOverlayBox = SNew(SHorizontalBox)
+			+ SHorizontalBox::Slot()
+			.AutoWidth()
+			.VAlign(VAlign_Center)
+			[
+				SNew(STextBlock)
+				.ColorAndOpacity(FSlateColor::UseSubduedForeground())
+				.ShadowOffset(FVector2D::UnitVector)
+				.Text(bNotAllSame ? LOCTEXT("ChooserTableEditor_AssetType_Varied", "Common Asset Type: ") : LOCTEXT("ChooserTableEditor_AssetType", "Asset Type: "))
+			]
+			+SHorizontalBox::Slot()
+			.AutoWidth()
+			.VAlign(VAlign_Center)
+			.Padding(0.0f, 0.0f, 8.0f, 0.0f)
+			[
+				SNew(SHyperlink)
+				.Style(FAppStyle::Get(), "Common.GotoNativeCodeHyperlink")
+				.OnNavigate_Lambda(OnNavigateToClassCode)
+				.Text(FText::FromName(CommonDenominatorClass->GetFName()))
+				.ToolTipText(FText::Format(LOCTEXT("GoToCode_ToolTip", "Click to open this source file in {0}"), FSourceCodeNavigation::GetSelectedSourceCodeIDE()))
+			];
+	
+		SetMenuOverlay(MenuOverlayBox);
+	}
+}
+
+/// Result widgets
+
+TMap<const UClass*, TFunction<TSharedRef<SWidget> (UObject* Column, int Row)>> FChooserTableEditor::ColumnWidgetCreators;
+
+TSharedRef<SWidget> CreateAssetWidget(UObject* Object)
+{
+	UDataInterface_Object_Asset* DIAsset = Cast<UDataInterface_Object_Asset>(Object);
+
+	UObject* Asset = DIAsset->Asset;
+	
+	return
+	SNew(SAssetDropTarget)
+	.OnAssetsDropped_Lambda([DIAsset](const FDragDropEvent& Event, TArrayView<FAssetData> AssetData)
+	{
+		if (AssetData.Num() > 0)
+		{
+			DIAsset->Asset = AssetData[0].GetAsset();
+		}
+	})
+	.Content()
+	[
+		SNew (STextBlock).Text_Lambda([DIAsset]()
+		{
+			if (DIAsset->Asset == nullptr)
+			{
+				return LOCTEXT("null", "[No Asset]");
+			}
+			else
+			{
+				return FText::FromString(DIAsset->Asset->GetName());
+			}
+		})
+	];
+}
+
+TSharedRef<SWidget> CreateEvaluateChooserWidget(UObject* Object)
+{
+	UDataInterface_EvaluateChooser* EvaluateChooser = Cast<UDataInterface_EvaluateChooser>(Object);
+
+	return
+	SNew(SAssetDropTarget)
+	.OnAssetsDropped_Lambda([EvaluateChooser](const FDragDropEvent& Event, TArrayView<FAssetData> AssetData)
+	{
+		if (AssetData.Num() > 0)
+		{
+			EvaluateChooser->Chooser = Cast<UChooserTable>(AssetData[0].GetAsset());
+		}
+	})
+	.Content()
+	[
+		SNew (STextBlock).Text_Lambda([EvaluateChooser]()
+		{
+			if (EvaluateChooser->Chooser == nullptr)
+			{
+				return LOCTEXT("null", "[No Chooser]");
+			}
+			else
+			{
+				return FText::FromString(EvaluateChooser->Chooser->GetName());
+			}
+		})
+	];
+}
+
+void ConvertToText_Base(const UObject* Object, FText& OutText)
+{
+	OutText = FText::FromString(Object->GetName());
+}
+
+void ConvertToText_Asset(const UObject* Object, FText& OutText)
+{
+	const UDataInterface_Object_Asset* AssetInterface = Cast<UDataInterface_Object_Asset>(Object);
+	if (AssetInterface->Asset == nullptr)
+	{
+		OutText = LOCTEXT("null", "[No Chooser]");
+	}
+	else
+	{
+		OutText = FText::FromString(AssetInterface->Asset->GetName());
+	}
+}
+
+void ConvertToText_EvaluateChooser(const UObject* Object, FText& OutText)
+{
+	const UDataInterface_EvaluateChooser* EvaluateChooser = Cast<UDataInterface_EvaluateChooser>(Object);
+	if (EvaluateChooser->Chooser == nullptr)
+	{
+		OutText = LOCTEXT("null", "[No Chooser]");
+	}
+	else
+	{
+		OutText = FText::FromString(EvaluateChooser->Chooser->GetName());
+	}
+}
+
+TSharedRef<SWidget> CreateObjectWidget(UObject* Object)
+{
+	FText ObjectText = FText::FromString(Object->GetName());
+	ConvertToText(Object, ObjectText);
+	
+	return SNew (STextBlock).Text(ObjectText); // could make this use Text_Lambda, to have it update correctly?
+}
+
+TSharedRef<SWidget> CreateBoolColumnWidget(UObject* Column, int Row)
+{
+	UChooserColumnBool* BoolColumn = Cast<UChooserColumnBool>(Column);
+
+	return SNew (SCheckBox)
+	.OnCheckStateChanged_Lambda([BoolColumn,Row](ECheckBoxState State)
+	{
+		if (Row < BoolColumn->RowValues.Num())
+		{
+			BoolColumn->RowValues[Row] = (State == ECheckBoxState::Checked);
+		}
+	})
+	.IsChecked_Lambda([BoolColumn, Row]()
+	{
+		const bool value = (Row < BoolColumn->RowValues.Num()) ? BoolColumn->RowValues[Row] : false;
+		return value ? ECheckBoxState::Checked : ECheckBoxState::Unchecked;
+	});
+}
+
+TSharedRef<SWidget> CreateFloatRangeColumnWidget(UObject* Column, int Row)
+{
+	UChooserColumnFloatRange* FloatRangeColumn = Cast<UChooserColumnFloatRange>(Column);
+
+	return SNew(SHorizontalBox)
+	+ SHorizontalBox::Slot().MaxWidth(10)
+	[
+		SNew(STextBlock).Text(LOCTEXT("FloatRangeLeft", "("))
+	]
+	+ SHorizontalBox::Slot().FillWidth(1.0)
+	[
+		SNew(SNumericEntryBox<float>)
+		.MaxValue_Lambda([FloatRangeColumn, Row]()
+		{
+			return (Row < FloatRangeColumn->RowValues.Num()) ? FloatRangeColumn->RowValues[Row].Max : 0;
+		})
+		.Value_Lambda([FloatRangeColumn, Row]()
+		{
+			return (Row < FloatRangeColumn->RowValues.Num()) ? FloatRangeColumn->RowValues[Row].Min : 0;
+		})
+		.OnValueCommitted_Lambda([FloatRangeColumn, Row](float NewValue, ETextCommit::Type CommitType)
+		{
+			if (Row < FloatRangeColumn->RowValues.Num())
+			{
+				FloatRangeColumn->RowValues[Row].Min = NewValue;
+			}
+		})
+	]
+	+ SHorizontalBox::Slot().MaxWidth(10)
+	[
+		SNew(STextBlock).Text(LOCTEXT("FloatRangeComma", " ,"))
+	]
+	+ SHorizontalBox::Slot().FillWidth(1.0)
+	[
+		SNew(SNumericEntryBox<float>)
+		.MinValue_Lambda([FloatRangeColumn, Row]()
+		{
+			return (Row < FloatRangeColumn->RowValues.Num()) ? FloatRangeColumn->RowValues[Row].Min : 0;
+		})
+		.Value_Lambda([FloatRangeColumn, Row]()
+		{
+			return (Row < FloatRangeColumn->RowValues.Num()) ? FloatRangeColumn->RowValues[Row].Max : 0;
+		})
+		.OnValueCommitted_Lambda([FloatRangeColumn, Row](float NewValue, ETextCommit::Type CommitType)
+		{
+			if (Row < FloatRangeColumn->RowValues.Num())
+			{
+				FloatRangeColumn->RowValues[Row].Max = NewValue;
+			}
+		})
+	]
+	+ SHorizontalBox::Slot().MaxWidth(10)
+	[
+		SNew(STextBlock).Text(LOCTEXT("FloatRangeRight", " )"))
+	];
+}
+
+class FChooserRowDetails : public IDetailCustomization
+{
+public:
+	FChooserRowDetails() {};
+	virtual ~FChooserRowDetails() override {};
+
+	static TSharedRef<IDetailCustomization> MakeInstance()
+	{
+		return MakeShareable( new FChooserRowDetails() );
+	}
+
+	// IDetailCustomization interface
+	virtual void CustomizeDetails(class IDetailLayoutBuilder& DetailBuilder) override;
+};
+
+// Make the details panel show the values for the selected row, showing each column value
+void FChooserRowDetails::CustomizeDetails(IDetailLayoutBuilder& DetailBuilder)
+{
+	TArray<TWeakObjectPtr<UObject>> Objects;
+	DetailBuilder.GetObjectsBeingCustomized(Objects);
+
+	UChooserRowDetails* Row = Cast<UChooserRowDetails>(Objects[0]);
+	auto Chooser = Row->Chooser;
+	
+	if (Chooser->Results.Num() > Row->Row)
+	{
+		IDetailCategoryBuilder& PropertiesCategory = DetailBuilder.EditCategory("Row Properties");
+
+		auto ChooserProperty = DetailBuilder.GetProperty("Chooser", Row->StaticClass());
+		DetailBuilder.HideProperty(ChooserProperty);
+	
+		auto ResultsArrayProperty = ChooserProperty->GetChildHandle("Results");
+		auto CurrentResultProperty = ResultsArrayProperty->AsArray()->GetElement(Row->Row);
+		IDetailPropertyRow& NewResultProperty = PropertiesCategory.AddProperty(CurrentResultProperty);
+		NewResultProperty.DisplayName(LOCTEXT("ResultColumnName","Result"));
+		NewResultProperty.ShowPropertyButtons(false); // hide array add button
+		NewResultProperty.ShouldAutoExpand(true);
+	
+		auto ColumnsArrayProperty = ChooserProperty->GetChildHandle("Columns");
+	
+		for(int ColumnIndex=0;ColumnIndex<Chooser->Columns.Num(); ColumnIndex++)
+		{
+			auto Column = Chooser->Columns[ColumnIndex];
+		
+			auto ColumnProperty = ColumnsArrayProperty->AsArray()->GetElement(ColumnIndex);
+			auto ColumnDataProperty = ColumnProperty->GetChildHandle("RowValues");
+			uint32 NumElements = 0;
+			ColumnDataProperty->AsArray()->GetNumElements(NumElements);
+			if (Row->Row < (int)NumElements)
+			{
+				auto CellData = ColumnDataProperty->AsArray()->GetElement(Row->Row);
+
+				IDetailPropertyRow& NewColumnProperty = PropertiesCategory.AddProperty(CellData);
+				NewColumnProperty.DisplayName(FText::FromName(Column.GetInterface()->GetDisplayName()));
+				NewColumnProperty.ShowPropertyButtons(false); // hide array add button
+				NewColumnProperty.ShouldAutoExpand(true);
+			}
+		}
+	}
+}
+
+void FChooserTableEditor::RegisterWidgets()
+{
+	using namespace DataInterfaceGraphEditor;
+	FDataInterfaceWidgetFactories::DataInterfaceTextConverter.Add(UObject::StaticClass(), ConvertToText_Base);
+	FDataInterfaceWidgetFactories::DataInterfaceTextConverter.Add(UDataInterface_Object_Asset::StaticClass(), ConvertToText_Asset);
+	FDataInterfaceWidgetFactories::DataInterfaceTextConverter.Add(UDataInterface_EvaluateChooser::StaticClass(), ConvertToText_EvaluateChooser);
+	
+	FDataInterfaceWidgetFactories::DataInterfaceWidgetCreators.Add(UObject::StaticClass(), CreateObjectWidget);
+	FDataInterfaceWidgetFactories::DataInterfaceWidgetCreators.Add(UDataInterface_Object_Asset::StaticClass(), CreateAssetWidget);
+	FDataInterfaceWidgetFactories::DataInterfaceWidgetCreators.Add(UDataInterface_EvaluateChooser::StaticClass(), CreateEvaluateChooserWidget);
+
+	ColumnWidgetCreators.Add(UChooserColumnBool::StaticClass(), CreateBoolColumnWidget);
+	ColumnWidgetCreators.Add(UChooserColumnFloatRange::StaticClass(), CreateFloatRangeColumnWidget);
+
+	FPropertyEditorModule& PropertyModule = FModuleManager::LoadModuleChecked<FPropertyEditorModule>("PropertyEditor");
+	PropertyModule.RegisterCustomClassLayout("ChooserRowDetails", FOnGetDetailCustomizationInstance::CreateStatic(&FChooserRowDetails::MakeInstance));	
+}
+}
+
+#undef LOCTEXT_NAMESPACE
diff --git a/Engine/Plugins/Experimental/Chooser/Source/ChooserEditor/Private/ChooserTableEditor.h b/Engine/Plugins/Experimental/Chooser/Source/ChooserEditor/Private/ChooserTableEditor.h
new file mode 100644
--- /dev/null
+++ b/Engine/Plugins/Experimental/Chooser/Source/ChooserEditor/Private/ChooserTableEditor.h
@@ -0,0 +1,136 @@
+// Copyright Epic Games, Inc. All Rights Reserved.
+
+#pragma once
+
+#include "CoreMinimal.h"
+#include "Toolkits/IToolkitHost.h"
+#include "Toolkits/AssetEditorToolkit.h"
+#include "Editor/PropertyEditor/Public/PropertyEditorDelegates.h"
+#include "Widgets/Views/STableViewBase.h"
+#include "Chooser.h"
+#include "IDetailCustomization.h"
+#include "ChooserTableEditor.generated.h"
+
+class IDetailsView;
+
+
+// Class used for chooser editor details customization
+UCLASS()
+class UChooserRowDetails : public UObject
+{
+	GENERATED_BODY()
+public:
+	UPROPERTY(EditAnywhere, Instanced, Category="Hidden")
+	TObjectPtr<UChooserTable> Chooser;
+	int Row;
+};
+
+namespace UE::ChooserEditor
+{
+	class FChooserTableEditor : public FAssetEditorToolkit
+	{
+	public:
+		/** Delegate that, given an array of assets, returns an array of objects to use in the details view of an FSimpleAssetEditor */
+		DECLARE_DELEGATE_RetVal_OneParam(TArray<UObject*>, FGetDetailsViewObjects, const TArray<UObject*>&);
+
+		virtual void RegisterTabSpawners(const TSharedRef<class FTabManager>& TabManager) override;
+		virtual void UnregisterTabSpawners(const TSharedRef<class FTabManager>& TabManager) override;
+
+
+		/**
+		* Edits the specified asset object
+		*
+		* @param	Mode					Asset editing mode for this editor (standalone or world-centric)
+		* @param	InitToolkitHost			When Mode is WorldCentric, this is the level editor instance to spawn this editor within
+		* @param	ObjectsToEdit			The object to edit
+		* @param	GetDetailsViewObjects	If bound, a delegate to get the array of objects to use in the details view; uses ObjectsToEdit if not bound
+		*/
+		void InitEditor( const EToolkitMode::Type Mode, const TSharedPtr< class IToolkitHost >& InitToolkitHost, const TArray<UObject*>& ObjectsToEdit, FGetDetailsViewObjects GetDetailsViewObjects );
+
+		/** Destructor */
+		virtual ~FChooserTableEditor();
+
+		/** IToolkit interface */
+		virtual FName GetToolkitFName() const override;
+		virtual FText GetBaseToolkitName() const override;
+		virtual FText GetToolkitName() const override;
+		virtual FText GetToolkitToolTipText() const override;
+		virtual FString GetWorldCentricTabPrefix() const override;
+		virtual FLinearColor GetWorldCentricTabColorScale() const override;
+		virtual bool IsPrimaryEditor() const override { return true; }
+		virtual bool IsSimpleAssetEditor() const override { return false; }
+
+		/** FAssetEditorToolkit interface */
+		virtual void PostRegenerateMenusAndToolbars() override;
+
+		UChooserTable* GetChooser() { return Cast<UChooserTable>(EditingObjects[0]); }
+	
+		/** Used to show or hide certain properties */
+		void SetPropertyVisibilityDelegate(FIsPropertyVisible InVisibilityDelegate);
+		/** Can be used to disable the details view making it read-only */
+		void SetPropertyEditingEnabledDelegate(FIsPropertyEditingEnabled InPropertyEditingDelegate);
+	
+		struct FChooserTableRow
+		{
+			FChooserTableRow(int32 i) { RowIndex = i; }
+			int32 RowIndex;
+		};
+
+		void UpdateTableRows();
+	private:
+
+		FReply SelectRootProperties();
+	
+		/** Create the properties tab and its content */
+		TSharedRef<SDockTab> SpawnPropertiesTab( const FSpawnTabArgs& Args );
+		/** Create the table tab and its content */
+		TSharedRef<SDockTab> SpawnTableTab( const FSpawnTabArgs& Args );
+	
+		TSharedRef<ITableRow> GenerateTableRow(TSharedPtr<FChooserTableRow> InItem, const TSharedRef<STableViewBase>& OwnerTable);
+
+		/** Handles when an asset is imported */
+		void HandleAssetPostImport(class UFactory* InFactory, UObject* InObject);
+
+		/** Called when objects need to be swapped out for new versions, like after a blueprint recompile. */
+		void OnObjectsReplaced(const TMap<UObject*, UObject*>& ReplacementMap);
+
+		/** Details view */
+		TSharedPtr< class IDetailsView > DetailsView;
+
+		/** App Identifier. */
+		static const FName ChooserEditorAppIdentifier;
+
+		/**	The tab ids for all the tabs used */
+		static const FName PropertiesTabId;
+		static const FName TableTabId;
+
+		/** The objects open within this editor */
+		TArray<UObject*> EditingObjects;
+
+		TArray<TObjectPtr<UChooserRowDetails>> SelectedRows;
+
+		void UpdateTableColumns();
+		TArray<TSharedPtr<FChooserTableRow>> TableRows;
+	
+		TSharedPtr<SComboButton> CreateColumnComboButton;
+		TSharedPtr<SHeaderRow> HeaderRow;
+		TSharedPtr<SListView<TSharedPtr<FChooserTableRow>>> TableView;
+
+		TArray<TSharedPtr<SEditableText>> ColumnText;
+
+		FName SelectedColumn;
+		bool EditColumnName = false;
+	public:
+
+		static TMap<const UClass*, TFunction<TSharedRef<SWidget> (UObject* Column, int Row)>> ColumnWidgetCreators;
+	
+		/** The name given to all instances of this type of editor */
+		static const FName ToolkitFName;
+
+		static TSharedRef<FChooserTableEditor> CreateEditor( const EToolkitMode::Type Mode, const TSharedPtr< IToolkitHost >& InitToolkitHost, UObject* ObjectToEdit, FGetDetailsViewObjects GetDetailsViewObjects = FGetDetailsViewObjects() );
+
+		static TSharedRef<FChooserTableEditor> CreateEditor( const EToolkitMode::Type Mode, const TSharedPtr< IToolkitHost >& InitToolkitHost, const TArray<UObject*>& ObjectsToEdit, FGetDetailsViewObjects GetDetailsViewObjects = FGetDetailsViewObjects() );
+
+		static void RegisterWidgets();
+	};
+}
diff --git a/Engine/Plugins/Experimental/DataInterface/Source/DataInterface/Public/DataInterfaceTypes.h b/Engine/Plugins/Experimental/DataInterface/Source/DataInterface/Public/DataInterfaceTypes.h
--- a/Engine/Plugins/Experimental/DataInterface/Source/DataInterface/Public/DataInterfaceTypes.h
+++ b/Engine/Plugins/Experimental/DataInterface/Source/DataInterface/Public/DataInterfaceTypes.h
@@ -8,6 +8,13 @@
 
 class IDataInterface;
 
+#define DATA_INTERFACE_RETURN_TYPE(Typename) \
+virtual FName GetReturnTypeNameImpl() const final override\
+{\
+	static FName ReturnTypeNameValue = #Typename;\
+	return ReturnTypeNameValue;\
+}\
+
 #define DECLARE_DATA_INTERFACE_PARAM_TYPE_INTERNAL(ValueType, Identifier) \
 	template<> \
 	struct UE::DataInterface::Private::TParamTypeImpl<ValueType> \
diff --git a/Engine/Plugins/Experimental/DataInterfaceGraph/Source/DataInterfaceGraph/Private/DataInterface_Bool.cpp b/Engine/Plugins/Experimental/DataInterfaceGraph/Source/DataInterfaceGraph/Private/DataInterface_Bool.cpp
new file mode 100644
--- /dev/null
+++ b/Engine/Plugins/Experimental/DataInterfaceGraph/Source/DataInterfaceGraph/Private/DataInterface_Bool.cpp
@@ -0,0 +1,50 @@
+﻿// Copyright Epic Games, Inc. All Rights Reserved.
+
+#include "DataInterface_Bool.h"
+#include "DataInterface.h"
+#include "DataInterfaceTypes.h"
+#include "DataInterfaceKernel.h"
+
+bool UDataInterface_Bool_And::GetDataImpl(const UE::DataInterface::FContext& Context) const
+{
+	using namespace UE::DataInterface;
+	
+	check(Inputs.Num() > 0);
+
+	bool bResult = true;
+	TAllocParam<bool> IntermediateParam(Context);
+	TParam<bool> Result = Context.GetResult<bool>();
+	
+	for(const TScriptInterface<IDataInterface>& Input : Inputs)
+	{
+		bResult &= UE::DataInterface::GetDataSafe(Input, Context, IntermediateParam);
+
+		FKernel::Run(Context,
+			[](bool& OutResult, bool InIntermediate)
+			{
+				OutResult = OutResult && InIntermediate;
+			},
+			Result, IntermediateParam);
+	}
+
+	return bResult;
+}
+
+bool UDataInterface_Bool_Not::GetDataImpl(const UE::DataInterface::FContext& Context) const
+{
+	using namespace UE::DataInterface;
+
+	TParam<bool> Result = Context.GetResult<bool>();
+	
+	TAllocParam<bool> CurrentValue(Context);
+	bool bResult = UE::DataInterface::GetDataSafe(Input, Context, CurrentValue);
+
+	FKernel::Run(Context,
+		[](bool& OutResult, bool InCurrent)
+		{
+			OutResult = !InCurrent;
+		},
+		Result, CurrentValue);
+
+	return bResult;
+}
diff --git a/Engine/Plugins/Experimental/DataInterfaceGraph/Source/DataInterfaceGraph/Private/DataInterface_Bool.h b/Engine/Plugins/Experimental/DataInterfaceGraph/Source/DataInterfaceGraph/Private/DataInterface_Bool.h
new file mode 100644
--- /dev/null
+++ b/Engine/Plugins/Experimental/DataInterfaceGraph/Source/DataInterfaceGraph/Private/DataInterface_Bool.h
@@ -0,0 +1,53 @@
+// Copyright Epic Games, Inc. All Rights Reserved.
+
+#pragma once
+
+#include "CoreMinimal.h"
+#include "DataInterfaceTypes.h"
+#include "IDataInterface.h"
+#include "DataInterfaceContext.h"
+#include "DataInterface_Bool.generated.h"
+
+UCLASS()
+class UDataInterface_Bool_Literal : public UObject, public IDataInterface
+{
+	GENERATED_BODY()
+	DATA_INTERFACE_RETURN_TYPE(bool)
+
+	// IDataInterface interface
+	virtual bool GetDataImpl(const UE::DataInterface::FContext& Context) const final override
+	{
+		Context.SetResult(Value);
+		return true;
+	}
+	
+	UPROPERTY(EditAnywhere, Category = "Parameters")
+	bool Value;
+};
+
+UCLASS()
+class UDataInterface_Bool_And : public UObject, public IDataInterface
+{
+	GENERATED_BODY()
+	DATA_INTERFACE_RETURN_TYPE(bool)
+
+	// IDataInterface interface
+	virtual bool GetDataImpl(const UE::DataInterface::FContext& Context) const final override;
+
+	UPROPERTY(EditAnywhere, Meta=(DataInterfaceType="bool"), Category = "Parameters")
+	TArray<TScriptInterface<IDataInterface>> Inputs;
+};
+
+
+UCLASS()
+class UDataInterface_Bool_Not : public UObject, public IDataInterface
+{
+	GENERATED_BODY()
+	DATA_INTERFACE_RETURN_TYPE(bool)
+
+	// IDataInterface interface
+	virtual bool GetDataImpl(const UE::DataInterface::FContext& Context) const final override;
+
+	UPROPERTY(EditAnywhere, Meta=(DataInterfaceType="bool"), Category = "Parameters")
+	TScriptInterface<IDataInterface> Input;
+};
\ No newline at end of file
diff --git a/Engine/Plugins/Experimental/DataInterfaceGraph/Source/DataInterfaceGraphEditor/Private/DataInterfacePropertyTypeCustomization.cpp b/Engine/Plugins/Experimental/DataInterfaceGraph/Source/DataInterfaceGraphEditor/Private/DataInterfacePropertyTypeCustomization.cpp
--- a/Engine/Plugins/Experimental/DataInterfaceGraph/Source/DataInterfaceGraphEditor/Private/DataInterfacePropertyTypeCustomization.cpp
+++ b/Engine/Plugins/Experimental/DataInterfaceGraph/Source/DataInterfaceGraphEditor/Private/DataInterfacePropertyTypeCustomization.cpp
@@ -1,11 +1,18 @@
 ﻿// Copyright Epic Games, Inc. All Rights Reserved.
 
 #include "DataInterfacePropertyTypeCustomization.h"
+
+#include "DataInterfaceGraph.h"
 #include "PropertyHandle.h"
 #include "IDataInterface.h"
 #include "DetailWidgetRow.h"
 #include "PropertyCustomizationHelpers.h"
 #include "Widgets/Text/STextBlock.h"
+#include "SClassViewer.h"
+#include "DataInterfaceWidgetFactories.h"
+#include "IDetailChildrenBuilder.h"
+#include "PropertyEditor/Private/ObjectPropertyNode.h"
+#include "PropertyEditor/Private/PropertyNode.h"
 
 #define LOCTEXT_NAMESPACE "DataInterfacePropertyTypeCustomization"
 
@@ -14,13 +21,38 @@ namespace UE::DataInterfaceGraphEditor
 
 bool FPropertyTypeIdentifier::IsPropertyTypeCustomized(const IPropertyHandle& PropertyHandle) const
 {
-	FInterfaceProperty* InterfaceProperty = CastField<FInterfaceProperty>(PropertyHandle.GetProperty());
-	return InterfaceProperty && InterfaceProperty->InterfaceClass->IsChildOf(UDataInterface::StaticClass());
+	return (PropertyHandle.GetMetaData("DataInterfaceType") != "");
 }
 
 void FDataInterfacePropertyTypeCustomization::CustomizeHeader(TSharedRef<IPropertyHandle> PropertyHandle, FDetailWidgetRow& HeaderRow, IPropertyTypeCustomizationUtils& CustomizationUtils)
 {
-	FInterfaceProperty* InterfaceProperty = CastField<FInterfaceProperty>(PropertyHandle->GetProperty());
+	FString TypeName = PropertyHandle->GetMetaData("DataInterfaceType");
+	FName DataInterfaceTypeName(TypeName);
+
+	void* ValuePtr;
+	PropertyHandle->GetValueData(ValuePtr);
+	FScriptInterface* PropertyValue = reinterpret_cast<FScriptInterface*>(ValuePtr);
+	
+	TSharedPtr<SWidget> Widget = FDataInterfaceWidgetFactories::CreateDataInterfaceWidget(DataInterfaceTypeName, PropertyValue->GetObject(),
+			FOnClassPicked::CreateLambda([this, PropertyHandle](UClass* ChosenClass)
+			{
+				TArray<void*> RawData;
+				PropertyHandle->AccessRawData(RawData);
+				TArray<UPackage*> OuterObjects;
+				PropertyHandle->GetOuterPackages(OuterObjects);
+
+				for(int i=0;i<RawData.Num();i++)
+				{
+					UObject* NewValue = NewObject<UObject>(OuterObjects[i], ChosenClass);
+				
+					PropertyHandle->NotifyPreChange();
+					PropertyHandle->SetValue(NewValue);
+					PropertyHandle->NotifyPostChange(EPropertyChangeType::ValueSet);
+					PropertyHandle->GetPropertyNode()->GetParentNode()->RequestRebuildChildren();
+				}
+			}),
+			nullptr
+	);
 
 	HeaderRow
 	.NameContent()
@@ -29,23 +61,27 @@ void FDataInterfacePropertyTypeCustomization::CustomizeHeader(TSharedRef<IProper
 	]
 	.ValueContent()
 	[
-		SNew(SObjectPropertyEntryBox)
-		.PropertyHandle(PropertyHandle)
-		.DisplayUseSelected(false)
-		.OnShouldFilterAsset_Lambda([InterfaceProperty](const FAssetData& InAssetData)
-		{
-			if(UClass* Class = InAssetData.GetClass())
-			{
-				if(Class->ImplementsInterface(InterfaceProperty->InterfaceClass))
-				{
-					return false;
-				}
-			}
-
-			return true;
-		})
+		Widget.ToSharedRef()
 	];
 }
+	
+void FDataInterfacePropertyTypeCustomization::CustomizeChildren(TSharedRef<IPropertyHandle> PropertyHandle, IDetailChildrenBuilder& ChildBuilder, IPropertyTypeCustomizationUtils& CustomizationUtils)
+{
+	// todo: make this work correctly for multiselect
+	void* ValuePtr;
+	PropertyHandle->GetValueData(ValuePtr);
+	FScriptInterface* PropertyValue = reinterpret_cast<FScriptInterface*>(ValuePtr);
+	if (PropertyValue->GetObject())
+	{
+		TSharedPtr<FStructOnScope> StructData = MakeShareable(new FStructOnScope(PropertyValue->GetObject()->GetClass(), (uint8*)PropertyValue->GetObject()));
+		StructData->SetPackage(PropertyValue->GetObject()->GetPackage());
+		TArray<TSharedPtr<IPropertyHandle>> NewHandles = ChildBuilder.AddAllExternalStructureProperties(StructData.ToSharedRef());
+		for (auto& NewHandle : NewHandles)
+		{
+			ChildBuilder.AddProperty(NewHandle.ToSharedRef());
+		}
+	}
+}
 
 }
 
diff --git a/Engine/Plugins/Experimental/DataInterfaceGraph/Source/DataInterfaceGraphEditor/Private/DataInterfacePropertyTypeCustomization.h b/Engine/Plugins/Experimental/DataInterfaceGraph/Source/DataInterfaceGraphEditor/Private/DataInterfacePropertyTypeCustomization.h
--- a/Engine/Plugins/Experimental/DataInterfaceGraph/Source/DataInterfaceGraphEditor/Private/DataInterfacePropertyTypeCustomization.h
+++ b/Engine/Plugins/Experimental/DataInterfaceGraph/Source/DataInterfaceGraphEditor/Private/DataInterfacePropertyTypeCustomization.h
@@ -13,7 +13,7 @@ class FDataInterfacePropertyTypeCustomization : public IPropertyTypeCustomizatio
 private:
 	// IPropertyTypeCustomization interface
 	virtual void CustomizeHeader(TSharedRef<IPropertyHandle> PropertyHandle, FDetailWidgetRow& HeaderRow, IPropertyTypeCustomizationUtils& CustomizationUtils) override;
-	virtual void CustomizeChildren(TSharedRef<IPropertyHandle> PropertyHandle, IDetailChildrenBuilder& ChildBuilder, IPropertyTypeCustomizationUtils& CustomizationUtils) override {}
+	virtual void CustomizeChildren(TSharedRef<IPropertyHandle> PropertyHandle, IDetailChildrenBuilder& ChildBuilder, IPropertyTypeCustomizationUtils& CustomizationUtils) override;
 };
 
 class FPropertyTypeIdentifier : public IPropertyTypeIdentifier
diff --git a/Engine/Plugins/Experimental/DataInterfaceGraph/Source/DataInterfaceGraphEditor/Private/DataInterfaceWidgetFactories.cpp b/Engine/Plugins/Experimental/DataInterfaceGraph/Source/DataInterfaceGraphEditor/Private/DataInterfaceWidgetFactories.cpp
new file mode 100644
--- /dev/null
+++ b/Engine/Plugins/Experimental/DataInterfaceGraph/Source/DataInterfaceGraphEditor/Private/DataInterfaceWidgetFactories.cpp
@@ -0,0 +1,122 @@
+// Copyright Epic Games, Inc. All Rights Reserved.
+
+#include "DataInterfaceWidgetFactories.h"
+#include "Modules/ModuleManager.h"
+#include "Widgets/Input/SHyperlink.h"
+#include "Widgets/Input/SNumericEntryBox.h"
+#include "Widgets/Input/SComboButton.h"
+#include "SClassViewer.h"
+#include "DetailCategoryBuilder.h"
+#include "Widgets/Layout/SWidgetSwitcher.h"
+#include "DataInterfaceClassFilter.h"
+
+#define LOCTEXT_NAMESPACE "DataInterfaceEditor"
+
+namespace UE::DataInterfaceGraphEditor
+{
+
+void ConvertToText(UObject* Object, FText& OutText)
+{
+	UClass* Class = Object->GetClass();
+	while (Class)
+	{
+		if (auto TextConverter = FDataInterfaceWidgetFactories::DataInterfaceTextConverter.Find(Class))
+		{
+			(*TextConverter)(Object, OutText);
+			break;
+		}
+		Class = Class->GetSuperClass();
+	}
+}
+
+TSharedPtr<SWidget> FDataInterfaceWidgetFactories::CreateDataInterfaceWidget(FName DataInterfaceTypeName, UObject* Value, const FOnClassPicked& CreateClassCallback, TSharedPtr<SBorder>* InnerWidget)
+{
+	TSharedPtr<SWidget> LeftWidget;
+
+	if (Value)
+	{
+		UClass* Class = Value->GetClass();
+		while (Class && !LeftWidget.IsValid())
+		{
+			if (auto Creator = FDataInterfaceWidgetFactories::DataInterfaceWidgetCreators.Find(Class))
+			{
+				LeftWidget = (*Creator)(Value);
+				break;
+			}
+			Class = Class->GetSuperClass();
+		}
+	}
+	
+	if (!LeftWidget.IsValid())
+	{
+		LeftWidget = SNew(STextBlock).Text(LOCTEXT("SelectDataType", "Select Data Type..."));
+	}
+
+
+	// button for replacing data with a different Data Interface class
+	TSharedPtr<SComboButton> Button = SNew(SComboButton)
+			.ComboButtonStyle(FAppStyle::Get(), "SimpleComboButton");
+	
+	Button->SetOnGetMenuContent(FOnGetContent::CreateLambda([DataInterfaceTypeName, Button, CreateClassCallback]()
+	{
+		FClassViewerInitializationOptions Options;
+		Options.ClassFilters.Add(MakeShared<FDataInterfaceClassFilter>(DataInterfaceTypeName, nullptr));
+		
+		// Add class filter for columns here
+		TSharedRef<SWidget>  ClassMenu = FModuleManager::LoadModuleChecked<FClassViewerModule>("ClassViewer").
+		CreateClassViewer(Options, FOnClassPicked::CreateLambda([Button, CreateClassCallback](UClass* Class)
+		{
+			Button->SetIsOpen(false);
+			CreateClassCallback.Execute(Class);
+		}));
+		return ClassMenu;
+	}));
+
+	TSharedPtr <SBorder> Border;
+	if (InnerWidget && InnerWidget->IsValid())
+	{
+		Border = *InnerWidget;
+	}
+	else
+	{
+		Border = SNew(SBorder);
+	}
+	
+	if (InnerWidget)
+	{
+		*InnerWidget = Border;
+	}
+	
+	Border->SetContent(LeftWidget.ToSharedRef());
+
+	TSharedPtr<SWidget> Widget = SNew(SHorizontalBox)
+		+SHorizontalBox::Slot().FillWidth(75)
+		[
+			Border.ToSharedRef()
+		]
+		+SHorizontalBox::Slot().FillWidth(25)
+		[
+			Button.ToSharedRef()
+		]
+	;
+
+
+	return Widget;
+}
+
+TMap<const UClass*, TFunction<void (const UObject* Object, FText& OutText)>> FDataInterfaceWidgetFactories::DataInterfaceTextConverter;
+TMap<const UClass*, TFunction<TSharedRef<SWidget> (UObject* Object)>> FDataInterfaceWidgetFactories::DataInterfaceWidgetCreators;
+
+void ConvertToText_Base(const UObject* Object, FText& OutText)
+{
+	OutText = FText::FromString(Object->GetName());
+}
+
+void FDataInterfaceWidgetFactories::RegisterWidgets()
+{
+	DataInterfaceTextConverter.Add(UObject::StaticClass(), ConvertToText_Base);
+}
+	
+}
+
+#undef LOCTEXT_NAMESPACE
diff --git a/Engine/Plugins/Experimental/DataInterfaceGraph/Source/DataInterfaceGraphEditor/Public/DataInterfaceClassFilter.h b/Engine/Plugins/Experimental/DataInterfaceGraph/Source/DataInterfaceGraphEditor/Public/DataInterfaceClassFilter.h
new file mode 100644
--- /dev/null
+++ b/Engine/Plugins/Experimental/DataInterfaceGraph/Source/DataInterfaceGraphEditor/Public/DataInterfaceClassFilter.h
@@ -0,0 +1,50 @@
+// Copyright Epic Games, Inc. All Rights Reserved.
+#pragma once
+
+#include "ClassViewerFilter.h"
+#include "IDataInterface.h"
+
+namespace UE::DataInterfaceGraphEditor
+{
+	// Filter class for things which implement IChooserColumn
+	class FDataInterfaceClassFilter : public IClassViewerFilter
+	{
+	public:
+		FDataInterfaceClassFilter(FName InTypeName, const UScriptStruct* InStructType) : TypeName(InTypeName), StructType(InStructType) { };
+		virtual ~FDataInterfaceClassFilter() override {};
+	
+		virtual bool IsClassAllowed(const FClassViewerInitializationOptions& InInitOptions, const UClass* InClass, TSharedRef< class FClassViewerFilterFuncs > InFilterFuncs ) override
+		{
+			if (InClass->ImplementsInterface(UDataInterface::StaticClass()))
+			{
+				// pass filter if the class has "" return type (for wrappers)
+				// pass all classes if the requested typename is ""
+				// otherwise fail if the TypeNames don't match
+				// for Object type, check type hierarchy. (todo)
+			
+				IDataInterface* Object = static_cast<IDataInterface*>(InClass->GetDefaultObject()->GetInterfaceAddress(UDataInterface::StaticClass()));
+				if (Object->GetReturnTypeName() == TypeName || Object->GetReturnTypeName() == "" || TypeName == "")
+				{
+					if (const UScriptStruct* ReturnStructType = Object->GetReturnTypeStruct())
+					{
+						return true; // kytodo: compare the struct types to see if they are compatible
+					}
+					else
+					{
+						return true;
+					}
+				}
+			}
+
+			return false;
+		}
+
+		virtual bool IsUnloadedClassAllowed(const FClassViewerInitializationOptions& InInitOptions, const TSharedRef< const class IUnloadedBlueprintData > InUnloadedClassData, TSharedRef< class FClassViewerFilterFuncs > InFilterFuncs) override
+		{
+			return false;
+		}
+	private:
+		FName TypeName;
+		const UScriptStruct* StructType;
+	};
+}
\ No newline at end of file
diff --git a/Engine/Plugins/Experimental/DataInterfaceGraph/Source/DataInterfaceGraphEditor/Public/DataInterfaceWidgetFactories.h b/Engine/Plugins/Experimental/DataInterfaceGraph/Source/DataInterfaceGraphEditor/Public/DataInterfaceWidgetFactories.h
new file mode 100644
--- /dev/null
+++ b/Engine/Plugins/Experimental/DataInterfaceGraph/Source/DataInterfaceGraphEditor/Public/DataInterfaceWidgetFactories.h
@@ -0,0 +1,22 @@
+// Copyright Epic Games, Inc. All Rights Reserved.
+
+#pragma once
+
+#include "CoreMinimal.h"
+#include "Widgets/SWidget.h"
+#include "Widgets/Layout/SBorder.h"
+#include "SClassViewer.h"
+
+namespace UE::DataInterfaceGraphEditor
+{
+	class DATAINTERFACEGRAPHEDITOR_API FDataInterfaceWidgetFactories 
+	{
+	public:
+		static TMap<const UClass*, TFunction<void (const UObject* Object, FText& OutText)>> DataInterfaceTextConverter;
+		static TMap<const UClass*, TFunction<TSharedRef<SWidget> (UObject* Object)>> DataInterfaceWidgetCreators;
+
+		static TSharedPtr<SWidget> CreateDataInterfaceWidget(FName DataInterfaceTypeName, UObject* Value, const FOnClassPicked& CreateClassCallback, TSharedPtr<SBorder>* InnerWidget = nullptr);
+		
+		static void RegisterWidgets();
+	};
+}
