fix nopch compile issues in Chooser plugin

diff --git a/Engine/Plugins/Experimental/Chooser/Source/ChooserEditor/Private/ChooserTableEditor.cpp b/Engine/Plugins/Experimental/Chooser/Source/ChooserEditor/Private/ChooserTableEditor.cpp
--- a/Engine/Plugins/Experimental/Chooser/Source/ChooserEditor/Private/ChooserTableEditor.cpp
+++ b/Engine/Plugins/Experimental/Chooser/Source/ChooserEditor/Private/ChooserTableEditor.cpp
@@ -7,7 +7,9 @@
 #include "PropertyEditorModule.h"
 #include "IDetailsView.h"
 #include "Editor.h"
+#include "Framework/MultiBox/MultiBoxBuilder.h"
 #include "Widgets/Input/SHyperlink.h"
+#include "Widgets/Input/SEditableText.h"
 #include "Widgets/Input/SNumericEntryBox.h"
 #include "Widgets/Views/SListView.h"
 #include "Widgets/Input/SComboButton.h"
diff --git a/Engine/Plugins/Experimental/Chooser/Source/ChooserEditor/Private/ChooserTableEditor.h b/Engine/Plugins/Experimental/Chooser/Source/ChooserEditor/Private/ChooserTableEditor.h
--- a/Engine/Plugins/Experimental/Chooser/Source/ChooserEditor/Private/ChooserTableEditor.h
+++ b/Engine/Plugins/Experimental/Chooser/Source/ChooserEditor/Private/ChooserTableEditor.h
@@ -11,6 +11,7 @@
 #include "IDetailCustomization.h"
 #include "ChooserTableEditor.generated.h"
 
+class SEditableText;
 class IDetailsView;
 
 
