fix details customisation causing crash in chooser editor when you select a row

diff --git a/Engine/Plugins/Experimental/Chooser/Source/Chooser/Internal/Chooser.h b/Engine/Plugins/Experimental/Chooser/Source/Chooser/Internal/Chooser.h
--- a/Engine/Plugins/Experimental/Chooser/Source/Chooser/Internal/Chooser.h
+++ b/Engine/Plugins/Experimental/Chooser/Source/Chooser/Internal/Chooser.h
@@ -116,7 +116,8 @@ public:
 	TArray<TScriptInterface<IChooserColumn>> Columns;
 
 	// array of results (rows of table)
-	UPROPERTY(EditAnywhere, Meta=(EditInline="true"), Category = "Runtime")
+	// todo: DataInterfaceType shouldn't be hard coded (Should be based on Result Type above) needed for Details customization to work
+	UPROPERTY(EditAnywhere, Meta=(DataInterfaceType="Asset"), Category = "Runtime")
 	TArray<TScriptInterface<IDataInterface>> Results;
 };
 
diff --git a/Engine/Plugins/Experimental/Chooser/Source/ChooserEditor/Private/ChooserTableEditor.cpp b/Engine/Plugins/Experimental/Chooser/Source/ChooserEditor/Private/ChooserTableEditor.cpp
--- a/Engine/Plugins/Experimental/Chooser/Source/ChooserEditor/Private/ChooserTableEditor.cpp
+++ b/Engine/Plugins/Experimental/Chooser/Source/ChooserEditor/Private/ChooserTableEditor.cpp
@@ -1,6 +1,7 @@
 // Copyright Epic Games, Inc. All Rights Reserved.
 
 #include "ChooserTableEditor.h"
+
 #include "Modules/ModuleManager.h"
 #include "Styling/AppStyle.h"
 #include "Widgets/Docking/SDockTab.h"
@@ -416,7 +417,7 @@ public:
 				const int ColumnIndex = ColumnName.GetNumber() - 1;
 				if (ColumnIndex < Chooser->Columns.Num() && ColumnIndex >=0)
 				{
-					auto& Column = Chooser->Columns[ColumnIndex];
+					TScriptInterface<IChooserColumn>& Column = Chooser->Columns[ColumnIndex];
 					TSharedPtr<SWidget> ColumnWidget;
 					UClass* ColumnClass = Column.GetObject()->GetClass();
 					while (ColumnClass && !ColumnWidget.IsValid())
@@ -518,7 +519,7 @@ void FChooserTableEditor::UpdateTableColumns()
 	int NumColumns = Chooser->Columns.Num();	
 	for(int ColumnIndex = 0; ColumnIndex < NumColumns; ColumnIndex++)
 	{
-		auto& Column = Chooser->Columns[ColumnIndex];
+		TScriptInterface<IChooserColumn>& Column = Chooser->Columns[ColumnIndex];
 
 		ColumnText.Push(
 			SNew(SEditableText)
@@ -1082,36 +1083,34 @@ void FChooserRowDetails::CustomizeDetails(IDetailLayoutBuilder& DetailBuilder)
 	DetailBuilder.GetObjectsBeingCustomized(Objects);
 
 	UChooserRowDetails* Row = Cast<UChooserRowDetails>(Objects[0]);
-	auto Chooser = Row->Chooser;
+	UChooserTable* Chooser = Row->Chooser;
 	
 	if (Chooser->Results.Num() > Row->Row)
 	{
 		IDetailCategoryBuilder& PropertiesCategory = DetailBuilder.EditCategory("Row Properties");
 
-		auto ChooserProperty = DetailBuilder.GetProperty("Chooser", Row->StaticClass());
+		TSharedPtr<IPropertyHandle> ChooserProperty = DetailBuilder.GetProperty("Chooser", Row->StaticClass());
 		DetailBuilder.HideProperty(ChooserProperty);
 	
-		auto ResultsArrayProperty = ChooserProperty->GetChildHandle("Results");
-		auto CurrentResultProperty = ResultsArrayProperty->AsArray()->GetElement(Row->Row);
+		TSharedPtr<IPropertyHandle> ResultsArrayProperty = ChooserProperty->GetChildHandle("Results");
+		TSharedPtr<IPropertyHandle> CurrentResultProperty = ResultsArrayProperty->AsArray()->GetElement(Row->Row);
 		IDetailPropertyRow& NewResultProperty = PropertiesCategory.AddProperty(CurrentResultProperty);
 		NewResultProperty.DisplayName(LOCTEXT("ResultColumnName","Result"));
 		NewResultProperty.ShowPropertyButtons(false); // hide array add button
 		NewResultProperty.ShouldAutoExpand(true);
 	
-		auto ColumnsArrayProperty = ChooserProperty->GetChildHandle("Columns");
-	
-		for(int ColumnIndex=0;ColumnIndex<Chooser->Columns.Num(); ColumnIndex++)
+		for(int ColumnIndex=0; ColumnIndex<Chooser->Columns.Num(); ColumnIndex++)
 		{
-			auto Column = Chooser->Columns[ColumnIndex];
-		
-			auto ColumnProperty = ColumnsArrayProperty->AsArray()->GetElement(ColumnIndex);
-			auto ColumnDataProperty = ColumnProperty->GetChildHandle("RowValues");
+			TScriptInterface<IChooserColumn>& Column = Chooser->Columns[ColumnIndex];
+			UObject* ColumnObject = Chooser->Columns[ColumnIndex].GetObject();
+			
+			TSharedPtr<IPropertyHandle> ColumnDataProperty =  DetailBuilder.AddObjectPropertyData({ColumnObject}, "RowValues");
 			uint32 NumElements = 0;
 			ColumnDataProperty->AsArray()->GetNumElements(NumElements);
 			if (Row->Row < (int)NumElements)
 			{
-				auto CellData = ColumnDataProperty->AsArray()->GetElement(Row->Row);
-
+				TSharedRef<IPropertyHandle> CellData = ColumnDataProperty->AsArray()->GetElement(Row->Row);
+	
 				IDetailPropertyRow& NewColumnProperty = PropertiesCategory.AddProperty(CellData);
 				NewColumnProperty.DisplayName(FText::FromName(Column.GetInterface()->GetDisplayName()));
 				NewColumnProperty.ShowPropertyButtons(false); // hide array add button
