Removing Data Interfaces dependencies from chooser prototype, and adding a Blueprint Callable function which can evaluate a chooser

diff --git a/Engine/Plugins/Experimental/Chooser/Chooser.uplugin b/Engine/Plugins/Experimental/Chooser/Chooser.uplugin
--- a/Engine/Plugins/Experimental/Chooser/Chooser.uplugin
+++ b/Engine/Plugins/Experimental/Chooser/Chooser.uplugin
@@ -16,10 +16,6 @@
 	"IsExperimentalVersion" : true,
 	"Installed" : false,
 	"Plugins": [
-		{
-			"Name": "DataInterfaceGraph",
-			"Enabled": true
-		}
 	],
 	"Modules" :
 	[
diff --git a/Engine/Plugins/Experimental/Chooser/Source/Chooser/Chooser.Build.cs b/Engine/Plugins/Experimental/Chooser/Source/Chooser/Chooser.Build.cs
--- a/Engine/Plugins/Experimental/Chooser/Source/Chooser/Chooser.Build.cs
+++ b/Engine/Plugins/Experimental/Chooser/Source/Chooser/Chooser.Build.cs
@@ -22,7 +22,6 @@ namespace UnrealBuildTool.Rules
 				new string[]
 				{
 					// ... add other public dependencies that you statically link with here ...
-					"DataInterface",
 				}
 			);
 
diff --git a/Engine/Plugins/Experimental/Chooser/Source/Chooser/Internal/Chooser.h b/Engine/Plugins/Experimental/Chooser/Source/Chooser/Internal/Chooser.h
--- a/Engine/Plugins/Experimental/Chooser/Source/Chooser/Internal/Chooser.h
+++ b/Engine/Plugins/Experimental/Chooser/Source/Chooser/Internal/Chooser.h
@@ -3,8 +3,9 @@
 
 #include "CoreMinimal.h"
 #include "UObject/Object.h"
-#include "IDataInterface.h"
-#include "DataInterface_Object.h"
+#include "IObjectChooser.h"
+#include "UObject/Interface.h"
+
 #include "Chooser.generated.h"
 
 UINTERFACE(NotBlueprintType, meta = (CannotImplementInterfaceInBlueprint))
@@ -18,7 +19,7 @@ class CHOOSER_API IChooserColumn
 	GENERATED_BODY()
 
 public:
-	virtual void Filter(const UE::DataInterface::FContext& Context, const TArray<uint32>& IndexListIn, TArray<uint32>& IndexListOut) {};
+	virtual void Filter(const UObject* ContextObject, const TArray<uint32>& IndexListIn, TArray<uint32>& IndexListOut) {};
 	virtual FName GetDisplayName() { return "ChooserColumn"; }
 	virtual void SetDisplayName(FName Name) { }
 	virtual void SetNumRows(uint32 NumRows) {}
@@ -32,23 +33,20 @@ class CHOOSER_API UChooserColumnBool : public UObject, public IChooserColumn
 	public:
 	UChooserColumnBool() { }
 	
-	UPROPERTY(EditAnywhere, Category = "Editor")
-	FName DisplayName = "Bool Column";
-	
-	UPROPERTY(EditAnywhere, Meta=(DataInterfaceType="bool"), Category = "Input")
-	TScriptInterface<IDataInterface> Value;
+	UPROPERTY(EditAnywhere, Category = "Input")
+	FName InputPropertyName;
 
 	UPROPERTY(EditAnywhere, Category=Runtime)
 	// array of results (cells for this column for each row in the table)
 	// should match the length of the Results array 
 	TArray<bool> RowValues;
 	
-	virtual void Filter(const UE::DataInterface::FContext& Context, const TArray<uint32>& IndexListIn, TArray<uint32>& IndexListOut) override;
+	virtual void Filter(const UObject* ContextObject, const TArray<uint32>& IndexListIn, TArray<uint32>& IndexListOut) override;
 
 	// todo: macro boilerplate
 	virtual void SetNumRows(uint32 NumRows) override { RowValues.SetNum(NumRows); }
-	virtual FName GetDisplayName() override { return DisplayName; }
-	virtual void SetDisplayName(FName Name) override { DisplayName = Name; }
+	virtual FName GetDisplayName() override { return InputPropertyName; }
+	virtual void SetDisplayName(FName Name) override { InputPropertyName = Name; }
 	virtual void DeleteRows(const TArray<uint32> & RowIndices )
 	{
 		for(uint32 Index : RowIndices)
@@ -76,22 +74,19 @@ class CHOOSER_API UChooserColumnFloatRange : public UObject, public IChooserColu
 	GENERATED_BODY()
 	public:
 	UChooserColumnFloatRange() { }
-	
-	UPROPERTY(EditAnywhere, Category = "Editor")
-	FName DisplayName = "Float Range Column";
-	
-	UPROPERTY(EditAnywhere, Meta=(DataInterfaceType="float"), Category = "Input")
-	TScriptInterface<IDataInterface> Value;
+		
+	UPROPERTY(EditAnywhere, Category = "Input")
+		FName InputPropertyName;
 
 	UPROPERTY(EditAnywhere, Category=Runtime)
 	// array of results (cells for this column for each row in the table)
 	// should match the length of the Results array 
 	TArray<FChooserFloatRangeRowData> RowValues;
 	
-	virtual void Filter(const UE::DataInterface::FContext& Context, const TArray<uint32>& IndexListIn, TArray<uint32>& IndexListOut) override;
+	virtual void Filter(const UObject* ContextObject, const TArray<uint32>& IndexListIn, TArray<uint32>& IndexListOut) override;
 	virtual void SetNumRows(uint32 NumRows) { RowValues.SetNum(NumRows); }
-	virtual FName GetDisplayName() override { return DisplayName; }
-	virtual void SetDisplayName(FName Name) override { DisplayName = Name; }
+	virtual FName GetDisplayName() override { return InputPropertyName; }
+	virtual void SetDisplayName(FName Name) override { InputPropertyName = Name; }
 	virtual void DeleteRows(const TArray<uint32> & RowIndices )
 	{
 		for(uint32 Index : RowIndices)
@@ -108,34 +103,20 @@ class UChooserTable : public UObject
 public:
 	UChooserTable() {}
 	
-	UPROPERTY(EditAnywhere, Meta=(EditInline="true"), Category = "Editor")
-	FName ResultType = "Object"; // todo: drop-down, populated by all possible return type names from all implemented data interfaces
-	// todo: struct type for Object return types
-
 	UPROPERTY(EditAnywhere, Meta=(EditInline="true"), Category = "Runtime")
 	TArray<TScriptInterface<IChooserColumn>> Columns;
 
 	// array of results (rows of table)
-	// todo: DataInterfaceType shouldn't be hard coded (Should be based on Result Type above) needed for Details customization to work
-	UPROPERTY(EditAnywhere, Meta=(DataInterfaceType="Asset"), Category = "Runtime")
-	TArray<TScriptInterface<IDataInterface>> Results;
+	UPROPERTY(EditAnywhere, Meta = (EditInline = "true"), Category = "Runtime")
+	TArray<TScriptInterface<IObjectChooser>> Results;
 };
 
 UCLASS()
-class CHOOSER_API UDataInterface_EvaluateChooser : public UObject, public IDataInterface
+class CHOOSER_API UObjectChooser_EvaluateChooser : public UObject, public IObjectChooser
 {
 	GENERATED_BODY()
 
-	// virtual bool GetObject(const UE::DataInterface::FContext& DataContext) const final override;
-
-	/** Get the return type name, used for dynamic type checking */
-    virtual FName GetReturnTypeNameImpl() const { return Chooser ? Chooser->ResultType : NAME_None; } 
-
-    /** Get the return type struct, used for dynamic type checking. If the result is not a struct, this should return nullptr */
-    virtual const UScriptStruct* GetReturnTypeStructImpl() const { return nullptr; } // todo, get type struct from chooser once it's added
-    
-    /** Get the value for this interface. @return true if successful, false if unsuccessful. */
-    virtual bool GetDataImpl(const UE::DataInterface::FContext& Context) const;
+	virtual UObject* ChooseObject(const UObject* ContextObject) const final override;
 
 	public:
 	
diff --git a/Engine/Plugins/Experimental/Chooser/Source/Chooser/Internal/DataInterface_Object.h b/Engine/Plugins/Experimental/Chooser/Source/Chooser/Internal/DataInterface_Object.h
deleted file mode 100644
diff --git a/Engine/Plugins/Experimental/Chooser/Source/Chooser/Internal/ObjectChooser_Asset.h b/Engine/Plugins/Experimental/Chooser/Source/Chooser/Internal/ObjectChooser_Asset.h
new file mode 100644
--- /dev/null
+++ b/Engine/Plugins/Experimental/Chooser/Source/Chooser/Internal/ObjectChooser_Asset.h
@@ -0,0 +1,20 @@
+// Copyright Epic Games, Inc. All Rights Reserved.
+
+#pragma once
+
+#include "CoreMinimal.h"
+#include "IObjectChooser.h"
+#include "ObjectChooser_Asset.generated.h"
+
+UCLASS()
+class CHOOSER_API UObjectChooser_Asset : public UObject, public IObjectChooser
+{
+	GENERATED_BODY()
+	
+	// IObjectChooser interface
+	virtual UObject* ChooseObject(const UObject* ContextObject) const final override;
+	
+public: 
+	UPROPERTY(EditAnywhere, Category = "Parameters")
+	TObjectPtr<UObject> Asset;
+};
diff --git a/Engine/Plugins/Experimental/Chooser/Source/Chooser/Private/Chooser.cpp b/Engine/Plugins/Experimental/Chooser/Source/Chooser/Private/Chooser.cpp
--- a/Engine/Plugins/Experimental/Chooser/Source/Chooser/Private/Chooser.cpp
+++ b/Engine/Plugins/Experimental/Chooser/Source/Chooser/Private/Chooser.cpp
@@ -1,7 +1,6 @@
 // Copyright Epic Games, Inc. All Rights Reserved.
 #include "Chooser.h"
-#include "DataInterfaceTypes.h"
-#include "DataInterface.h"
+#include "ChooserFunctionLibrary.h"
 
 UChooserTable::UChooserTable(const FObjectInitializer& Initializer)
 	:Super(Initializer)
@@ -9,49 +8,60 @@ UChooserTable::UChooserTable(const FObjectInitializer& Initializer)
 
 }
 
-void UChooserColumnBool::Filter(const UE::DataInterface::FContext& Context, const TArray<uint32>& IndexListIn, TArray<uint32>& IndexListOut)
+void UChooserColumnBool::Filter(const UObject* ContextObject, const TArray<uint32>& IndexListIn, TArray<uint32>& IndexListOut)
 {
 	bool Result = false;
-	if(UE::DataInterface::GetDataSafe(Value, Context, Result))
+
+	if (ContextObject)
 	{
-		for(uint32 Index : IndexListIn)
+		if (const FBoolProperty* Property = FindFProperty<FBoolProperty>(ContextObject->GetClass(), InputPropertyName))
 		{
-			if (RowValues.Num() > (int)Index)
+			Result = *Property->ContainerPtrToValuePtr<bool>(ContextObject);
+			for (uint32 Index : IndexListIn)
 			{
-				if (Result == RowValues[Index])
+				if (RowValues.Num() > (int)Index)
 				{
-					IndexListOut.Push(Index);
+					if (Result == RowValues[Index])
+					{
+						IndexListOut.Push(Index);
+					}
 				}
 			}
 		}
 	}
 }
 
-void UChooserColumnFloatRange::Filter(const UE::DataInterface::FContext& Context, const TArray<uint32>& IndexListIn, TArray<uint32>& IndexListOut)
+void UChooserColumnFloatRange::Filter(const UObject* ContextObject, const TArray<uint32>& IndexListIn, TArray<uint32>& IndexListOut)
 {
 	float Result = 0.0f;
-	
-	if(UE::DataInterface::GetDataSafe(Value, Context, Result))
+
+	Result = 0.0f;
+	if (ContextObject)
 	{
-		for(uint32 Index : IndexListIn)
+		if (const FDoubleProperty* Property = FindFProperty<FDoubleProperty>(ContextObject->GetClass(), InputPropertyName))
 		{
-			if (RowValues.Num() > (int)Index)
+			Result = *Property->ContainerPtrToValuePtr<double>(ContextObject);
+	
+			for(uint32 Index : IndexListIn)
 			{
-				const FChooserFloatRangeRowData& RowValue = RowValues[Index];
-				if (Result >= RowValue.Min && Result <= RowValue.Max)
+				if (RowValues.Num() > (int)Index)
 				{
-					IndexListOut.Push(Index);
+					const FChooserFloatRangeRowData& RowValue = RowValues[Index];
+					if (Result >= RowValue.Min && Result <= RowValue.Max)
+					{
+						IndexListOut.Push(Index);
+					}
 				}
 			}
 		}
 	}
 }
 
-bool UDataInterface_EvaluateChooser::GetDataImpl(const UE::DataInterface::FContext& DataContext) const
+static UObject* StaticEvaluateChooser(const UObject* ContextObject, const UChooserTable* Chooser)
 {
 	if (Chooser == nullptr)
 	{
-		return false;
+		return nullptr;
 	}
 
 	TArray<uint32> Indices1;
@@ -72,7 +82,7 @@ bool UDataInterface_EvaluateChooser::GetDataImpl(const UE::DataInterface::FConte
 		{
 			Swap(IndicesIn, IndicesOut);
 			IndicesOut->SetNum(0, false);
-			Column->Filter(DataContext, *IndicesIn, *IndicesOut);
+			Column->Filter(ContextObject, *IndicesIn, *IndicesOut);
 		}
 	}
 
@@ -81,15 +91,29 @@ bool UDataInterface_EvaluateChooser::GetDataImpl(const UE::DataInterface::FConte
 	{
 		if (Chooser->Results.Num() > (int32)SelectedIndex)
 		{
-			const TScriptInterface<IDataInterface>& SelectedResult = Chooser->Results[SelectedIndex];
+			const TScriptInterface<IObjectChooser>& SelectedResult = Chooser->Results[SelectedIndex];
 
-			// todo: GetResultRaw returns const, GetData requires non-const
-			if(UE::DataInterface::GetDataSafe(SelectedResult, DataContext, const_cast<UE::DataInterface::FParam&>(DataContext.GetResultRaw())))
+			if(UObject* Result = SelectedResult->ChooseObject(ContextObject))
 			{
-				return true;
+				return Result;
 			}
 		}
 	}
 
-	return false;
+	return nullptr;
+}
+
+UObject* UObjectChooser_EvaluateChooser::ChooseObject(const UObject* ContextObject) const
+{
+	return StaticEvaluateChooser(ContextObject, Chooser);
+}
+
+UChooserFunctionLibrary::UChooserFunctionLibrary(const FObjectInitializer& ObjectInitializer)
+	: Super(ObjectInitializer)
+{
+}
+
+UObject* UChooserFunctionLibrary::EvaluateChooser(const UObject* ContextObject, const UChooserTable* Chooser)
+{
+	return StaticEvaluateChooser(ContextObject, Chooser);
 }
diff --git a/Engine/Plugins/Experimental/Chooser/Source/Chooser/Private/DataInterface_Object.cpp b/Engine/Plugins/Experimental/Chooser/Source/Chooser/Private/DataInterface_Object.cpp
deleted file mode 100644
diff --git a/Engine/Plugins/Experimental/Chooser/Source/Chooser/Private/ObjectChooser_Asset.cpp b/Engine/Plugins/Experimental/Chooser/Source/Chooser/Private/ObjectChooser_Asset.cpp
new file mode 100644
--- /dev/null
+++ b/Engine/Plugins/Experimental/Chooser/Source/Chooser/Private/ObjectChooser_Asset.cpp
@@ -0,0 +1,8 @@
+// Copyright Epic Games, Inc. All Rights Reserved.
+
+#include "ObjectChooser_Asset.h"
+
+UObject* UObjectChooser_Asset::ChooseObject(const UObject* ContextObject) const
+{
+	return Asset;
+}
\ No newline at end of file
diff --git a/Engine/Plugins/Experimental/Chooser/Source/Chooser/Public/ChooserFunctionLibrary.h b/Engine/Plugins/Experimental/Chooser/Source/Chooser/Public/ChooserFunctionLibrary.h
new file mode 100644
--- /dev/null
+++ b/Engine/Plugins/Experimental/Chooser/Source/Chooser/Public/ChooserFunctionLibrary.h
@@ -0,0 +1,28 @@
+// Copyright Epic Games, Inc. All Rights Reserved.
+
+#pragma once
+
+#include "CoreMinimal.h"
+#include "Kismet/BlueprintFunctionLibrary.h"
+#include "Chooser.h"
+#include "ChooserFunctionLibrary.generated.h"
+
+/**
+ * Morpheus Extensions Function Library
+ */
+UCLASS()
+class CHOOSER_API UChooserFunctionLibrary : public UBlueprintFunctionLibrary
+{
+	GENERATED_UCLASS_BODY()
+
+public:
+	/**
+	* Evaluate a chooser table and return the selected UObject, or null
+	*
+	* @param ContextObject			(in) An Object from which the parameters to the Chooser Table will be read
+	* @param ChooserTable			(in) The ChooserTable asset
+	*/
+	UFUNCTION(BlueprintPure, meta = (BlueprintThreadSafe), Category = "Animation")
+	static UObject* EvaluateChooser(const UObject* ContextObject, const UChooserTable* ChooserTable);
+
+};
diff --git a/Engine/Plugins/Experimental/Chooser/Source/Chooser/Public/IObjectChooser.h b/Engine/Plugins/Experimental/Chooser/Source/Chooser/Public/IObjectChooser.h
new file mode 100644
--- /dev/null
+++ b/Engine/Plugins/Experimental/Chooser/Source/Chooser/Public/IObjectChooser.h
@@ -0,0 +1,23 @@
+// Copyright Epic Games, Inc. All Rights Reserved.
+
+#pragma once
+
+#include "CoreMinimal.h"
+#include "UObject/Object.h"
+#include "UObject/Interface.h"
+#include "IObjectChooser.generated.h"
+
+UINTERFACE(NotBlueprintType, meta = (CannotImplementInterfaceInBlueprint))
+class CHOOSER_API UObjectChooser : public UInterface
+{
+	GENERATED_BODY()
+};
+
+class CHOOSER_API IObjectChooser
+{
+	GENERATED_BODY()
+
+public:
+
+	virtual UObject* ChooseObject(const UObject* ContextObject) const { return nullptr; };
+};
diff --git a/Engine/Plugins/Experimental/Chooser/Source/ChooserEditor/ChooserEditor.Build.cs b/Engine/Plugins/Experimental/Chooser/Source/ChooserEditor/ChooserEditor.Build.cs
--- a/Engine/Plugins/Experimental/Chooser/Source/ChooserEditor/ChooserEditor.Build.cs
+++ b/Engine/Plugins/Experimental/Chooser/Source/ChooserEditor/ChooserEditor.Build.cs
@@ -32,20 +32,12 @@ namespace UnrealBuildTool.Rules
 					"CoreUObject",
 					"Engine",
 					"AssetTools",
-					"DataInterface",
-					"DataInterfaceGraphEditor", // for DataInterfaceClassFilter, and Widget Factory system (could be moved to a more generic DataInterfaceEditor module
 					"Chooser",
 					"UnrealEd",
 					"EditorWidgets",
 					"SlateCore",
 					"Slate",
 					"PropertyEditor",
-					"RigVM",
-					"RigVMDeveloper",
-					"ControlRig",
-					"ControlRigDeveloper",
-					"ControlRigEditor",
-					"GraphEditor",
 					"InputCore",
 					"EditorStyle",
 					// ... add private dependencies that you statically link with here ...
diff --git a/Engine/Plugins/Experimental/Chooser/Source/ChooserEditor/Private/ChooserTableEditor.cpp b/Engine/Plugins/Experimental/Chooser/Source/ChooserEditor/Private/ChooserTableEditor.cpp
--- a/Engine/Plugins/Experimental/Chooser/Source/ChooserEditor/Private/ChooserTableEditor.cpp
+++ b/Engine/Plugins/Experimental/Chooser/Source/ChooserEditor/Private/ChooserTableEditor.cpp
@@ -25,8 +25,9 @@
 #include "SAssetDropTarget.h"
 #include "Widgets/Layout/SWidgetSwitcher.h"
 #include "DragAndDrop/DecoratedDragDropOp.h"
-#include "DataInterfaceClassFilter.h"
-#include "DataInterfaceWidgetFactories.h"
+#include "ObjectChooser_Asset.h"
+#include "ObjectChooserClassFilter.h"
+#include "ObjectChooserWidgetFactories.h"
 
 
 #define LOCTEXT_NAMESPACE "ChooserEditor"
@@ -39,12 +40,12 @@ const FName FChooserTableEditor::PropertiesTabId( TEXT( "ChooserEditor_Propertie
 const FName FChooserTableEditor::TableTabId( TEXT( "ChooserEditor_Table" ) );
 
 
-void ConvertToText(UObject* Object, FText& OutText)
+void ConvertObjectToText(UObject* Object, FText& OutText)
 {
 	UClass* Class = Object->GetClass();
 	while (Class)
 	{
-		if (auto TextConverter = DataInterfaceGraphEditor::FDataInterfaceWidgetFactories::DataInterfaceTextConverter.Find(Class))
+		if (auto TextConverter = FObjectChooserWidgetFactories::ChooserTextConverter.Find(Class))
 		{
 			(*TextConverter)(Object, OutText);
 			break;
@@ -309,7 +310,7 @@ public:
 		Operation->ChooserEditor = InEditor;
 		Operation->RowIndex = InRowIndex;
 		Operation->DefaultHoverText = LOCTEXT("Chooser Row", "Chooser Row");
-		ConvertToText(InEditor->GetChooser()->Results[InRowIndex].GetObject(), Operation->DefaultHoverText);
+		ConvertObjectToText(InEditor->GetChooser()->Results[InRowIndex].GetObject(), Operation->DefaultHoverText);
 		Operation->CurrentHoverText = Operation->DefaultHoverText;
 			
 		Operation->Construct();
@@ -397,16 +398,15 @@ public:
 				
 				return 		SNew(SChooserRowHandle).ChooserEditor(Editor).RowIndex(RowIndex->RowIndex);
 			}
-			else if (ColumnName == Result)
+			else if (ColumnName == Result) 
 			{
-		
 				UObject* RowValue = Chooser->Results[RowIndex->RowIndex].GetObject();
-				TSharedPtr<SWidget> ResultWidget = DataInterfaceGraphEditor::FDataInterfaceWidgetFactories::CreateDataInterfaceWidget(Chooser->ResultType, RowValue,
+				TSharedPtr<SWidget> ResultWidget = FObjectChooserWidgetFactories::CreateWidget("", RowValue,
 				FOnClassPicked::CreateLambda([this, RowIndex=RowIndex->RowIndex](UClass* ChosenClass)
 				{
 					UObject* RowValue = NewObject<UObject>(Chooser, ChosenClass);
 					Chooser->Results[RowIndex] = RowValue;
-					DataInterfaceGraphEditor::FDataInterfaceWidgetFactories::CreateDataInterfaceWidget(Chooser->ResultType, RowValue, FOnClassPicked(), &CacheBorder);
+					FObjectChooserWidgetFactories::CreateWidget("", RowValue, FOnClassPicked(), &CacheBorder);
 				}),
 				&CacheBorder
 				);
@@ -448,7 +448,7 @@ public:
 				CreateRowComboButton->SetOnGetMenuContent(FOnGetContent::CreateLambda([this, CreateRowComboButton]()
 					{
 						FClassViewerInitializationOptions Options;
-						Options.ClassFilters.Add(MakeShared<UE::DataInterfaceGraphEditor::FDataInterfaceClassFilter>(Chooser->ResultType, nullptr));
+						Options.ClassFilters.Add(MakeShared<FObjectChooserClassFilter>("", nullptr));
 						
 						// Add class filter for columns here
 						TSharedRef<SWidget> Widget = FModuleManager::LoadModuleChecked<FClassViewerModule>("ClassViewer").CreateClassViewer(Options,
@@ -888,7 +888,7 @@ TMap<const UClass*, TFunction<TSharedRef<SWidget> (UObject* Column, int Row)>> F
 
 TSharedRef<SWidget> CreateAssetWidget(UObject* Object)
 {
-	UDataInterface_Object_Asset* DIAsset = Cast<UDataInterface_Object_Asset>(Object);
+	UObjectChooser_Asset* DIAsset = Cast<UObjectChooser_Asset>(Object);
 
 	UObject* Asset = DIAsset->Asset;
 	
@@ -919,7 +919,7 @@ TSharedRef<SWidget> CreateAssetWidget(UObject* Object)
 
 TSharedRef<SWidget> CreateEvaluateChooserWidget(UObject* Object)
 {
-	UDataInterface_EvaluateChooser* EvaluateChooser = Cast<UDataInterface_EvaluateChooser>(Object);
+	UObjectChooser_EvaluateChooser* EvaluateChooser = Cast<UObjectChooser_EvaluateChooser>(Object);
 
 	return
 	SNew(SAssetDropTarget)
@@ -946,14 +946,9 @@ TSharedRef<SWidget> CreateEvaluateChooserWidget(UObject* Object)
 	];
 }
 
-void ConvertToText_Base(const UObject* Object, FText& OutText)
-{
-	OutText = FText::FromString(Object->GetName());
-}
-
 void ConvertToText_Asset(const UObject* Object, FText& OutText)
 {
-	const UDataInterface_Object_Asset* AssetInterface = Cast<UDataInterface_Object_Asset>(Object);
+	const UObjectChooser_Asset* AssetInterface = Cast<UObjectChooser_Asset>(Object);
 	if (AssetInterface->Asset == nullptr)
 	{
 		OutText = LOCTEXT("NoChooser", "[No Chooser]");
@@ -966,7 +961,7 @@ void ConvertToText_Asset(const UObject* Object, FText& OutText)
 
 void ConvertToText_EvaluateChooser(const UObject* Object, FText& OutText)
 {
-	const UDataInterface_EvaluateChooser* EvaluateChooser = Cast<UDataInterface_EvaluateChooser>(Object);
+	const UObjectChooser_EvaluateChooser* EvaluateChooser = Cast<UObjectChooser_EvaluateChooser>(Object);
 	if (EvaluateChooser->Chooser == nullptr)
 	{
 		OutText = LOCTEXT("NoChooser", "[No Chooser]");
@@ -980,7 +975,7 @@ void ConvertToText_EvaluateChooser(const UObject* Object, FText& OutText)
 TSharedRef<SWidget> CreateObjectWidget(UObject* Object)
 {
 	FText ObjectText = FText::FromString(Object->GetName());
-	ConvertToText(Object, ObjectText);
+	ConvertObjectToText(Object, ObjectText);
 	
 	return SNew (STextBlock).Text(ObjectText); // could make this use Text_Lambda, to have it update correctly?
 }
@@ -1122,14 +1117,12 @@ void FChooserRowDetails::CustomizeDetails(IDetailLayoutBuilder& DetailBuilder)
 
 void FChooserTableEditor::RegisterWidgets()
 {
-	using namespace DataInterfaceGraphEditor;
-	FDataInterfaceWidgetFactories::DataInterfaceTextConverter.Add(UObject::StaticClass(), ConvertToText_Base);
-	FDataInterfaceWidgetFactories::DataInterfaceTextConverter.Add(UDataInterface_Object_Asset::StaticClass(), ConvertToText_Asset);
-	FDataInterfaceWidgetFactories::DataInterfaceTextConverter.Add(UDataInterface_EvaluateChooser::StaticClass(), ConvertToText_EvaluateChooser);
+	FObjectChooserWidgetFactories::ChooserTextConverter.Add(UObjectChooser_Asset::StaticClass(), ConvertToText_Asset);
+	FObjectChooserWidgetFactories::ChooserTextConverter.Add(UObjectChooser_EvaluateChooser::StaticClass(), ConvertToText_EvaluateChooser);
 	
-	FDataInterfaceWidgetFactories::DataInterfaceWidgetCreators.Add(UObject::StaticClass(), CreateObjectWidget);
-	FDataInterfaceWidgetFactories::DataInterfaceWidgetCreators.Add(UDataInterface_Object_Asset::StaticClass(), CreateAssetWidget);
-	FDataInterfaceWidgetFactories::DataInterfaceWidgetCreators.Add(UDataInterface_EvaluateChooser::StaticClass(), CreateEvaluateChooserWidget);
+	FObjectChooserWidgetFactories::ChooserWidgetCreators.Add(UObject::StaticClass(), CreateObjectWidget);
+	FObjectChooserWidgetFactories::ChooserWidgetCreators.Add(UObjectChooser_Asset::StaticClass(), CreateAssetWidget);
+	FObjectChooserWidgetFactories::ChooserWidgetCreators.Add(UObjectChooser_EvaluateChooser::StaticClass(), CreateEvaluateChooserWidget);
 
 	ColumnWidgetCreators.Add(UChooserColumnBool::StaticClass(), CreateBoolColumnWidget);
 	ColumnWidgetCreators.Add(UChooserColumnFloatRange::StaticClass(), CreateFloatRangeColumnWidget);
diff --git a/Engine/Plugins/Experimental/Chooser/Source/ChooserEditor/Private/ObjectChooserWidgetFactories.cpp b/Engine/Plugins/Experimental/Chooser/Source/ChooserEditor/Private/ObjectChooserWidgetFactories.cpp
new file mode 100644
--- /dev/null
+++ b/Engine/Plugins/Experimental/Chooser/Source/ChooserEditor/Private/ObjectChooserWidgetFactories.cpp
@@ -0,0 +1,122 @@
+// Copyright Epic Games, Inc. All Rights Reserved.
+
+#include "ObjectChooserWidgetFactories.h"
+#include "Modules/ModuleManager.h"
+#include "Widgets/Input/SHyperlink.h"
+#include "Widgets/Input/SNumericEntryBox.h"
+#include "Widgets/Input/SComboButton.h"
+#include "SClassViewer.h"
+#include "DetailCategoryBuilder.h"
+#include "Widgets/Layout/SWidgetSwitcher.h"
+#include "ObjectChooserClassFilter.h"
+
+#define LOCTEXT_NAMESPACE "DataInterfaceEditor"
+
+namespace UE::ChooserEditor
+{
+
+void ConvertToText(UObject* Object, FText& OutText)
+{
+	UClass* Class = Object->GetClass();
+	while (Class)
+	{
+		if (FChooserTextConverter* TextConverter = FObjectChooserWidgetFactories::ChooserTextConverter.Find(Class))
+		{
+			(*TextConverter)(Object, OutText);
+			break;
+		}
+		Class = Class->GetSuperClass();
+	}
+}
+
+TSharedPtr<SWidget> FObjectChooserWidgetFactories::CreateWidget(FName DataInterfaceTypeName, UObject* Value, const FOnClassPicked& CreateClassCallback, TSharedPtr<SBorder>* InnerWidget)
+{
+	TSharedPtr<SWidget> LeftWidget;
+
+	if (Value)
+	{
+		UClass* Class = Value->GetClass();
+		while (Class && !LeftWidget.IsValid())
+		{
+			if (FChooserWidgetCreator* Creator = FObjectChooserWidgetFactories::ChooserWidgetCreators.Find(Class))
+			{
+				LeftWidget = (*Creator)(Value);
+				break;
+			}
+			Class = Class->GetSuperClass();
+		}
+	}
+	
+	if (!LeftWidget.IsValid())
+	{
+		LeftWidget = SNew(STextBlock).Text(LOCTEXT("SelectDataType", "Select Data Type..."));
+	}
+
+
+	// button for replacing data with a different Data Interface class
+	TSharedPtr<SComboButton> Button = SNew(SComboButton)
+			.ComboButtonStyle(FAppStyle::Get(), "SimpleComboButton");
+	
+	Button->SetOnGetMenuContent(FOnGetContent::CreateLambda([DataInterfaceTypeName, Button, CreateClassCallback]()
+	{
+		FClassViewerInitializationOptions Options;
+		Options.ClassFilters.Add(MakeShared<FObjectChooserClassFilter>(DataInterfaceTypeName, nullptr));
+		
+		// Add class filter for columns here
+		TSharedRef<SWidget>  ClassMenu = FModuleManager::LoadModuleChecked<FClassViewerModule>("ClassViewer").
+		CreateClassViewer(Options, FOnClassPicked::CreateLambda([Button, CreateClassCallback](UClass* Class)
+		{
+			Button->SetIsOpen(false);
+			CreateClassCallback.Execute(Class);
+		}));
+		return ClassMenu;
+	}));
+
+	TSharedPtr <SBorder> Border;
+	if (InnerWidget && InnerWidget->IsValid())
+	{
+		Border = *InnerWidget;
+	}
+	else
+	{
+		Border = SNew(SBorder);
+	}
+	
+	if (InnerWidget)
+	{
+		*InnerWidget = Border;
+	}
+	
+	Border->SetContent(LeftWidget.ToSharedRef());
+
+	TSharedPtr<SWidget> Widget = SNew(SHorizontalBox)
+		+SHorizontalBox::Slot().FillWidth(75)
+		[
+			Border.ToSharedRef()
+		]
+		+SHorizontalBox::Slot().FillWidth(25)
+		[
+			Button.ToSharedRef()
+		]
+	;
+
+
+	return Widget;
+}
+
+TMap<const UClass*, FChooserTextConverter> FObjectChooserWidgetFactories::ChooserTextConverter;
+TMap<const UClass*, FChooserWidgetCreator> FObjectChooserWidgetFactories::ChooserWidgetCreators;
+
+void ConvertToText_Base(const UObject* Object, FText& OutText)
+{
+	OutText = FText::FromString(Object->GetName());
+}
+
+void FObjectChooserWidgetFactories::RegisterWidgets()
+{
+	ChooserTextConverter.Add(UObject::StaticClass(), ConvertToText_Base);
+}
+	
+}
+
+#undef LOCTEXT_NAMESPACE
diff --git a/Engine/Plugins/Experimental/Chooser/Source/ChooserEditor/Public/ObjectChooserClassFilter.h b/Engine/Plugins/Experimental/Chooser/Source/ChooserEditor/Public/ObjectChooserClassFilter.h
new file mode 100644
--- /dev/null
+++ b/Engine/Plugins/Experimental/Chooser/Source/ChooserEditor/Public/ObjectChooserClassFilter.h
@@ -0,0 +1,35 @@
+// Copyright Epic Games, Inc. All Rights Reserved.
+#pragma once
+
+#include "ClassViewerFilter.h"
+#include "IObjectChooser.h"
+
+namespace UE::ChooserEditor
+{
+	// Filter class for things which implement IChooserColumn
+	class FObjectChooserClassFilter : public IClassViewerFilter
+	{
+	public:
+		FObjectChooserClassFilter(FName InTypeName, const UScriptStruct* InStructType) : TypeName(InTypeName), StructType(InStructType) { };
+		virtual ~FObjectChooserClassFilter() override {};
+	
+		virtual bool IsClassAllowed(const FClassViewerInitializationOptions& InInitOptions, const UClass* InClass, TSharedRef< class FClassViewerFilterFuncs > InFilterFuncs ) override
+		{
+			if (InClass->ImplementsInterface(UObjectChooser::StaticClass()))
+			{
+				// for now, just return true if the class implements IObjectChooser
+				return true;
+			}
+
+			return false;
+		}
+
+		virtual bool IsUnloadedClassAllowed(const FClassViewerInitializationOptions& InInitOptions, const TSharedRef< const class IUnloadedBlueprintData > InUnloadedClassData, TSharedRef< class FClassViewerFilterFuncs > InFilterFuncs) override
+		{
+			return false;
+		}
+	private:
+		FName TypeName;
+		const UScriptStruct* StructType;
+	};
+}
\ No newline at end of file
diff --git a/Engine/Plugins/Experimental/Chooser/Source/ChooserEditor/Public/ObjectChooserWidgetFactories.h b/Engine/Plugins/Experimental/Chooser/Source/ChooserEditor/Public/ObjectChooserWidgetFactories.h
new file mode 100644
--- /dev/null
+++ b/Engine/Plugins/Experimental/Chooser/Source/ChooserEditor/Public/ObjectChooserWidgetFactories.h
@@ -0,0 +1,25 @@
+// Copyright Epic Games, Inc. All Rights Reserved.
+
+#pragma once
+
+#include "CoreMinimal.h"
+#include "Widgets/SWidget.h"
+#include "Widgets/Layout/SBorder.h"
+#include "SClassViewer.h"
+
+namespace UE::ChooserEditor
+{
+	typedef TFunction<void(const UObject* Object, FText& OutText)> FChooserTextConverter;
+	typedef TFunction<TSharedRef<SWidget>(UObject* Object)> FChooserWidgetCreator;
+
+	class CHOOSEREDITOR_API FObjectChooserWidgetFactories
+	{
+	public:
+		static TMap<const UClass*, FChooserTextConverter> ChooserTextConverter;
+		static TMap<const UClass*, FChooserWidgetCreator> ChooserWidgetCreators;
+
+		static TSharedPtr<SWidget> CreateWidget(FName DataInterfaceTypeName, UObject* Value, const FOnClassPicked& CreateClassCallback, TSharedPtr<SBorder>* InnerWidget = nullptr);
+		
+		static void RegisterWidgets();
+	};
+}
