Updates to Chooser system.
 - Add interface layer to allow for custom column parameter types
 - Add property access ui for context object properties
 - Use proper, filtered asset selection UI for result assets and nested chooser tables.

diff --git a/Engine/Plugins/Experimental/Chooser/Source/Chooser/Internal/Chooser.h b/Engine/Plugins/Experimental/Chooser/Source/Chooser/Internal/Chooser.h
--- a/Engine/Plugins/Experimental/Chooser/Source/Chooser/Internal/Chooser.h
+++ b/Engine/Plugins/Experimental/Chooser/Source/Chooser/Internal/Chooser.h
@@ -8,6 +8,69 @@
 
 #include "Chooser.generated.h"
 
+UINTERFACE(NotBlueprintType, meta = (CannotImplementInterfaceInBlueprint))
+class CHOOSER_API UChooserParameterBool : public UInterface
+{
+	GENERATED_BODY()
+};
+
+class CHOOSER_API IChooserParameterBool
+{
+	GENERATED_BODY()
+
+public:
+	virtual bool GetValue(const UObject* ContextObject, bool& OutResult) { return false; }
+};
+
+UCLASS()
+class CHOOSER_API UChooserParameterBool_ContextProperty :  public UObject, public IChooserParameterBool
+{
+	GENERATED_BODY()
+public:
+	UPROPERTY()
+ 	FName PropertyName;
+	
+	virtual bool GetValue(const UObject* ContextObject, bool& OutResult) override;
+
+	static const FString& CPPTypeName()
+	{
+		static FString TypeName = "bool";
+		return TypeName;
+	}
+};
+
+
+UINTERFACE(NotBlueprintType, meta = (CannotImplementInterfaceInBlueprint))
+class CHOOSER_API UChooserParameterFloat : public UInterface
+{
+	GENERATED_BODY()
+};
+
+class CHOOSER_API IChooserParameterFloat
+{
+	GENERATED_BODY()
+
+public:
+	virtual bool GetValue(const UObject* ContextObject, float& OutResult) { return false; }
+};
+
+UCLASS()
+class CHOOSER_API UChooserParameterFloat_ContextProperty :  public UObject, public IChooserParameterFloat
+{
+	GENERATED_BODY()
+public:
+	UPROPERTY()
+	FName PropertyName;
+	
+	virtual bool GetValue(const UObject* ContextObject, float& OutResult) override;
+	
+	static const FString& CPPTypeName()
+	{
+		static FString TypeName = "double";
+		return TypeName;
+	}
+};
+
 UINTERFACE(NotBlueprintType, meta = (CannotImplementInterfaceInBlueprint))
 class CHOOSER_API UChooserColumn : public UInterface
 {
@@ -20,10 +83,11 @@ class CHOOSER_API IChooserColumn
 
 public:
 	virtual void Filter(const UObject* ContextObject, const TArray<uint32>& IndexListIn, TArray<uint32>& IndexListOut) {};
-	virtual FName GetDisplayName() { return "ChooserColumn"; }
-	virtual void SetDisplayName(FName Name) { }
 	virtual void SetNumRows(uint32 NumRows) {}
 	virtual void DeleteRows(const TArray<uint32> & RowIndices) {}
+	virtual UClass* GetInputValueInterface() {return nullptr;};
+	virtual UObject* GetInputValue() {return nullptr;}
+	virtual void SetInputValue(UObject* InputValue) {}
 };
 
 UCLASS()
@@ -31,10 +95,11 @@ class CHOOSER_API UChooserColumnBool : public UObject, public IChooserColumn
 {
 	GENERATED_BODY()
 	public:
-	UChooserColumnBool() { }
+	UChooserColumnBool() {};
+	UChooserColumnBool(const FObjectInitializer& ObjectInitializer);
 	
-	UPROPERTY(EditAnywhere, Category = "Input")
-	FName InputPropertyName;
+	UPROPERTY(EditAnywhere, Meta = (EditInlineInterface = "true"), Category = "Input")
+	TScriptInterface<IChooserParameterBool> InputValue;
 
 	UPROPERTY(EditAnywhere, Category=Runtime)
 	// array of results (cells for this column for each row in the table)
@@ -45,8 +110,6 @@ class CHOOSER_API UChooserColumnBool : public UObject, public IChooserColumn
 
 	// todo: macro boilerplate
 	virtual void SetNumRows(uint32 NumRows) override { RowValues.SetNum(NumRows); }
-	virtual FName GetDisplayName() override { return InputPropertyName; }
-	virtual void SetDisplayName(FName Name) override { InputPropertyName = Name; }
 	virtual void DeleteRows(const TArray<uint32> & RowIndices )
 	{
 		for(uint32 Index : RowIndices)
@@ -54,6 +117,9 @@ class CHOOSER_API UChooserColumnBool : public UObject, public IChooserColumn
 			RowValues.RemoveAt(Index);
 		}
 	}
+	virtual UClass* GetInputValueInterface() override { return UChooserParameterBool::StaticClass(); };
+	virtual UObject* GetInputValue() override { return InputValue.GetObject(); };
+	virtual void SetInputValue(UObject* Value) override { InputValue = Value; };
 };
 
 USTRUCT()
@@ -73,11 +139,12 @@ class CHOOSER_API UChooserColumnFloatRange : public UObject, public IChooserColu
 {
 	GENERATED_BODY()
 	public:
-	UChooserColumnFloatRange() { }
+	UChooserColumnFloatRange() {}
+	UChooserColumnFloatRange(const FObjectInitializer& ObjectInitializer);
 		
-	UPROPERTY(EditAnywhere, Category = "Input")
-		FName InputPropertyName;
-
+	UPROPERTY(EditAnywhere, Meta = (EditInlineInterface = "true"), Category = "Input")
+	TScriptInterface<IChooserParameterFloat> InputValue;
+	
 	UPROPERTY(EditAnywhere, Category=Runtime)
 	// array of results (cells for this column for each row in the table)
 	// should match the length of the Results array 
@@ -85,8 +152,6 @@ class CHOOSER_API UChooserColumnFloatRange : public UObject, public IChooserColu
 	
 	virtual void Filter(const UObject* ContextObject, const TArray<uint32>& IndexListIn, TArray<uint32>& IndexListOut) override;
 	virtual void SetNumRows(uint32 NumRows) { RowValues.SetNum(NumRows); }
-	virtual FName GetDisplayName() override { return InputPropertyName; }
-	virtual void SetDisplayName(FName Name) override { InputPropertyName = Name; }
 	virtual void DeleteRows(const TArray<uint32> & RowIndices )
 	{
 		for(uint32 Index : RowIndices)
@@ -94,6 +159,9 @@ class CHOOSER_API UChooserColumnFloatRange : public UObject, public IChooserColu
 			RowValues.RemoveAt(Index);
 		}
 	}
+	virtual UClass* GetInputValueInterface() { return UChooserParameterFloat::StaticClass(); };
+	virtual UObject* GetInputValue() override { return InputValue.GetObject(); };
+	virtual void SetInputValue(UObject* Value) override { InputValue = Value; };
 };
 
 UCLASS(MinimalAPI)
@@ -102,13 +170,20 @@ class UChooserTable : public UObject
 	GENERATED_UCLASS_BODY()
 public:
 	UChooserTable() {}
+
+	// each possible result
+	UPROPERTY(EditAnywhere, Meta = (EditInlineInterface = "true"), Category = "Hidden")
+	TArray<TScriptInterface<IObjectChooser>> Results;
 	
-	UPROPERTY(EditAnywhere, Meta=(EditInline="true"), Category = "Runtime")
+	// columns which filter results
+	UPROPERTY(EditAnywhere, Meta = (EditInlineInterface = "true"), Category="Hidden")
 	TArray<TScriptInterface<IChooserColumn>> Columns;
 
-	// array of results (rows of table)
-	UPROPERTY(EditAnywhere, Meta = (EditInline = "true"), Category = "Runtime")
-	TArray<TScriptInterface<IObjectChooser>> Results;
+	UPROPERTY(EditAnywhere, Category="Input")
+	TObjectPtr<UClass> ContextObjectType;
+	
+	UPROPERTY(EditAnywhere, Category="Output")
+	TObjectPtr<UClass> OutputObjectType;
 };
 
 UCLASS()
@@ -122,4 +197,15 @@ class CHOOSER_API UObjectChooser_EvaluateChooser : public UObject, public IObjec
 	
 	UPROPERTY(EditAnywhere, Category="Parameters")
 	TObjectPtr<UChooserTable> Chooser;
+};
+
+
+UCLASS()
+class CHOOSER_API UChooserColumnMenuContext : public UObject
+{
+	GENERATED_BODY()
+public:
+	class FAssetEditorToolkit* Editor;
+	TWeakObjectPtr<UChooserTable> Chooser;
+	int ColumnIndex;
 };
\ No newline at end of file
diff --git a/Engine/Plugins/Experimental/Chooser/Source/Chooser/Private/Chooser.cpp b/Engine/Plugins/Experimental/Chooser/Source/Chooser/Private/Chooser.cpp
--- a/Engine/Plugins/Experimental/Chooser/Source/Chooser/Private/Chooser.cpp
+++ b/Engine/Plugins/Experimental/Chooser/Source/Chooser/Private/Chooser.cpp
@@ -2,59 +2,91 @@
 #include "Chooser.h"
 #include "ChooserFunctionLibrary.h"
 
+bool UChooserParameterBool_ContextProperty::GetValue(const UObject* ContextObject, bool& OutResult)
+{
+	if (const FBoolProperty* Property = FindFProperty<FBoolProperty>(ContextObject->GetClass(), PropertyName))
+	{
+		OutResult = *Property->ContainerPtrToValuePtr<bool>(ContextObject);
+	}
+	
+	return false;
+}
+
+bool UChooserParameterFloat_ContextProperty::GetValue(const UObject* ContextObject, float& OutResult)
+{
+	if (const FDoubleProperty* Property = FindFProperty<FDoubleProperty>(ContextObject->GetClass(), PropertyName))
+	{
+		OutResult = *Property->ContainerPtrToValuePtr<double>(ContextObject);
+	}
+	
+	return false;
+}
+
 UChooserTable::UChooserTable(const FObjectInitializer& Initializer)
 	:Super(Initializer)
 {
 
 }
 
-void UChooserColumnBool::Filter(const UObject* ContextObject, const TArray<uint32>& IndexListIn, TArray<uint32>& IndexListOut)
+UChooserColumnBool::UChooserColumnBool(const FObjectInitializer& ObjectInitializer)
 {
-	bool Result = false;
+	InputValue = ObjectInitializer.CreateDefaultSubobject<UChooserParameterBool_ContextProperty>(this, "InputValue");
+}
 
-	if (ContextObject)
+void UChooserColumnBool::Filter(const UObject* ContextObject, const TArray<uint32>& IndexListIn, TArray<uint32>& IndexListOut)
+{
+	if (ContextObject && InputValue)
 	{
-		if (const FBoolProperty* Property = FindFProperty<FBoolProperty>(ContextObject->GetClass(), InputPropertyName))
+		bool Result = false;
+		InputValue->GetValue(ContextObject,Result);
+		
+		for (uint32 Index : IndexListIn)
 		{
-			Result = *Property->ContainerPtrToValuePtr<bool>(ContextObject);
-			for (uint32 Index : IndexListIn)
+			if (RowValues.Num() > (int)Index)
 			{
-				if (RowValues.Num() > (int)Index)
+				if (Result == RowValues[Index])
 				{
-					if (Result == RowValues[Index])
-					{
-						IndexListOut.Push(Index);
-					}
+					IndexListOut.Push(Index);
 				}
 			}
 		}
 	}
+	else
+	{
+		// passthrough fallback (behaves better during live editing)
+		IndexListOut = IndexListIn;
+	}
 }
 
-void UChooserColumnFloatRange::Filter(const UObject* ContextObject, const TArray<uint32>& IndexListIn, TArray<uint32>& IndexListOut)
+UChooserColumnFloatRange::UChooserColumnFloatRange(const FObjectInitializer& ObjectInitializer)
 {
-	float Result = 0.0f;
+	InputValue = ObjectInitializer.CreateDefaultSubobject<UChooserParameterFloat_ContextProperty>(this, "InputValue");
+}
 
-	Result = 0.0f;
-	if (ContextObject)
+void UChooserColumnFloatRange::Filter(const UObject* ContextObject, const TArray<uint32>& IndexListIn, TArray<uint32>& IndexListOut)
+{
+	if (ContextObject && InputValue)
 	{
-		if (const FDoubleProperty* Property = FindFProperty<FDoubleProperty>(ContextObject->GetClass(), InputPropertyName))
+		float Result = 0.0f;
+		InputValue->GetValue(ContextObject, Result);
+
+		for(uint32 Index : IndexListIn)
 		{
-			Result = *Property->ContainerPtrToValuePtr<double>(ContextObject);
-	
-			for(uint32 Index : IndexListIn)
+			if (RowValues.Num() > (int)Index)
 			{
-				if (RowValues.Num() > (int)Index)
+				const FChooserFloatRangeRowData& RowValue = RowValues[Index];
+				if (Result >= RowValue.Min && Result <= RowValue.Max)
 				{
-					const FChooserFloatRangeRowData& RowValue = RowValues[Index];
-					if (Result >= RowValue.Min && Result <= RowValue.Max)
-					{
-						IndexListOut.Push(Index);
-					}
+					IndexListOut.Push(Index);
 				}
 			}
 		}
 	}
+	else
+	{
+		// passthrough fallback (behaves better during live editing)
+		IndexListOut = IndexListIn;
+	}
 }
 
 static UObject* StaticEvaluateChooser(const UObject* ContextObject, const UChooserTable* Chooser)
@@ -91,7 +123,7 @@ static UObject* StaticEvaluateChooser(const UObject* ContextObject, const UChoos
 	{
 		if (Chooser->Results.Num() > (int32)SelectedIndex)
 		{
-			const TScriptInterface<IObjectChooser>& SelectedResult = Chooser->Results[SelectedIndex];
+			const IObjectChooser* SelectedResult = Chooser->Results[SelectedIndex].GetInterface();
 
 			if(UObject* Result = SelectedResult->ChooseObject(ContextObject))
 			{
diff --git a/Engine/Plugins/Experimental/Chooser/Source/ChooserEditor/ChooserEditor.Build.cs b/Engine/Plugins/Experimental/Chooser/Source/ChooserEditor/ChooserEditor.Build.cs
--- a/Engine/Plugins/Experimental/Chooser/Source/ChooserEditor/ChooserEditor.Build.cs
+++ b/Engine/Plugins/Experimental/Chooser/Source/ChooserEditor/ChooserEditor.Build.cs
@@ -21,6 +21,9 @@ namespace UnrealBuildTool.Rules
 					"PropertyEditor",
 					"InputCore",
 					"EditorStyle",
+					"PropertyAccessEditor",
+					"PropertyEditor",
+					"BlueprintGraph"
 					// ... add private dependencies that you statically link with here ...
 				}
 			);
diff --git a/Engine/Plugins/Experimental/Chooser/Source/ChooserEditor/Private/ChooserEditorModule.cpp b/Engine/Plugins/Experimental/Chooser/Source/ChooserEditor/Private/ChooserEditorModule.cpp
--- a/Engine/Plugins/Experimental/Chooser/Source/ChooserEditor/Private/ChooserEditorModule.cpp
+++ b/Engine/Plugins/Experimental/Chooser/Source/ChooserEditor/Private/ChooserEditorModule.cpp
@@ -5,6 +5,8 @@
 #include "AssetTypeActions_Chooser.h"
 #include "ChooserTableEditor.h"
 
+#define LOCTEXT_NAMESPACE "ChooserEditorModule"
+
 namespace UE::ChooserEditor
 {
 
@@ -14,7 +16,17 @@ void FModule::StartupModule()
 	AssetTypeActions_ChooserTable = MakeShared<FAssetTypeActions_ChooserTable>();
 	AssetTools.RegisterAssetTypeActions(AssetTypeActions_ChooserTable.ToSharedRef());
 	FChooserTableEditor::RegisterWidgets();
+	
+	InterfacePropertyTypeIdentifier = MakeShared<UE::ChooserEditor::FPropertyTypeIdentifier>();
+	
+	FPropertyEditorModule& PropertyModule = FModuleManager::LoadModuleChecked<FPropertyEditorModule>("PropertyEditor");
+	
+	PropertyModule.RegisterCustomPropertyTypeLayout(
+		FInterfaceProperty::StaticClass()->GetFName(),
+		FOnGetPropertyTypeCustomizationInstance::CreateLambda([] { return MakeShared<UE::ChooserEditor::FInterfacePropertyTypeCustomization>(UChooserParameterBool::StaticClass()); }),
+		InterfacePropertyTypeIdentifier);
 }
+		 
 
 void FModule::ShutdownModule()
 {
@@ -27,4 +39,6 @@ void FModule::ShutdownModule()
 
 }
 
-IMPLEMENT_MODULE(UE::ChooserEditor::FModule, ChooserEditor);
\ No newline at end of file
+IMPLEMENT_MODULE(UE::ChooserEditor::FModule, ChooserEditor);
+
+#undef LOCTEXT_NAMESPACE
\ No newline at end of file
diff --git a/Engine/Plugins/Experimental/Chooser/Source/ChooserEditor/Private/ChooserEditorModule.h b/Engine/Plugins/Experimental/Chooser/Source/ChooserEditor/Private/ChooserEditorModule.h
--- a/Engine/Plugins/Experimental/Chooser/Source/ChooserEditor/Private/ChooserEditorModule.h
+++ b/Engine/Plugins/Experimental/Chooser/Source/ChooserEditor/Private/ChooserEditorModule.h
@@ -3,7 +3,7 @@
 #pragma once
 
 #include "CoreMinimal.h"
-#include "Modules/ModuleManager.h"
+#include "InterfacePropertyTypeCustomization.h"
 
 namespace UE::ChooserEditor
 {
@@ -17,6 +17,7 @@ private:
 	virtual void ShutdownModule() override;
 
 	TSharedPtr<FAssetTypeActions_ChooserTable> AssetTypeActions_ChooserTable;
+	TSharedPtr<UE::ChooserEditor::FPropertyTypeIdentifier> InterfacePropertyTypeIdentifier;
 };
 
 }
\ No newline at end of file
diff --git a/Engine/Plugins/Experimental/Chooser/Source/ChooserEditor/Private/ChooserTableEditor.cpp b/Engine/Plugins/Experimental/Chooser/Source/ChooserEditor/Private/ChooserTableEditor.cpp
--- a/Engine/Plugins/Experimental/Chooser/Source/ChooserEditor/Private/ChooserTableEditor.cpp
+++ b/Engine/Plugins/Experimental/Chooser/Source/ChooserEditor/Private/ChooserTableEditor.cpp
@@ -10,7 +10,6 @@
 #include "Editor.h"
 #include "Framework/MultiBox/MultiBoxBuilder.h"
 #include "Widgets/Input/SHyperlink.h"
-#include "Widgets/Input/SEditableText.h"
 #include "Widgets/Input/SNumericEntryBox.h"
 #include "Widgets/Views/SListView.h"
 #include "Widgets/Input/SComboButton.h"
@@ -21,14 +20,15 @@
 #include "DetailCategoryBuilder.h"
 #include "DetailLayoutBuilder.h"
 #include "ClassViewerFilter.h"
-#include "IDetailGroup.h"
-#include "SAssetDropTarget.h"
+#include "IPropertyAccessEditor.h"
+#include "LandscapeRender.h"
 #include "Widgets/Layout/SWidgetSwitcher.h"
 #include "DragAndDrop/DecoratedDragDropOp.h"
 #include "ObjectChooser_Asset.h"
 #include "ObjectChooserClassFilter.h"
 #include "ObjectChooserWidgetFactories.h"
-
+#include "EdGraphSchema_K2.h"
+#include "PropertyCustomizationHelpers.h"
 
 #define LOCTEXT_NAMESPACE "ChooserEditor"
 
@@ -54,25 +54,6 @@ void ConvertObjectToText(UObject* Object, FText& OutText)
 	}
 }
 
-// Filter class for things which implement IChooserColumn
-class FInterfaceClassFilter : public IClassViewerFilter
-{
-public:
-	FInterfaceClassFilter(UClass* InInterfaceType) : InterfaceType(InInterfaceType)  { };
-	virtual ~FInterfaceClassFilter() override {};
-	
-	virtual bool IsClassAllowed(const FClassViewerInitializationOptions& InInitOptions, const UClass* InClass, TSharedRef< class FClassViewerFilterFuncs > InFilterFuncs ) override
-	{
-		return (InClass->ImplementsInterface(InterfaceType) && InClass->HasAnyClassFlags(CLASS_Abstract) == false);
-	}
-
-	virtual bool IsUnloadedClassAllowed(const FClassViewerInitializationOptions& InInitOptions, const TSharedRef< const class IUnloadedBlueprintData > InUnloadedClassData, TSharedRef< class FClassViewerFilterFuncs > InFilterFuncs) override
-	{
-		return false;
-	}
-private:
-	UClass* InterfaceType;
-};
 
 void FChooserTableEditor::RegisterTabSpawners(const TSharedRef<class FTabManager>& InTabManager)
 {
@@ -401,12 +382,12 @@ public:
 			else if (ColumnName == Result) 
 			{
 				UObject* RowValue = Chooser->Results[RowIndex->RowIndex].GetObject();
-				TSharedPtr<SWidget> ResultWidget = FObjectChooserWidgetFactories::CreateWidget("", RowValue,
+				TSharedPtr<SWidget> ResultWidget = FObjectChooserWidgetFactories::CreateWidget(UObjectChooser::StaticClass(), RowValue,Chooser->ContextObjectType,
 				FOnClassPicked::CreateLambda([this, RowIndex=RowIndex->RowIndex](UClass* ChosenClass)
 				{
 					UObject* RowValue = NewObject<UObject>(Chooser, ChosenClass);
 					Chooser->Results[RowIndex] = RowValue;
-					FObjectChooserWidgetFactories::CreateWidget("", RowValue, FOnClassPicked(), &CacheBorder);
+					FObjectChooserWidgetFactories::CreateWidget(UObjectChooser::StaticClass(), RowValue, Chooser->ContextObjectType, FOnClassPicked(), &CacheBorder);
 				}),
 				&CacheBorder
 				);
@@ -448,7 +429,7 @@ public:
 				CreateRowComboButton->SetOnGetMenuContent(FOnGetContent::CreateLambda([this, CreateRowComboButton]()
 					{
 						FClassViewerInitializationOptions Options;
-						Options.ClassFilters.Add(MakeShared<FObjectChooserClassFilter>("", nullptr));
+						Options.ClassFilters.Add(MakeShared<FInterfaceClassFilter>(UObjectChooser::StaticClass()) );
 						
 						// Add class filter for columns here
 						TSharedRef<SWidget> Widget = FModuleManager::LoadModuleChecked<FClassViewerModule>("ClassViewer").CreateClassViewer(Options,
@@ -456,7 +437,7 @@ public:
 							{
 								CreateRowComboButton->SetIsOpen(false);
 								UObject* ResultObject = NewObject<UObject>(Chooser, ChosenClass);
-								Chooser->Results.Add(ResultObject);
+								Chooser->Results.Add({ResultObject});
 								Editor->UpdateTableRows();
 							}));
 						return Widget;
@@ -511,9 +492,7 @@ void FChooserTableEditor::UpdateTableColumns()
 	
 	HeaderRow->AddColumn(SHeaderRow::Column("Result")
 					.DefaultLabel(LOCTEXT("ResultColumnName", "Result"))
-					.ManualWidth(200));
-
-	ColumnText.SetNum(0);
+					.ManualWidth(300));
 
 	FName ColumnId("ChooserColumn", 1);
 	int NumColumns = Chooser->Columns.Num();	
@@ -521,33 +500,63 @@ void FChooserTableEditor::UpdateTableColumns()
 	{
 		TScriptInterface<IChooserColumn>& Column = Chooser->Columns[ColumnIndex];
 
-		ColumnText.Push(
-			SNew(SEditableText)
-					.Text_Lambda([&Column](){ return FText::FromName(Column->GetDisplayName()); })
-					.OnTextCommitted_Lambda([this, &Column](const FText& Text, ETextCommit::Type CommitType)
-					{
-						Column->SetDisplayName(ToCStr(Text.ToString()));
-						EditColumnName = false;
-					} )
-		);
-
 		HeaderRow->AddColumn(SHeaderRow::FColumn::FArguments()
 			.ColumnId(ColumnId)
 			.ManualWidth(200)
-			.OnGetMenuContent_Lambda([this, ColumnIndex]()
+			.OnGetMenuContent_Lambda([this, &Column, Chooser, ColumnIndex, ColumnId]()
 			{
+				UChooserColumnMenuContext* MenuContext = NewObject<UChooserColumnMenuContext>();
+				MenuContext->Editor = this;
+                MenuContext->Chooser = Chooser;
+                MenuContext->ColumnIndex = ColumnIndex;
+
 				FMenuBuilder MenuBuilder(true,nullptr);
-				MenuBuilder.AddMenuEntry(LOCTEXT("Delete Column","Delete Column"),LOCTEXT("Delete Column ToolTip", "Remove this column and all its data from the table"),FSlateIcon(),
+
+				MenuBuilder.AddMenuEntry(LOCTEXT("Column Properties","Properties"),LOCTEXT("Delete Column ToolTip", "Remove this column and all its data from the table"),FSlateIcon(),
 					FUIAction(
-						FExecuteAction::CreateLambda([this,ColumnIndex]()
+						FExecuteAction::CreateLambda([this,ColumnIndex, ColumnId, &Column]()
 						{
 							UChooserTable* Chooser = Cast<UChooserTable>(EditingObjects[0]);
-							Chooser->Columns.RemoveAt(ColumnIndex);
-							UpdateTableColumns();
+							DetailsView->SetObject(Chooser->Columns[ColumnIndex].GetObject());
+							SelectedColumn = ColumnId;
 						})
 						)
 					);
+
+				MenuBuilder.AddMenuEntry(LOCTEXT("Delete Column","Delete"),LOCTEXT("Delete Column ToolTip", "Remove this column and all its data from the table"),FSlateIcon(),
+					FUIAction(
+						FExecuteAction::CreateLambda([this,ColumnIndex, &Column]()
+						{
+							UChooserTable* Chooser = Cast<UChooserTable>(EditingObjects[0]);
+							Chooser->Columns.RemoveAt(ColumnIndex);
+							UpdateTableColumns();
+						})
+						));
+				
+					FClassViewerInitializationOptions Options;
+					Options.ClassFilters.Add(MakeShared<FInterfaceClassFilter>(Column->GetInputValueInterface()));
+
+				MenuBuilder.AddSubMenu(LOCTEXT("Input Type", "Input Type"),
+					LOCTEXT("InputTypeToolTip", "Change input parameter type"),
+					FNewMenuDelegate::CreateLambda([this, &Column](FMenuBuilder& Builder)
+					{
+						FClassViewerInitializationOptions Options;
+						Options.ClassFilters.Add(MakeShared<FInterfaceClassFilter>(Column->GetInputValueInterface()));
+						
+						// Add class filter for columns here
+						TSharedRef<SWidget> Widget = FModuleManager::LoadModuleChecked<FClassViewerModule>("ClassViewer").CreateClassViewer(Options, FOnClassPicked::CreateLambda([this, &Column](UClass* ChosenClass)
+						{
+							UObject* Value = NewObject<UObject>(Column.GetObject(), ChosenClass);
+							Column->SetInputValue(Value);
+							UpdateTableColumns();
+							UpdateTableRows();
+						}));
+						
+						Builder.AddWidget(Widget, LOCTEXT("", ""));
+					}));
+				
 				return MenuBuilder.MakeWidget();
+				
 			})
 			.HeaderContent()
 			[
@@ -566,20 +575,8 @@ void FChooserTableEditor::UpdateTableColumns()
 					
 					if (PointerEvent.GetEffectingButton() == EKeys::LeftMouseButton )
 					{
-						// on second click, enable editing, and select all text
-						bool ShouldEditColumn = (SelectedColumn == ColumnId);
-						if (ShouldEditColumn && !EditColumnName)
-						{
-							EditColumnName = ShouldEditColumn;
-							FSlateApplication::Get().SetKeyboardFocus(ColumnText[ColumnIndex]);  // This doesn't work, because it can't find the path to the widget? because it's not active yet
-							ColumnText[ColumnIndex]->SelectAllText();
-						}
-						EditColumnName = ShouldEditColumn;
-						
 						DetailsView->SetObject(Chooser->Columns[ColumnIndex].GetObject());
-						
 						SelectedColumn = ColumnId;
-						
 						return FReply::Handled();
 					}
 					else 
@@ -590,18 +587,7 @@ void FChooserTableEditor::UpdateTableColumns()
 					}
 				})
 				[
-					SNew(SWidgetSwitcher).WidgetIndex_Lambda([this, ColumnId]()
-					{
-						return EditColumnName && SelectedColumn == ColumnId ? 1 : 0;
-					})
-					+ SWidgetSwitcher::Slot()
-					[
-						SNew(STextBlock).Text_Lambda([&Column](){ return FText::FromName(Column->GetDisplayName()); })
-					]
-					+ SWidgetSwitcher::Slot()
-					[
-						ColumnText[ColumnIndex].ToSharedRef()
-					]
+					FObjectChooserWidgetFactories::CreateWidget(Chooser->Columns[ColumnIndex]->GetInputValue(), Chooser->ContextObjectType).ToSharedRef()
 				]
 				
 			]);
@@ -644,17 +630,13 @@ TSharedRef<SDockTab> FChooserTableEditor::SpawnTableTab( const FSpawnTabArgs& Ar
 			UChooserTable* Chooser = Cast<UChooserTable>(EditingObjects[0]);
 			UObject* ColumnObject = NewObject<UObject>(Chooser, ChosenClass);
 			TScriptInterface<IChooserColumn> Column;
-			Column.SetObject(ColumnObject);
-			Column.SetInterface(CastChecked<IChooserColumn>(ColumnObject));
+			Column = ColumnObject;
 			Chooser->Columns.Add(TScriptInterface<IChooserColumn>(Column));
 			UpdateTableColumns();
 			UpdateTableRows();
 			SelectedColumn = "ChooserColumn";
-			SelectedColumn.SetNumber(ColumnText.Num());
+			SelectedColumn.SetNumber(Chooser->Columns.Num());
 			DetailsView->SetObject(ColumnObject);
-			EditColumnName = true;
-			FSlateApplication::Get().SetKeyboardFocus(ColumnText[ColumnText.Num() - 1]);
-			ColumnText[ColumnText.Num() - 1]->SelectAllText();
 		}));
 		return Widget;
 	})
@@ -882,68 +864,54 @@ void FChooserTableEditor::PostRegenerateMenusAndToolbars()
 	}
 }
 
+	
+void FChooserTableEditor::SelectColumn(int Index)
+{
+	UChooserTable* Chooser = GetChooser();
+	if (Index < Chooser->Columns.Num())
+	{
+		SelectedColumn = "Column";
+		SelectedColumn.SetNumber(Index);
+		DetailsView->SetObject(Chooser->Columns[Index].GetObject());
+	}
+}
+	
+void FChooserTableEditor::DeleteColumn(int Index)
+{
+	UChooserTable* Chooser = GetChooser();
+	if (Index < Chooser->Columns.Num())
+	{
+		Chooser->Columns.RemoveAt(Index);
+		UpdateTableColumns();
+	}
+}
+
 /// Result widgets
 
 TMap<const UClass*, TFunction<TSharedRef<SWidget> (UObject* Column, int Row)>> FChooserTableEditor::ColumnWidgetCreators;
 
-TSharedRef<SWidget> CreateAssetWidget(UObject* Object)
+TSharedRef<SWidget> CreateAssetWidget(UObject* Object, UClass* ContextClass)
 {
 	UObjectChooser_Asset* DIAsset = Cast<UObjectChooser_Asset>(Object);
 
 	UObject* Asset = DIAsset->Asset;
+
+	UChooserTable* Chooser = Object->GetTypedOuter<UChooserTable>();
 	
-	return
-	SNew(SAssetDropTarget)
-	.OnAssetsDropped_Lambda([DIAsset](const FDragDropEvent& Event, TArrayView<FAssetData> AssetData)
-	{
-		if (AssetData.Num() > 0)
-		{
-			DIAsset->Asset = AssetData[0].GetAsset();
-		}
-	})
-	.Content()
-	[
-		SNew (STextBlock).Text_Lambda([DIAsset]()
-		{
-			if (DIAsset->Asset == nullptr)
-			{
-				return LOCTEXT("NoAsset", "[No Asset]");
-			}
-			else
-			{
-				return FText::FromString(DIAsset->Asset->GetName());
-			}
-		})
-	];
+	return SNew(SObjectPropertyEntryBox)
+		.AllowedClass(Chooser->OutputObjectType!=nullptr ? Chooser->OutputObjectType.Get() : UObject::StaticClass())
+		.ObjectPath_Lambda([DIAsset](){ return DIAsset->Asset ? DIAsset->Asset.GetPath() : "";})
+		.OnObjectChanged_Lambda([DIAsset](const FAssetData& AssetData){ DIAsset->Asset = AssetData.GetAsset(); });
 }
 
-TSharedRef<SWidget> CreateEvaluateChooserWidget(UObject* Object)
+TSharedRef<SWidget> CreateEvaluateChooserWidget(UObject* Object, UClass* ContextObject)
 {
 	UObjectChooser_EvaluateChooser* EvaluateChooser = Cast<UObjectChooser_EvaluateChooser>(Object);
-
-	return
-	SNew(SAssetDropTarget)
-	.OnAssetsDropped_Lambda([EvaluateChooser](const FDragDropEvent& Event, TArrayView<FAssetData> AssetData)
-	{
-		if (AssetData.Num() > 0)
-		{
-			EvaluateChooser->Chooser = Cast<UChooserTable>(AssetData[0].GetAsset());
-		}
-	})
-	.Content()
-	[
-		SNew (STextBlock).Text_Lambda([EvaluateChooser]()
-		{
-			if (EvaluateChooser->Chooser == nullptr)
-			{
-				return LOCTEXT("NoChooser", "[No Chooser]");
-			}
-			else
-			{
-				return FText::FromString(EvaluateChooser->Chooser->GetName());
-			}
-		})
-	];
+	
+	return SNew(SObjectPropertyEntryBox)
+		.AllowedClass(UChooserTable::StaticClass())
+		.ObjectPath_Lambda([EvaluateChooser](){ return EvaluateChooser->Chooser ? EvaluateChooser->Chooser.GetPath() : "";})
+		.OnObjectChanged_Lambda([EvaluateChooser](const FAssetData& AssetData){ EvaluateChooser->Chooser = Cast<UChooserTable>(AssetData.GetAsset()); });
 }
 
 void ConvertToText_Asset(const UObject* Object, FText& OutText)
@@ -972,7 +940,7 @@ void ConvertToText_EvaluateChooser(const UObject* Object, FText& OutText)
 	}
 }
 
-TSharedRef<SWidget> CreateObjectWidget(UObject* Object)
+TSharedRef<SWidget> CreateObjectWidget(UObject* Object, UClass* ContextClass)
 {
 	FText ObjectText = FText::FromString(Object->GetName());
 	ConvertObjectToText(Object, ObjectText);
@@ -998,6 +966,89 @@ TSharedRef<SWidget> CreateBoolColumnWidget(UObject* Column, int Row)
 		return value ? ECheckBoxState::Checked : ECheckBoxState::Unchecked;
 	});
 }
+	
+template <typename PropertyType>
+void ConvertToText_ContextProperty(const UObject* Object, FText& OutText)
+{
+	if (const PropertyType* ContextProperty = Cast<PropertyType>(Object))
+	{
+		OutText = FText::FromName(ContextProperty->PropertyName);
+	}
+}
+
+template <typename PropertyType>
+TSharedRef<SWidget> CreatePropertyWidget(UObject* Object, UClass* ContextClass)
+{
+	TWeakObjectPtr<PropertyType> ContextProperty = Cast<PropertyType>(Object);
+
+	FPropertyBindingWidgetArgs Args;
+	Args.bAllowPropertyBindings = true;
+	Args.OnCanBindProperty = FOnCanBindProperty::CreateLambda([](FProperty* Property)
+	{
+		return Property == nullptr || Property->GetCPPType() == PropertyType::CPPTypeName();
+	});
+	Args.OnCanBindToClass = FOnCanBindToClass::CreateLambda([ContextClass ](UClass* InClass)
+	{
+		return ContextClass->IsChildOf(InClass);
+	});
+	
+	Args.CurrentBindingColor = MakeAttributeLambda([ContextProperty, ContextClass]() {
+		if (ContextProperty.IsValid())
+		{
+			if (const FProperty* Property = FindFProperty<FProperty>(ContextClass, ContextProperty->PropertyName))
+			{
+				const UEdGraphSchema_K2* Schema = GetDefault<UEdGraphSchema_K2>();
+				FEdGraphPinType VariablePinType;
+				Schema->ConvertPropertyToPinType(Property, VariablePinType);
+				return Schema->GetPinTypeColor(VariablePinType);
+			}
+		}
+		return FLinearColor::Gray;
+	});
+
+	Args.OnCanAcceptPropertyOrChildren = FOnCanBindProperty::CreateLambda([](FProperty* InProperty)
+		{
+			// Make only editor visible properties visible for binding.
+			return InProperty->HasAnyPropertyFlags(CPF_Edit);
+		});	
+
+	Args.OnAddBinding = FOnAddBinding::CreateLambda([ContextProperty](FName InPropertyName, const TArray<FBindingChainElement>& InBindingChain)
+		{
+			IPropertyAccessEditor& PropertyAccessEditor = IModularFeatures::Get().GetModularFeature<IPropertyAccessEditor>("PropertyAccessEditor");
+			if (ContextProperty.IsValid())
+			{
+				if (InBindingChain.Num() == 1)
+			 	{
+					TArray<FString> PropertyPath; 
+					PropertyAccessEditor.MakeStringPath(InBindingChain, PropertyPath);
+					ContextProperty->PropertyName = FName(PropertyPath[0]);
+					// do this better so undo will work?
+				}
+			}
+		});
+
+	Args.CurrentBindingText = MakeAttributeLambda([ContextProperty]()
+			{
+				const FText Bind = LOCTEXT("Bind", "Bind");
+				FText CurrentValue = Bind;
+
+				if (ContextProperty.IsValid())
+				{
+					CurrentValue = FText::FromName(ContextProperty->PropertyName);
+				}
+
+				return CurrentValue;
+			});
+
+	Args.CurrentBindingImage = MakeAttributeLambda([]() -> const FSlateBrush*
+		{
+			static FName PropertyIcon(TEXT("Kismet.Tabs.Variables"));
+			return FAppStyle::GetBrush(PropertyIcon);
+		});
+
+	IPropertyAccessEditor& PropertyAccessEditor = IModularFeatures::Get().GetModularFeature<IPropertyAccessEditor>("PropertyAccessEditor");
+	return PropertyAccessEditor.MakePropertyBindingWidget((UBlueprint*)ContextClass, Args);
+}
 
 TSharedRef<SWidget> CreateFloatRangeColumnWidget(UObject* Column, int Row)
 {
@@ -1107,7 +1158,12 @@ void FChooserRowDetails::CustomizeDetails(IDetailLayoutBuilder& DetailBuilder)
 				TSharedRef<IPropertyHandle> CellData = ColumnDataProperty->AsArray()->GetElement(Row->Row);
 	
 				IDetailPropertyRow& NewColumnProperty = PropertiesCategory.AddProperty(CellData);
-				NewColumnProperty.DisplayName(FText::FromName(Column.GetInterface()->GetDisplayName()));
+				FText DisplayName = LOCTEXT("No Input Value", "No Input Value");
+				if (UObject* InputValue = Column.GetInterface()->GetInputValue())
+				{
+					ConvertObjectToText(InputValue, DisplayName);
+				}
+				NewColumnProperty.DisplayName(DisplayName);
 				NewColumnProperty.ShowPropertyButtons(false); // hide array add button
 				NewColumnProperty.ShouldAutoExpand(true);
 			}
@@ -1123,6 +1179,11 @@ void FChooserTableEditor::RegisterWidgets()
 	FObjectChooserWidgetFactories::ChooserWidgetCreators.Add(UObject::StaticClass(), CreateObjectWidget);
 	FObjectChooserWidgetFactories::ChooserWidgetCreators.Add(UObjectChooser_Asset::StaticClass(), CreateAssetWidget);
 	FObjectChooserWidgetFactories::ChooserWidgetCreators.Add(UObjectChooser_EvaluateChooser::StaticClass(), CreateEvaluateChooserWidget);
+	
+	FObjectChooserWidgetFactories::ChooserWidgetCreators.Add(UChooserParameterBool_ContextProperty::StaticClass(), CreatePropertyWidget<UChooserParameterBool_ContextProperty>);
+	FObjectChooserWidgetFactories::ChooserWidgetCreators.Add(UChooserParameterFloat_ContextProperty::StaticClass(), CreatePropertyWidget<UChooserParameterFloat_ContextProperty>);
+	FObjectChooserWidgetFactories::ChooserTextConverter.Add(UChooserParameterBool_ContextProperty::StaticClass(), ConvertToText_ContextProperty<UChooserParameterBool_ContextProperty>);
+	FObjectChooserWidgetFactories::ChooserTextConverter.Add(UChooserParameterFloat_ContextProperty::StaticClass(), ConvertToText_ContextProperty<UChooserParameterFloat_ContextProperty>);
 
 	ColumnWidgetCreators.Add(UChooserColumnBool::StaticClass(), CreateBoolColumnWidget);
 	ColumnWidgetCreators.Add(UChooserColumnFloatRange::StaticClass(), CreateFloatRangeColumnWidget);
diff --git a/Engine/Plugins/Experimental/Chooser/Source/ChooserEditor/Private/ChooserTableEditor.h b/Engine/Plugins/Experimental/Chooser/Source/ChooserEditor/Private/ChooserTableEditor.h
--- a/Engine/Plugins/Experimental/Chooser/Source/ChooserEditor/Private/ChooserTableEditor.h
+++ b/Engine/Plugins/Experimental/Chooser/Source/ChooserEditor/Private/ChooserTableEditor.h
@@ -78,6 +78,8 @@ namespace UE::ChooserEditor
 		};
 
 		void UpdateTableRows();
+		void SelectColumn(int Index);
+		void DeleteColumn(int Index);
 	private:
 
 		FReply SelectRootProperties();
@@ -117,10 +119,7 @@ namespace UE::ChooserEditor
 		TSharedPtr<SHeaderRow> HeaderRow;
 		TSharedPtr<SListView<TSharedPtr<FChooserTableRow>>> TableView;
 
-		TArray<TSharedPtr<SEditableText>> ColumnText;
-
 		FName SelectedColumn;
-		bool EditColumnName = false;
 	public:
 
 		static TMap<const UClass*, TFunction<TSharedRef<SWidget> (UObject* Column, int Row)>> ColumnWidgetCreators;
diff --git a/Engine/Plugins/Experimental/Chooser/Source/ChooserEditor/Private/InterfacePropertyTypeCustomization.cpp b/Engine/Plugins/Experimental/Chooser/Source/ChooserEditor/Private/InterfacePropertyTypeCustomization.cpp
new file mode 100644
--- /dev/null
+++ b/Engine/Plugins/Experimental/Chooser/Source/ChooserEditor/Private/InterfacePropertyTypeCustomization.cpp
@@ -0,0 +1,96 @@
+﻿// Copyright Epic Games, Inc. All Rights Reserved.
+
+#include "InterfacePropertyTypeCustomization.h"
+
+#include "Chooser.h"
+#include "PropertyHandle.h"
+#include "DetailWidgetRow.h"
+#include "PropertyCustomizationHelpers.h"
+#include "Widgets/Text/STextBlock.h"
+#include "SClassViewer.h"
+#include "ObjectChooserWidgetFactories.h"
+#include "IDetailChildrenBuilder.h"
+#include "PropertyEditor/Private/ObjectPropertyNode.h"
+#include "PropertyEditor/Private/PropertyNode.h"
+
+#define LOCTEXT_NAMESPACE "InterfacePropertyTypeCustomization"
+
+namespace UE::ChooserEditor
+{
+
+bool FPropertyTypeIdentifier::IsPropertyTypeCustomized(const IPropertyHandle& PropertyHandle) const
+{
+	return (PropertyHandle.GetMetaData("EditInlineInterface") != "");
+}
+
+void FInterfacePropertyTypeCustomization::CustomizeHeader(TSharedRef<IPropertyHandle> PropertyHandle, FDetailWidgetRow& HeaderRow, IPropertyTypeCustomizationUtils& CustomizationUtils)
+{
+	FProperty* Property = PropertyHandle->GetProperty();
+	FInterfaceProperty* IntProperty = CastField<FInterfaceProperty>( Property );
+
+	void* ValuePtr;
+	PropertyHandle->GetValueData(ValuePtr);
+	FScriptInterface* PropertyValue = reinterpret_cast<FScriptInterface*>(ValuePtr);
+	
+	UClass* ContextClass = nullptr;
+	if (const UObject* Object = PropertyValue->GetObject())
+	{
+		if (const UChooserTable* ChooserTable = Object->GetTypedOuter<UChooserTable>())
+		{
+			ContextClass = ChooserTable->ContextObjectType;
+		}
+	}
+	
+	TSharedPtr<SWidget> Widget = FObjectChooserWidgetFactories::CreateWidget(InterfaceType, PropertyValue->GetObject(), ContextClass,
+			FOnClassPicked::CreateLambda([this, PropertyHandle](UClass* ChosenClass)
+			{
+				TArray<void*> RawData;
+				PropertyHandle->AccessRawData(RawData);
+				TArray<UPackage*> OuterObjects;
+				PropertyHandle->GetOuterPackages(OuterObjects);
+
+				for(int i=0;i<RawData.Num();i++)
+				{
+					UObject* NewValue = NewObject<UObject>(OuterObjects[i], ChosenClass);
+				
+					PropertyHandle->NotifyPreChange();
+					PropertyHandle->SetValue(NewValue);
+					PropertyHandle->NotifyPostChange(EPropertyChangeType::ValueSet);
+					PropertyHandle->GetPropertyNode()->GetParentNode()->RequestRebuildChildren();
+				}
+			}),
+			nullptr
+	);
+
+	HeaderRow
+	.NameContent()
+	[
+		PropertyHandle->CreatePropertyNameWidget()
+	]
+	.ValueContent()
+	[
+		Widget.ToSharedRef()
+	];
+}
+	
+void FInterfacePropertyTypeCustomization::CustomizeChildren(TSharedRef<IPropertyHandle> PropertyHandle, IDetailChildrenBuilder& ChildBuilder, IPropertyTypeCustomizationUtils& CustomizationUtils)
+{
+	// todo: make this work correctly for multiselect
+	void* ValuePtr;
+	PropertyHandle->GetValueData(ValuePtr);
+	FScriptInterface* PropertyValue = reinterpret_cast<FScriptInterface*>(ValuePtr);
+	if (PropertyValue->GetObject())
+	{
+		TSharedPtr<FStructOnScope> StructData = MakeShareable(new FStructOnScope(PropertyValue->GetObject()->GetClass(), (uint8*)PropertyValue->GetObject()));
+		StructData->SetPackage(PropertyValue->GetObject()->GetPackage());
+		TArray<TSharedPtr<IPropertyHandle>> NewHandles = ChildBuilder.AddAllExternalStructureProperties(StructData.ToSharedRef());
+		for (auto& NewHandle : NewHandles)
+		{
+			ChildBuilder.AddProperty(NewHandle.ToSharedRef());
+		}
+	}
+}
+
+}
+
+#undef LOCTEXT_NAMESPACE
\ No newline at end of file
diff --git a/Engine/Plugins/Experimental/Chooser/Source/ChooserEditor/Private/InterfacePropertyTypeCustomization.h b/Engine/Plugins/Experimental/Chooser/Source/ChooserEditor/Private/InterfacePropertyTypeCustomization.h
new file mode 100644
--- /dev/null
+++ b/Engine/Plugins/Experimental/Chooser/Source/ChooserEditor/Private/InterfacePropertyTypeCustomization.h
@@ -0,0 +1,28 @@
+﻿// Copyright Epic Games, Inc. All Rights Reserved.
+
+#pragma once
+
+#include "IPropertyTypeCustomization.h"
+#include "PropertyEditorModule.h"
+
+namespace UE::ChooserEditor
+{
+
+class FInterfacePropertyTypeCustomization : public IPropertyTypeCustomization
+{
+public:
+	FInterfacePropertyTypeCustomization(UClass* Interface) : InterfaceType(Interface) {}
+private:
+	// IPropertyTypeCustomization interface
+	virtual void CustomizeHeader(TSharedRef<IPropertyHandle> PropertyHandle, FDetailWidgetRow& HeaderRow, IPropertyTypeCustomizationUtils& CustomizationUtils) override;
+	virtual void CustomizeChildren(TSharedRef<IPropertyHandle> PropertyHandle, IDetailChildrenBuilder& ChildBuilder, IPropertyTypeCustomizationUtils& CustomizationUtils) override;
+
+	UClass* InterfaceType;
+};
+
+class FPropertyTypeIdentifier : public IPropertyTypeIdentifier
+{
+	virtual bool IsPropertyTypeCustomized(const IPropertyHandle& PropertyHandle) const override;
+};
+
+}
\ No newline at end of file
diff --git a/Engine/Plugins/Experimental/Chooser/Source/ChooserEditor/Private/ObjectChooserWidgetFactories.cpp b/Engine/Plugins/Experimental/Chooser/Source/ChooserEditor/Private/ObjectChooserWidgetFactories.cpp
--- a/Engine/Plugins/Experimental/Chooser/Source/ChooserEditor/Private/ObjectChooserWidgetFactories.cpp
+++ b/Engine/Plugins/Experimental/Chooser/Source/ChooserEditor/Private/ObjectChooserWidgetFactories.cpp
@@ -28,24 +28,30 @@ void ConvertToText(UObject* Object, FText& OutText)
 		Class = Class->GetSuperClass();
 	}
 }
-
-TSharedPtr<SWidget> FObjectChooserWidgetFactories::CreateWidget(FName DataInterfaceTypeName, UObject* Value, const FOnClassPicked& CreateClassCallback, TSharedPtr<SBorder>* InnerWidget)
+	
+TSharedPtr<SWidget> FObjectChooserWidgetFactories::CreateWidget(UObject* Value, UClass* ContextClass)
 {
-	TSharedPtr<SWidget> LeftWidget;
-
 	if (Value)
 	{
 		UClass* Class = Value->GetClass();
-		while (Class && !LeftWidget.IsValid())
+		while (Class)
 		{
 			if (FChooserWidgetCreator* Creator = FObjectChooserWidgetFactories::ChooserWidgetCreators.Find(Class))
 			{
-				LeftWidget = (*Creator)(Value);
+				return (*Creator)(Value, ContextClass);
 				break;
 			}
 			Class = Class->GetSuperClass();
 		}
 	}
+
+	return nullptr;
+}
+
+TSharedPtr<SWidget> FObjectChooserWidgetFactories::CreateWidget(UClass* InterfaceType, UObject* Value, UClass* ContextClass, const FOnClassPicked& CreateClassCallback, TSharedPtr<SBorder>* InnerWidget)
+{
+	TSharedPtr<SWidget> LeftWidget = CreateWidget(Value, ContextClass);
+
 	
 	if (!LeftWidget.IsValid())
 	{
@@ -57,10 +63,10 @@ TSharedPtr<SWidget> FObjectChooserWidgetFactories::CreateWidget(FName DataInterf
 	TSharedPtr<SComboButton> Button = SNew(SComboButton)
 			.ComboButtonStyle(FAppStyle::Get(), "SimpleComboButton");
 	
-	Button->SetOnGetMenuContent(FOnGetContent::CreateLambda([DataInterfaceTypeName, Button, CreateClassCallback]()
+	Button->SetOnGetMenuContent(FOnGetContent::CreateLambda([InterfaceType, Button, CreateClassCallback]()
 	{
 		FClassViewerInitializationOptions Options;
-		Options.ClassFilters.Add(MakeShared<FObjectChooserClassFilter>(DataInterfaceTypeName, nullptr));
+		Options.ClassFilters.Add(MakeShared<FInterfaceClassFilter>(InterfaceType));
 		
 		// Add class filter for columns here
 		TSharedRef<SWidget>  ClassMenu = FModuleManager::LoadModuleChecked<FClassViewerModule>("ClassViewer").
@@ -90,11 +96,11 @@ TSharedPtr<SWidget> FObjectChooserWidgetFactories::CreateWidget(FName DataInterf
 	Border->SetContent(LeftWidget.ToSharedRef());
 
 	TSharedPtr<SWidget> Widget = SNew(SHorizontalBox)
-		+SHorizontalBox::Slot().FillWidth(75)
+		+SHorizontalBox::Slot().FillWidth(100)
 		[
 			Border.ToSharedRef()
 		]
-		+SHorizontalBox::Slot().FillWidth(25)
+		+SHorizontalBox::Slot().AutoWidth()
 		[
 			Button.ToSharedRef()
 		]
diff --git a/Engine/Plugins/Experimental/Chooser/Source/ChooserEditor/Public/ObjectChooserClassFilter.h b/Engine/Plugins/Experimental/Chooser/Source/ChooserEditor/Public/ObjectChooserClassFilter.h
--- a/Engine/Plugins/Experimental/Chooser/Source/ChooserEditor/Public/ObjectChooserClassFilter.h
+++ b/Engine/Plugins/Experimental/Chooser/Source/ChooserEditor/Public/ObjectChooserClassFilter.h
@@ -6,22 +6,16 @@
 
 namespace UE::ChooserEditor
 {
-	// Filter class for things which implement IChooserColumn
-	class FObjectChooserClassFilter : public IClassViewerFilter
+	// Filter class for things which implement a particular interface
+	class FInterfaceClassFilter : public IClassViewerFilter
 	{
 	public:
-		FObjectChooserClassFilter(FName InTypeName, const UScriptStruct* InStructType) : TypeName(InTypeName), StructType(InStructType) { };
-		virtual ~FObjectChooserClassFilter() override {};
-	
+		FInterfaceClassFilter(UClass* InInterfaceType) : InterfaceType(InInterfaceType)  { };
+		virtual ~FInterfaceClassFilter() override {};
+		
 		virtual bool IsClassAllowed(const FClassViewerInitializationOptions& InInitOptions, const UClass* InClass, TSharedRef< class FClassViewerFilterFuncs > InFilterFuncs ) override
 		{
-			if (InClass->ImplementsInterface(UObjectChooser::StaticClass()))
-			{
-				// for now, just return true if the class implements IObjectChooser
-				return true;
-			}
-
-			return false;
+			return (InClass->ImplementsInterface(InterfaceType) && InClass->HasAnyClassFlags(CLASS_Abstract) == false);
 		}
 
 		virtual bool IsUnloadedClassAllowed(const FClassViewerInitializationOptions& InInitOptions, const TSharedRef< const class IUnloadedBlueprintData > InUnloadedClassData, TSharedRef< class FClassViewerFilterFuncs > InFilterFuncs) override
@@ -29,7 +23,7 @@ namespace UE::ChooserEditor
 			return false;
 		}
 	private:
-		FName TypeName;
-		const UScriptStruct* StructType;
+		UClass* InterfaceType;
 	};
+		
 }
\ No newline at end of file
diff --git a/Engine/Plugins/Experimental/Chooser/Source/ChooserEditor/Public/ObjectChooserWidgetFactories.h b/Engine/Plugins/Experimental/Chooser/Source/ChooserEditor/Public/ObjectChooserWidgetFactories.h
--- a/Engine/Plugins/Experimental/Chooser/Source/ChooserEditor/Public/ObjectChooserWidgetFactories.h
+++ b/Engine/Plugins/Experimental/Chooser/Source/ChooserEditor/Public/ObjectChooserWidgetFactories.h
@@ -10,7 +10,7 @@
 namespace UE::ChooserEditor
 {
 	typedef TFunction<void(const UObject* Object, FText& OutText)> FChooserTextConverter;
-	typedef TFunction<TSharedRef<SWidget>(UObject* Object)> FChooserWidgetCreator;
+	typedef TFunction<TSharedRef<SWidget>(UObject* Object, UClass* ContextClass)> FChooserWidgetCreator;
 
 	class CHOOSEREDITOR_API FObjectChooserWidgetFactories
 	{
@@ -18,7 +18,8 @@ namespace UE::ChooserEditor
 		static TMap<const UClass*, FChooserTextConverter> ChooserTextConverter;
 		static TMap<const UClass*, FChooserWidgetCreator> ChooserWidgetCreators;
 
-		static TSharedPtr<SWidget> CreateWidget(FName DataInterfaceTypeName, UObject* Value, const FOnClassPicked& CreateClassCallback, TSharedPtr<SBorder>* InnerWidget = nullptr);
+		static TSharedPtr<SWidget> CreateWidget(UObject* Value, UClass* ContextClass);
+		static TSharedPtr<SWidget> CreateWidget(UClass* InterfaceType, UObject* Value, UClass* ContextClass, const FOnClassPicked& CreateClassCallback, TSharedPtr<SBorder>* InnerWidget = nullptr);
 		
 		static void RegisterWidgets();
 	};
