Fix empty LOCTEXT macro

diff --git a/Engine/Plugins/Experimental/Chooser/Source/ChooserEditor/Private/ChooserTableEditor.cpp b/Engine/Plugins/Experimental/Chooser/Source/ChooserEditor/Private/ChooserTableEditor.cpp
--- a/Engine/Plugins/Experimental/Chooser/Source/ChooserEditor/Private/ChooserTableEditor.cpp
+++ b/Engine/Plugins/Experimental/Chooser/Source/ChooserEditor/Private/ChooserTableEditor.cpp
@@ -552,7 +552,7 @@ void FChooserTableEditor::UpdateTableColumns()
 							UpdateTableRows();
 						}));
 						
-						Builder.AddWidget(Widget, LOCTEXT("", ""));
+						Builder.AddWidget(Widget, FText());
 					}));
 				
 				return MenuBuilder.MakeWidget();
