Make column and column parameter interfaces for Choosers public

diff --git a/Engine/Plugins/Experimental/Chooser/Source/Chooser/Internal/Chooser.h b/Engine/Plugins/Experimental/Chooser/Source/Chooser/Internal/Chooser.h
--- a/Engine/Plugins/Experimental/Chooser/Source/Chooser/Internal/Chooser.h
+++ b/Engine/Plugins/Experimental/Chooser/Source/Chooser/Internal/Chooser.h
@@ -4,24 +4,11 @@
 #include "CoreMinimal.h"
 #include "UObject/Object.h"
 #include "IObjectChooser.h"
-#include "UObject/Interface.h"
+#include "IChooserColumn.h"
+#include "IChooserColumnParameter.h"
 
 #include "Chooser.generated.h"
 
-UINTERFACE(NotBlueprintType, meta = (CannotImplementInterfaceInBlueprint))
-class CHOOSER_API UChooserParameterBool : public UInterface
-{
-	GENERATED_BODY()
-};
-
-class CHOOSER_API IChooserParameterBool
-{
-	GENERATED_BODY()
-
-public:
-	virtual bool GetValue(const UObject* ContextObject, bool& OutResult) { return false; }
-};
-
 UCLASS()
 class CHOOSER_API UChooserParameterBool_ContextProperty :  public UObject, public IChooserParameterBool
 {
@@ -39,21 +26,6 @@ public:
 	}
 };
 
-
-UINTERFACE(NotBlueprintType, meta = (CannotImplementInterfaceInBlueprint))
-class CHOOSER_API UChooserParameterFloat : public UInterface
-{
-	GENERATED_BODY()
-};
-
-class CHOOSER_API IChooserParameterFloat
-{
-	GENERATED_BODY()
-
-public:
-	virtual bool GetValue(const UObject* ContextObject, float& OutResult) { return false; }
-};
-
 UCLASS()
 class CHOOSER_API UChooserParameterFloat_ContextProperty :  public UObject, public IChooserParameterFloat
 {
@@ -71,25 +43,6 @@ public:
 	}
 };
 
-UINTERFACE(NotBlueprintType, meta = (CannotImplementInterfaceInBlueprint))
-class CHOOSER_API UChooserColumn : public UInterface
-{
-	GENERATED_BODY()
-};
-
-class CHOOSER_API IChooserColumn 
-{
-	GENERATED_BODY()
-
-public:
-	virtual void Filter(const UObject* ContextObject, const TArray<uint32>& IndexListIn, TArray<uint32>& IndexListOut) {};
-	virtual void SetNumRows(uint32 NumRows) {}
-	virtual void DeleteRows(const TArray<uint32> & RowIndices) {}
-	virtual UClass* GetInputValueInterface() {return nullptr;};
-	virtual UObject* GetInputValue() {return nullptr;}
-	virtual void SetInputValue(UObject* InputValue) {}
-};
-
 UCLASS()
 class CHOOSER_API UChooserColumnBool : public UObject, public IChooserColumn
 {
diff --git a/Engine/Plugins/Experimental/Chooser/Source/Chooser/Public/IChooserColumn.h b/Engine/Plugins/Experimental/Chooser/Source/Chooser/Public/IChooserColumn.h
new file mode 100644
--- /dev/null
+++ b/Engine/Plugins/Experimental/Chooser/Source/Chooser/Public/IChooserColumn.h
@@ -0,0 +1,27 @@
+// Copyright Epic Games, Inc. All Rights Reserved.
+
+#pragma once
+
+#include "CoreMinimal.h"
+#include "UObject/Object.h"
+#include "UObject/Interface.h"
+#include "IChooserColumn.generated.h"
+
+UINTERFACE(NotBlueprintType, meta = (CannotImplementInterfaceInBlueprint))
+class CHOOSER_API UChooserColumn : public UInterface
+{
+	GENERATED_BODY()
+};
+
+class CHOOSER_API IChooserColumn 
+{
+	GENERATED_BODY()
+
+public:
+	virtual void Filter(const UObject* ContextObject, const TArray<uint32>& IndexListIn, TArray<uint32>& IndexListOut) {};
+	virtual void SetNumRows(uint32 NumRows) {}
+	virtual void DeleteRows(const TArray<uint32> & RowIndices) {}
+	virtual UClass* GetInputValueInterface() {return nullptr;};
+	virtual UObject* GetInputValue() {return nullptr;}
+	virtual void SetInputValue(UObject* InputValue) {}
+};
diff --git a/Engine/Plugins/Experimental/Chooser/Source/Chooser/Public/IChooserColumnParameter.h b/Engine/Plugins/Experimental/Chooser/Source/Chooser/Public/IChooserColumnParameter.h
new file mode 100644
--- /dev/null
+++ b/Engine/Plugins/Experimental/Chooser/Source/Chooser/Public/IChooserColumnParameter.h
@@ -0,0 +1,36 @@
+// Copyright Epic Games, Inc. All Rights Reserved.
+
+#pragma once
+
+#include "CoreMinimal.h"
+#include "UObject/Object.h"
+#include "UObject/Interface.h"
+#include "IChooserColumnParameter.generated.h"
+
+UINTERFACE(NotBlueprintType, meta = (CannotImplementInterfaceInBlueprint))
+class CHOOSER_API UChooserParameterBool : public UInterface
+{
+	GENERATED_BODY()
+};
+
+class CHOOSER_API IChooserParameterBool
+{
+	GENERATED_BODY()
+
+public:
+	virtual bool GetValue(const UObject* ContextObject, bool& OutResult) { return false; }
+};
+
+UINTERFACE(NotBlueprintType, meta = (CannotImplementInterfaceInBlueprint))
+class CHOOSER_API UChooserParameterFloat : public UInterface
+{
+	GENERATED_BODY()
+};
+
+class CHOOSER_API IChooserParameterFloat
+{
+	GENERATED_BODY()
+
+public:
+	virtual bool GetValue(const UObject* ContextObject, float& OutResult) { return false; }
+};
\ No newline at end of file
