Chooser Table Editor undo support

diff --git a/Engine/Plugins/Experimental/Chooser/Source/Chooser/Private/Chooser.cpp b/Engine/Plugins/Experimental/Chooser/Source/Chooser/Private/Chooser.cpp
--- a/Engine/Plugins/Experimental/Chooser/Source/Chooser/Private/Chooser.cpp
+++ b/Engine/Plugins/Experimental/Chooser/Source/Chooser/Private/Chooser.cpp
@@ -31,6 +31,7 @@ UChooserTable::UChooserTable(const FObjectInitializer& Initializer)
 UChooserColumnBool::UChooserColumnBool(const FObjectInitializer& ObjectInitializer)
 {
 	InputValue = ObjectInitializer.CreateDefaultSubobject<UChooserParameterBool_ContextProperty>(this, "InputValue");
+	InputValue.GetObject()->SetFlags(RF_Transactional);
 }
 
 void UChooserColumnBool::Filter(const UObject* ContextObject, const TArray<uint32>& IndexListIn, TArray<uint32>& IndexListOut)
@@ -61,6 +62,7 @@ void UChooserColumnBool::Filter(const UObject* ContextObject, const TArray<uint3
 UChooserColumnFloatRange::UChooserColumnFloatRange(const FObjectInitializer& ObjectInitializer)
 {
 	InputValue = ObjectInitializer.CreateDefaultSubobject<UChooserParameterFloat_ContextProperty>(this, "InputValue");
+	InputValue.GetObject()->SetFlags(RF_Transactional);
 }
 
 void UChooserColumnFloatRange::Filter(const UObject* ContextObject, const TArray<uint32>& IndexListIn, TArray<uint32>& IndexListOut)
diff --git a/Engine/Plugins/Experimental/Chooser/Source/ChooserEditor/Private/ChooserTableEditor.cpp b/Engine/Plugins/Experimental/Chooser/Source/ChooserEditor/Private/ChooserTableEditor.cpp
--- a/Engine/Plugins/Experimental/Chooser/Source/ChooserEditor/Private/ChooserTableEditor.cpp
+++ b/Engine/Plugins/Experimental/Chooser/Source/ChooserEditor/Private/ChooserTableEditor.cpp
@@ -28,7 +28,9 @@
 #include "ObjectChooserClassFilter.h"
 #include "ObjectChooserWidgetFactories.h"
 #include "EdGraphSchema_K2.h"
+#include "IDetailCustomization.h"
 #include "PropertyCustomizationHelpers.h"
+#include "ScopedTransaction.h"
 
 #define LOCTEXT_NAMESPACE "ChooserEditor"
 
@@ -85,7 +87,7 @@ FChooserTableEditor::~FChooserTableEditor()
 {
 	GEditor->GetEditorSubsystem<UImportSubsystem>()->OnAssetPostImport.RemoveAll(this);
 	FCoreUObjectDelegates::OnObjectsReplaced.RemoveAll(this);
-
+	
 	DetailsView.Reset();
 }
 
@@ -143,6 +145,18 @@ FText FChooserTableEditor::GetBaseToolkitName() const
 	return LOCTEXT("AppLabel", "Chooser Table Editor");
 }
 
+void FChooserTableEditor::PostUndo(bool bSuccess)
+{
+	UpdateTableColumns();
+	UpdateTableRows();
+}
+	
+void FChooserTableEditor::PostRedo(bool bSuccess)
+{
+	UpdateTableColumns();
+	UpdateTableRows();
+}
+	
 FText FChooserTableEditor::GetToolkitName() const
 {
 	const TArray<UObject*>& EditingObjs = GetEditingObjects();
@@ -377,7 +391,7 @@ public:
 			{
 				// row drag handle
 				
-				return 		SNew(SChooserRowHandle).ChooserEditor(Editor).RowIndex(RowIndex->RowIndex);
+				return SNew(SChooserRowHandle).ChooserEditor(Editor).RowIndex(RowIndex->RowIndex);
 			}
 			else if (ColumnName == Result) 
 			{
@@ -385,7 +399,9 @@ public:
 				TSharedPtr<SWidget> ResultWidget = FObjectChooserWidgetFactories::CreateWidget(UObjectChooser::StaticClass(), RowValue,Chooser->ContextObjectType,
 				FOnClassPicked::CreateLambda([this, RowIndex=RowIndex->RowIndex](UClass* ChosenClass)
 				{
-					UObject* RowValue = NewObject<UObject>(Chooser, ChosenClass);
+					const FScopedTransaction Transaction(LOCTEXT("Change Row Result Type", "Change Row Result Type"));
+					UObject* RowValue = NewObject<UObject>(Chooser, ChosenClass, NAME_None, RF_Transactional);
+					Chooser->Modify(true);
 					Chooser->Results[RowIndex] = RowValue;
 					FObjectChooserWidgetFactories::CreateWidget(UObjectChooser::StaticClass(), RowValue, Chooser->ContextObjectType, FOnClassPicked(), &CacheBorder);
 				}),
@@ -419,31 +435,7 @@ public:
 			// on the row past the end, show an Add button in the result column
 			if (ColumnName == Result)
 			{
-				TSharedRef<SComboButton> CreateRowComboButton = SNew(SComboButton)
-									.ComboButtonStyle(FAppStyle::Get(), "SimpleComboButton")
-									.ButtonContent()
-									[
-										SNew(STextBlock).Text(LOCTEXT("AddRow", "+ Add Row"))
-									];
-
-				CreateRowComboButton->SetOnGetMenuContent(FOnGetContent::CreateLambda([this, CreateRowComboButton]()
-					{
-						FClassViewerInitializationOptions Options;
-						Options.ClassFilters.Add(MakeShared<FInterfaceClassFilter>(UObjectChooser::StaticClass()) );
-						
-						// Add class filter for columns here
-						TSharedRef<SWidget> Widget = FModuleManager::LoadModuleChecked<FClassViewerModule>("ClassViewer").CreateClassViewer(Options,
-							FOnClassPicked::CreateLambda([this,CreateRowComboButton](UClass* ChosenClass)
-							{
-								CreateRowComboButton->SetIsOpen(false);
-								UObject* ResultObject = NewObject<UObject>(Chooser, ChosenClass);
-								Chooser->Results.Add({ResultObject});
-								Editor->UpdateTableRows();
-							}));
-						return Widget;
-					}));
-
-				return CreateRowComboButton;
+				return Editor->GetCreateRowComboButton().ToSharedRef();
 			}
 		}
 		return SNullWidget::NullWidget;
@@ -527,7 +519,9 @@ void FChooserTableEditor::UpdateTableColumns()
 					FUIAction(
 						FExecuteAction::CreateLambda([this,ColumnIndex, &Column]()
 						{
+							const FScopedTransaction Transaction(LOCTEXT("Delete Column Transaction", "Delete Column"));
 							UChooserTable* Chooser = Cast<UChooserTable>(EditingObjects[0]);
+							Chooser->Modify(true);
 							Chooser->Columns.RemoveAt(ColumnIndex);
 							UpdateTableColumns();
 						})
@@ -546,7 +540,9 @@ void FChooserTableEditor::UpdateTableColumns()
 						// Add class filter for columns here
 						TSharedRef<SWidget> Widget = FModuleManager::LoadModuleChecked<FClassViewerModule>("ClassViewer").CreateClassViewer(Options, FOnClassPicked::CreateLambda([this, &Column](UClass* ChosenClass)
 						{
-							UObject* Value = NewObject<UObject>(Column.GetObject(), ChosenClass);
+							const FScopedTransaction Transaction(LOCTEXT("SetColumnInputType", "Set Column Input Type"));
+							UObject* Value = NewObject<UObject>(Column.GetObject(), ChosenClass, NAME_None, RF_Transactional);
+							Column.GetObject()->Modify(true);
 							Column->SetInputValue(Value);
 							UpdateTableColumns();
 							UpdateTableRows();
@@ -628,10 +624,10 @@ TSharedRef<SDockTab> FChooserTableEditor::SpawnTableTab( const FSpawnTabArgs& Ar
 		{
 			CreateColumnComboButton->SetIsOpen(false);
 			UChooserTable* Chooser = Cast<UChooserTable>(EditingObjects[0]);
-			UObject* ColumnObject = NewObject<UObject>(Chooser, ChosenClass);
-			TScriptInterface<IChooserColumn> Column;
-			Column = ColumnObject;
-			Chooser->Columns.Add(TScriptInterface<IChooserColumn>(Column));
+			const FScopedTransaction Transaction(LOCTEXT("Add Column Transaction", "Add Column"));
+			Chooser->Modify(true);
+			UObject* ColumnObject = NewObject<UObject>(Chooser, ChosenClass, NAME_None, RF_Transactional);
+			Chooser->Columns.Add(TScriptInterface<IChooserColumn>(ColumnObject));
 			UpdateTableColumns();
 			UpdateTableRows();
 			SelectedColumn = "ChooserColumn";
@@ -647,6 +643,32 @@ TSharedRef<SDockTab> FChooserTableEditor::SpawnTableTab( const FSpawnTabArgs& Ar
 	];
 
 
+	CreateRowComboButton = SNew(SComboButton)
+		.ComboButtonStyle(FAppStyle::Get(), "SimpleComboButton")
+		.ButtonContent()
+		[
+			SNew(STextBlock).Text(LOCTEXT("AddRow", "+ Add Row"))
+		]
+		.OnGetMenuContent_Lambda([this]()
+		{
+			FClassViewerInitializationOptions Options;
+			Options.ClassFilters.Add(MakeShared<FInterfaceClassFilter>(UObjectChooser::StaticClass()) );
+			
+			// Add class filter for columns here
+			TSharedRef<SWidget> Widget = FModuleManager::LoadModuleChecked<FClassViewerModule>("ClassViewer").CreateClassViewer(Options,
+				FOnClassPicked::CreateLambda([this](UClass* ChosenClass)
+				{
+					CreateRowComboButton->SetIsOpen(false);
+					UChooserTable* Chooser = Cast<UChooserTable>(EditingObjects[0]);
+					const FScopedTransaction Transaction(LOCTEXT("Add Row Transaction", "Add Row"));
+					Chooser->Modify(true);
+					UObject* ResultObject = NewObject<UObject>(Chooser, ChosenClass, NAME_None, RF_Transactional);
+					Chooser->Results.Add(TScriptInterface<IObjectChooser>(ResultObject));
+					UpdateTableRows();
+				}));
+			return Widget;
+		});
+
 	HeaderRow = SNew(SHeaderRow);
 
 	UpdateTableRows();
@@ -658,7 +680,9 @@ TSharedRef<SDockTab> FChooserTableEditor::SpawnTableTab( const FSpawnTabArgs& Ar
 				{
 					if (Event.GetKey() == EKeys::Delete)
 					{
+						const FScopedTransaction Transaction(LOCTEXT("Delete Row Transaction", "Delete Row"));
 						UChooserTable* Chooser = Cast<UChooserTable>(EditingObjects[0]);
+						Chooser->Modify(true);
 						// delete selected rows.
 						TArray<uint32> RowsToDelete;
 						for(auto& SelectedRow:SelectedRows)
@@ -901,7 +925,12 @@ TSharedRef<SWidget> CreateAssetWidget(UObject* Object, UClass* ContextClass)
 	return SNew(SObjectPropertyEntryBox)
 		.AllowedClass(Chooser->OutputObjectType!=nullptr ? Chooser->OutputObjectType.Get() : UObject::StaticClass())
 		.ObjectPath_Lambda([DIAsset](){ return DIAsset->Asset ? DIAsset->Asset.GetPath() : "";})
-		.OnObjectChanged_Lambda([DIAsset](const FAssetData& AssetData){ DIAsset->Asset = AssetData.GetAsset(); });
+		.OnObjectChanged_Lambda([DIAsset](const FAssetData& AssetData)
+		{
+			const FScopedTransaction Transaction(LOCTEXT("Edit Asset", "Edit Asset"));
+			DIAsset->Modify(true);
+			DIAsset->Asset = AssetData.GetAsset();
+		});
 }
 
 TSharedRef<SWidget> CreateEvaluateChooserWidget(UObject* Object, UClass* ContextObject)
@@ -957,6 +986,8 @@ TSharedRef<SWidget> CreateBoolColumnWidget(UObject* Column, int Row)
 	{
 		if (Row < BoolColumn->RowValues.Num())
 		{
+			const FScopedTransaction Transaction(LOCTEXT("Change Bool Value", "Change Bool Value"));
+			BoolColumn->Modify(true);
 			BoolColumn->RowValues[Row] = (State == ECheckBoxState::Checked);
 		}
 	})
@@ -1019,10 +1050,11 @@ TSharedRef<SWidget> CreatePropertyWidget(UObject* Object, UClass* ContextClass)
 			{
 				if (InBindingChain.Num() == 1)
 			 	{
+					const FScopedTransaction Transaction(LOCTEXT("Change Property Binding", "Change Property Binding"));
 					TArray<FString> PropertyPath; 
 					PropertyAccessEditor.MakeStringPath(InBindingChain, PropertyPath);
+					ContextProperty->Modify(true);
 					ContextProperty->PropertyName = FName(PropertyPath[0]);
-					// do this better so undo will work?
 				}
 			}
 		});
@@ -1074,6 +1106,8 @@ TSharedRef<SWidget> CreateFloatRangeColumnWidget(UObject* Column, int Row)
 		{
 			if (Row < FloatRangeColumn->RowValues.Num())
 			{
+				const FScopedTransaction Transaction(LOCTEXT("Edit Min Value", "Edit Min Value"));
+				FloatRangeColumn->Modify(true);
 				FloatRangeColumn->RowValues[Row].Min = NewValue;
 			}
 		})
@@ -1097,6 +1131,8 @@ TSharedRef<SWidget> CreateFloatRangeColumnWidget(UObject* Column, int Row)
 		{
 			if (Row < FloatRangeColumn->RowValues.Num())
 			{
+				const FScopedTransaction Transaction(LOCTEXT("Edit Max", "Edit Max Value"));
+				FloatRangeColumn->Modify(true);
 				FloatRangeColumn->RowValues[Row].Max = NewValue;
 			}
 		})
diff --git a/Engine/Plugins/Experimental/Chooser/Source/ChooserEditor/Private/ChooserTableEditor.h b/Engine/Plugins/Experimental/Chooser/Source/ChooserEditor/Private/ChooserTableEditor.h
--- a/Engine/Plugins/Experimental/Chooser/Source/ChooserEditor/Private/ChooserTableEditor.h
+++ b/Engine/Plugins/Experimental/Chooser/Source/ChooserEditor/Private/ChooserTableEditor.h
@@ -8,7 +8,7 @@
 #include "Editor/PropertyEditor/Public/PropertyEditorDelegates.h"
 #include "Widgets/Views/STableViewBase.h"
 #include "Chooser.h"
-#include "IDetailCustomization.h"
+#include "EditorUndoClient.h"
 #include "ChooserTableEditor.generated.h"
 
 class SEditableText;
@@ -28,7 +28,7 @@ public:
 
 namespace UE::ChooserEditor
 {
-	class FChooserTableEditor : public FAssetEditorToolkit
+	class FChooserTableEditor : public FAssetEditorToolkit, public FSelfRegisteringEditorUndoClient
 	{
 	public:
 		/** Delegate that, given an array of assets, returns an array of objects to use in the details view of an FSimpleAssetEditor */
@@ -60,6 +60,10 @@ namespace UE::ChooserEditor
 		virtual FLinearColor GetWorldCentricTabColorScale() const override;
 		virtual bool IsPrimaryEditor() const override { return true; }
 		virtual bool IsSimpleAssetEditor() const override { return false; }
+		
+		/** FEditorUndoClient Interface */
+		virtual void PostUndo(bool bSuccess) override;
+		virtual void PostRedo(bool bSuccess) override;
 
 		/** FAssetEditorToolkit interface */
 		virtual void PostRegenerateMenusAndToolbars() override;
@@ -116,11 +120,14 @@ namespace UE::ChooserEditor
 		TArray<TSharedPtr<FChooserTableRow>> TableRows;
 	
 		TSharedPtr<SComboButton> CreateColumnComboButton;
+		TSharedPtr<SComboButton> CreateRowComboButton;
+		
 		TSharedPtr<SHeaderRow> HeaderRow;
 		TSharedPtr<SListView<TSharedPtr<FChooserTableRow>>> TableView;
 
 		FName SelectedColumn;
 	public:
+		TSharedPtr<SComboButton>& GetCreateRowComboButton() { return CreateRowComboButton; };
 
 		static TMap<const UClass*, TFunction<TSharedRef<SWidget> (UObject* Column, int Row)>> ColumnWidgetCreators;
 	
diff --git a/Engine/Plugins/Experimental/Chooser/Source/ChooserEditor/Private/InterfacePropertyTypeCustomization.cpp b/Engine/Plugins/Experimental/Chooser/Source/ChooserEditor/Private/InterfacePropertyTypeCustomization.cpp
--- a/Engine/Plugins/Experimental/Chooser/Source/ChooserEditor/Private/InterfacePropertyTypeCustomization.cpp
+++ b/Engine/Plugins/Experimental/Chooser/Source/ChooserEditor/Private/InterfacePropertyTypeCustomization.cpp
@@ -51,7 +51,7 @@ void FInterfacePropertyTypeCustomization::CustomizeHeader(TSharedRef<IPropertyHa
 
 				for(int i=0;i<RawData.Num();i++)
 				{
-					UObject* NewValue = NewObject<UObject>(OuterObjects[i], ChosenClass);
+					UObject* NewValue = NewObject<UObject>(OuterObjects[i], ChosenClass, NAME_None, RF_Transactional);
 				
 					PropertyHandle->NotifyPreChange();
 					PropertyHandle->SetValue(NewValue);
