PropertyBinding changes for Choosers

diff --git a/Engine/Plugins/Experimental/Chooser/Source/ChooserEditor/Private/ChooserTableEditor.cpp b/Engine/Plugins/Experimental/Chooser/Source/ChooserEditor/Private/ChooserTableEditor.cpp
--- a/Engine/Plugins/Experimental/Chooser/Source/ChooserEditor/Private/ChooserTableEditor.cpp
+++ b/Engine/Plugins/Experimental/Chooser/Source/ChooserEditor/Private/ChooserTableEditor.cpp
@@ -1048,13 +1048,11 @@ TSharedRef<SWidget> CreatePropertyWidget(UObject* Object, UClass* ContextClass)
 			IPropertyAccessEditor& PropertyAccessEditor = IModularFeatures::Get().GetModularFeature<IPropertyAccessEditor>("PropertyAccessEditor");
 			if (ContextProperty.IsValid())
 			{
-				if (InBindingChain.Num() == 1)
+				if (InBindingChain.Num() == 2)
 			 	{
 					const FScopedTransaction Transaction(LOCTEXT("Change Property Binding", "Change Property Binding"));
-					TArray<FString> PropertyPath; 
-					PropertyAccessEditor.MakeStringPath(InBindingChain, PropertyPath);
 					ContextProperty->Modify(true);
-					ContextProperty->PropertyName = FName(PropertyPath[0]);
+					ContextProperty->PropertyName = FName(InBindingChain[1].Field.GetName());
 				}
 			}
 		});
@@ -1079,7 +1077,9 @@ TSharedRef<SWidget> CreatePropertyWidget(UObject* Object, UClass* ContextClass)
 		});
 
 	IPropertyAccessEditor& PropertyAccessEditor = IModularFeatures::Get().GetModularFeature<IPropertyAccessEditor>("PropertyAccessEditor");
-	return PropertyAccessEditor.MakePropertyBindingWidget((UBlueprint*)ContextClass, Args);
+	FBindingContextStruct StructInfo;
+	StructInfo.Struct = ContextClass;
+	return PropertyAccessEditor.MakePropertyBindingWidget({StructInfo}, Args);
 }
 
 TSharedRef<SWidget> CreateFloatRangeColumnWidget(UObject* Column, int Row)
