Chooser table UI fixes:
Refresh columns when changing context object type, so property access will show the correct properties.
Refresh rows when output object type is changed, so that the asset browsers will filter for the right object type

diff --git a/Engine/Plugins/Experimental/Chooser/Source/ChooserEditor/Private/ChooserTableEditor.cpp b/Engine/Plugins/Experimental/Chooser/Source/ChooserEditor/Private/ChooserTableEditor.cpp
--- a/Engine/Plugins/Experimental/Chooser/Source/ChooserEditor/Private/ChooserTableEditor.cpp
+++ b/Engine/Plugins/Experimental/Chooser/Source/ChooserEditor/Private/ChooserTableEditor.cpp
@@ -101,6 +101,7 @@ void FChooserTableEditor::InitEditor( const EToolkitMode::Type Mode, const TShar
 
 	FPropertyEditorModule& PropertyEditorModule = FModuleManager::GetModuleChecked<FPropertyEditorModule>( "PropertyEditor" );
 	FDetailsViewArgs DetailsViewArgs;
+	DetailsViewArgs.NotifyHook = this;
 	DetailsViewArgs.NameAreaSettings = FDetailsViewArgs::HideNameArea;
 	DetailsView = PropertyEditorModule.CreateDetailView( DetailsViewArgs );
 	const TSharedRef<FTabManager::FLayout> StandaloneDefaultLayout = FTabManager::NewLayout( "Standalone_ChooserTableEditor_Layout_v1" )
@@ -157,6 +158,26 @@ void FChooserTableEditor::PostRedo(bool bSuccess)
 	UpdateTableColumns();
 	UpdateTableRows();
 }
+
+	
+void FChooserTableEditor::NotifyPreChange(FProperty* PropertyAboutToChange)
+{
+}
+
+void FChooserTableEditor::NotifyPostChange(const FPropertyChangedEvent& PropertyChangedEvent, FProperty* PropertyThatChanged)
+{
+	if (PropertyThatChanged->GetNameCPP() == GET_MEMBER_NAME_STRING_CHECKED(UChooserTable, OutputObjectType))
+	{
+		// rebuild all result widgets
+		UpdateTableRows();
+	}
+	if (PropertyThatChanged->GetNameCPP() == GET_MEMBER_NAME_STRING_CHECKED(UChooserTable, ContextObjectType))
+	{
+		// rebuild all column header widgets
+		UpdateTableColumns();
+	}
+}
+
 	
 FText FChooserTableEditor::GetToolkitName() const
 {
@@ -1012,7 +1033,7 @@ template <typename PropertyType>
 TSharedRef<SWidget> CreatePropertyWidget(UObject* Object, UClass* ContextClass)
 {
 	TWeakObjectPtr<PropertyType> ContextProperty = Cast<PropertyType>(Object);
-
+	
 	FPropertyBindingWidgetArgs Args;
 	Args.bAllowPropertyBindings = true;
 	Args.OnCanBindProperty = FOnCanBindProperty::CreateLambda([](FProperty* Property)
diff --git a/Engine/Plugins/Experimental/Chooser/Source/ChooserEditor/Private/ChooserTableEditor.h b/Engine/Plugins/Experimental/Chooser/Source/ChooserEditor/Private/ChooserTableEditor.h
--- a/Engine/Plugins/Experimental/Chooser/Source/ChooserEditor/Private/ChooserTableEditor.h
+++ b/Engine/Plugins/Experimental/Chooser/Source/ChooserEditor/Private/ChooserTableEditor.h
@@ -10,6 +10,7 @@
 #include "Widgets/Views/STableViewBase.h"
 #include "Chooser.h"
 #include "EditorUndoClient.h"
+#include "Misc/NotifyHook.h"
 #include "ChooserTableEditor.generated.h"
 
 class SComboButton;
@@ -29,7 +30,7 @@ public:
 
 namespace UE::ChooserEditor
 {
-	class FChooserTableEditor : public FAssetEditorToolkit, public FSelfRegisteringEditorUndoClient
+	class FChooserTableEditor : public FAssetEditorToolkit, public FSelfRegisteringEditorUndoClient, public FNotifyHook
 	{
 	public:
 		/** Delegate that, given an array of assets, returns an array of objects to use in the details view of an FSimpleAssetEditor */
@@ -65,6 +66,10 @@ namespace UE::ChooserEditor
 		/** FEditorUndoClient Interface */
 		virtual void PostUndo(bool bSuccess) override;
 		virtual void PostRedo(bool bSuccess) override;
+		
+		/** Begin FNotifyHook Interface */
+		virtual void NotifyPreChange( FProperty* PropertyAboutToChange ) override;
+		virtual void NotifyPostChange( const FPropertyChangedEvent& PropertyChangedEvent, FProperty* PropertyThatChanged) override;
 
 		/** FAssetEditorToolkit interface */
 		virtual void PostRegenerateMenusAndToolbars() override;
