support for getting multiple results from a chooser

diff --git a/Engine/Plugins/Experimental/Chooser/Source/Chooser/Internal/Chooser.h b/Engine/Plugins/Experimental/Chooser/Source/Chooser/Internal/Chooser.h
--- a/Engine/Plugins/Experimental/Chooser/Source/Chooser/Internal/Chooser.h
+++ b/Engine/Plugins/Experimental/Chooser/Source/Chooser/Internal/Chooser.h
@@ -132,10 +132,10 @@ public:
 	UPROPERTY(EditAnywhere, Meta = (EditInlineInterface = "true"), Category="Hidden")
 	TArray<TScriptInterface<IChooserColumn>> Columns;
 
-	UPROPERTY(EditAnywhere, Category="Input")
+	UPROPERTY(EditAnywhere, Category="Input", Meta = (AllowAbstract=true))
 	TObjectPtr<UClass> ContextObjectType;
 	
-	UPROPERTY(EditAnywhere, Category="Output")
+	UPROPERTY(EditAnywhere, Category="Output", Meta = (AllowAbstract=true))
 	TObjectPtr<UClass> OutputObjectType;
 };
 
@@ -145,7 +145,7 @@ class CHOOSER_API UObjectChooser_EvaluateChooser : public UObject, public IObjec
 	GENERATED_BODY()
 
 	virtual UObject* ChooseObject(const UObject* ContextObject) const final override;
-
+	virtual EIteratorStatus ChooseMulti(const UObject* ContextObject, FObjectChooserIteratorCallback Callback) const final override;
 	public:
 	
 	UPROPERTY(EditAnywhere, Category="Parameters")
diff --git a/Engine/Plugins/Experimental/Chooser/Source/Chooser/Private/Chooser.cpp b/Engine/Plugins/Experimental/Chooser/Source/Chooser/Private/Chooser.cpp
--- a/Engine/Plugins/Experimental/Chooser/Source/Chooser/Private/Chooser.cpp
+++ b/Engine/Plugins/Experimental/Chooser/Source/Chooser/Private/Chooser.cpp
@@ -91,11 +91,11 @@ void UChooserColumnFloatRange::Filter(const UObject* ContextObject, const TArray
 	}
 }
 
-static UObject* StaticEvaluateChooser(const UObject* ContextObject, const UChooserTable* Chooser)
+static IObjectChooser::EIteratorStatus StaticEvaluateChooser(const UObject* ContextObject, const UChooserTable* Chooser, IObjectChooser::FObjectChooserIteratorCallback Callback)
 {
 	if (Chooser == nullptr)
 	{
-		return nullptr;
+		return IObjectChooser::EIteratorStatus::Continue;
 	}
 
 	TArray<uint32> Indices1;
@@ -125,21 +125,34 @@ static UObject* StaticEvaluateChooser(const UObject* ContextObject, const UChoos
 	{
 		if (Chooser->Results.Num() > (int32)SelectedIndex)
 		{
-			const IObjectChooser* SelectedResult = Chooser->Results[SelectedIndex].GetInterface();
-
-			if(UObject* Result = SelectedResult->ChooseObject(ContextObject))
+			if (const IObjectChooser* SelectedResult = Chooser->Results[SelectedIndex].GetInterface())
 			{
-				return Result;
+				if (SelectedResult->ChooseMulti(ContextObject, Callback) == IObjectChooser::EIteratorStatus::Stop)
+				{
+					return IObjectChooser::EIteratorStatus::Stop;
+				}
 			}
 		}
 	}
-
-	return nullptr;
+	
+	return IObjectChooser::EIteratorStatus::Continue;
 }
 
 UObject* UObjectChooser_EvaluateChooser::ChooseObject(const UObject* ContextObject) const
 {
-	return StaticEvaluateChooser(ContextObject, Chooser);
+	UObject* Result = nullptr;
+	StaticEvaluateChooser(ContextObject, Chooser, FObjectChooserIteratorCallback::CreateLambda([&Result](UObject* InResult)
+	{
+		Result = InResult;
+		return IObjectChooser::EIteratorStatus::Stop;
+	}));
+
+	return Result;
+}
+
+IObjectChooser::EIteratorStatus UObjectChooser_EvaluateChooser::ChooseMulti(const UObject* ContextObject, FObjectChooserIteratorCallback Callback) const
+{
+	return StaticEvaluateChooser(ContextObject, Chooser, Callback);
 }
 
 UChooserFunctionLibrary::UChooserFunctionLibrary(const FObjectInitializer& ObjectInitializer)
@@ -149,5 +162,14 @@ UChooserFunctionLibrary::UChooserFunctionLibrary(const FObjectInitializer& Objec
 
 UObject* UChooserFunctionLibrary::EvaluateChooser(const UObject* ContextObject, const UChooserTable* Chooser)
 {
-	return StaticEvaluateChooser(ContextObject, Chooser);
+	UObject* Result = nullptr;
+	StaticEvaluateChooser(ContextObject, Chooser, IObjectChooser::FObjectChooserIteratorCallback::CreateLambda([&Result](UObject* InResult)
+	{
+		Result = InResult;
+		return IObjectChooser::EIteratorStatus::Stop;
+	}));
+
+	return Result;
+
+	return Result;
 }
diff --git a/Engine/Plugins/Experimental/Chooser/Source/Chooser/Public/IObjectChooser.h b/Engine/Plugins/Experimental/Chooser/Source/Chooser/Public/IObjectChooser.h
--- a/Engine/Plugins/Experimental/Chooser/Source/Chooser/Public/IObjectChooser.h
+++ b/Engine/Plugins/Experimental/Chooser/Source/Chooser/Public/IObjectChooser.h
@@ -20,4 +20,17 @@ class CHOOSER_API IObjectChooser
 public:
 
 	virtual UObject* ChooseObject(const UObject* ContextObject) const { return nullptr; };
+
+	enum class EIteratorStatus { Continue,Stop };
+
+	DECLARE_DELEGATE_RetVal_OneParam( EIteratorStatus, FObjectChooserIteratorCallback, UObject*);
+	virtual EIteratorStatus ChooseMulti(const UObject* ContextObject, FObjectChooserIteratorCallback Callback) const
+	{
+		// fallback implementation just calls the single version
+		if (UObject* Result = ChooseObject(ContextObject))
+		{
+			return Callback.Execute(Result);
+		}
+		return EIteratorStatus::Continue;
+	};
 };
diff --git a/Engine/Plugins/Experimental/Chooser/Source/ChooserEditor/ChooserEditor.Build.cs b/Engine/Plugins/Experimental/Chooser/Source/ChooserEditor/ChooserEditor.Build.cs
--- a/Engine/Plugins/Experimental/Chooser/Source/ChooserEditor/ChooserEditor.Build.cs
+++ b/Engine/Plugins/Experimental/Chooser/Source/ChooserEditor/ChooserEditor.Build.cs
@@ -21,7 +21,6 @@ namespace UnrealBuildTool.Rules
 					"PropertyEditor",
 					"InputCore",
 					"EditorStyle",
-					"PropertyAccessEditor",
 					"PropertyEditor",
 					"BlueprintGraph"
 					// ... add private dependencies that you statically link with here ...
