better property binding for choosers

diff --git a/Engine/Plugins/Experimental/Chooser/Source/Chooser/Internal/Chooser.h b/Engine/Plugins/Experimental/Chooser/Source/Chooser/Internal/Chooser.h
--- a/Engine/Plugins/Experimental/Chooser/Source/Chooser/Internal/Chooser.h
+++ b/Engine/Plugins/Experimental/Chooser/Source/Chooser/Internal/Chooser.h
@@ -14,16 +14,19 @@ class CHOOSER_API UChooserParameterBool_ContextProperty :  public UObject, publi
 {
 	GENERATED_BODY()
 public:
+
 	UPROPERTY()
- 	FName PropertyName;
+	TArray<FName> PropertyBindingChain;
 	
-	virtual bool GetValue(const UObject* ContextObject, bool& OutResult) override;
+	virtual bool GetValue(const UObject* ContextObject, bool& OutResult) const override;
 
-	static const FString& CPPTypeName()
+#if WITH_EDITOR
+	static bool CanBind(const FString& TypeName)
 	{
-		static FString TypeName = "bool";
-		return TypeName;
+		static FString BoolTypeName = "bool";
+		return TypeName == BoolTypeName;
 	}
+#endif
 };
 
 UCLASS()
@@ -32,15 +35,18 @@ class CHOOSER_API UChooserParameterFloat_ContextProperty :  public UObject, publ
 	GENERATED_BODY()
 public:
 	UPROPERTY()
-	FName PropertyName;
-	
-	virtual bool GetValue(const UObject* ContextObject, float& OutResult) override;
+	TArray<FName> PropertyBindingChain;
 	
-	static const FString& CPPTypeName()
+	virtual bool GetValue(const UObject* ContextObject, float& OutResult) const override;
+
+#if WITH_EDITOR
+	static bool CanBind(const FString& TypeName)
 	{
-		static FString TypeName = "double";
-		return TypeName;
+		static FString DoubleTypeName = "double";
+		static FString FloatTypeName = "float";
+		return TypeName == FloatTypeName || TypeName == DoubleTypeName;
 	}
+#endif
 };
 
 UCLASS()
diff --git a/Engine/Plugins/Experimental/Chooser/Source/Chooser/Private/Chooser.cpp b/Engine/Plugins/Experimental/Chooser/Source/Chooser/Private/Chooser.cpp
--- a/Engine/Plugins/Experimental/Chooser/Source/Chooser/Private/Chooser.cpp
+++ b/Engine/Plugins/Experimental/Chooser/Source/Chooser/Private/Chooser.cpp
@@ -2,23 +2,76 @@
 #include "Chooser.h"
 #include "ChooserFunctionLibrary.h"
 
-bool UChooserParameterBool_ContextProperty::GetValue(const UObject* ContextObject, bool& OutResult)
+bool ResolvePropertyChain(const void*& Container, UStruct*& StructType, const TArray<FName>& PropertyBindingChain)
 {
-	if (const FBoolProperty* Property = FindFProperty<FBoolProperty>(ContextObject->GetClass(), PropertyName))
+	if (PropertyBindingChain.IsEmpty())
 	{
-		OutResult = *Property->ContainerPtrToValuePtr<bool>(ContextObject);
+		return false;
 	}
 	
-	return false;
+	const int PropertyChainLength = PropertyBindingChain.Num();
+	for(int PropertyChainIndex = 0; PropertyChainIndex < PropertyChainLength - 1; PropertyChainIndex++)
+	{
+		if (const FStructProperty* StructProperty = FindFProperty<FStructProperty>(StructType, PropertyBindingChain[PropertyChainIndex]))
+		{
+			StructType = StructProperty->Struct;
+			Container = StructProperty->ContainerPtrToValuePtr<void>(Container);
+		}
+		else if (const FObjectProperty* ObjectProperty = FindFProperty<FObjectProperty>(StructType, PropertyBindingChain[PropertyChainIndex]))
+		{
+			StructType = ObjectProperty->PropertyClass;
+			Container = *ObjectProperty->ContainerPtrToValuePtr<TObjectPtr<UObject>>(Container);
+			if (Container == nullptr)
+			{
+				return false;
+			}
+		}
+		else
+		{
+			return false;
+		}
+	}
+	
+	return true;
 }
 
-bool UChooserParameterFloat_ContextProperty::GetValue(const UObject* ContextObject, float& OutResult)
+bool UChooserParameterBool_ContextProperty::GetValue(const UObject* ContextObject, bool& OutResult) const
 {
-	if (const FDoubleProperty* Property = FindFProperty<FDoubleProperty>(ContextObject->GetClass(), PropertyName))
+	UStruct* StructType = ContextObject->GetClass();
+	const void* Container = ContextObject;
+	
+	if (ResolvePropertyChain(Container, StructType, PropertyBindingChain))
 	{
-		OutResult = *Property->ContainerPtrToValuePtr<double>(ContextObject);
+		if (const FBoolProperty* Property = FindFProperty<FBoolProperty>(StructType, PropertyBindingChain.Last()))
+		{
+			OutResult = *Property->ContainerPtrToValuePtr<bool>(Container);
+			return true;
+		}
 	}
+
+	return false;
+}
+
+bool UChooserParameterFloat_ContextProperty::GetValue(const UObject* ContextObject, float& OutResult) const
+{
+	UStruct* StructType = ContextObject->GetClass();
+	const void* Container = ContextObject;
 	
+	if (ResolvePropertyChain(Container, StructType, PropertyBindingChain))
+	{
+		if (const FDoubleProperty* DoubleProperty = FindFProperty<FDoubleProperty>(StructType, PropertyBindingChain.Last()))
+		{
+			OutResult = *DoubleProperty->ContainerPtrToValuePtr<double>(Container);
+			return true;
+		}
+		
+		if (const FFloatProperty* FloatProperty = FindFProperty<FFloatProperty>(StructType, PropertyBindingChain.Last()))
+		{
+			OutResult = *FloatProperty->ContainerPtrToValuePtr<float>(Container);
+			return true;
+		}
+	}
+
 	return false;
 }
 
diff --git a/Engine/Plugins/Experimental/Chooser/Source/Chooser/Public/IChooserColumnParameter.h b/Engine/Plugins/Experimental/Chooser/Source/Chooser/Public/IChooserColumnParameter.h
--- a/Engine/Plugins/Experimental/Chooser/Source/Chooser/Public/IChooserColumnParameter.h
+++ b/Engine/Plugins/Experimental/Chooser/Source/Chooser/Public/IChooserColumnParameter.h
@@ -18,7 +18,7 @@ class CHOOSER_API IChooserParameterBool
 	GENERATED_BODY()
 
 public:
-	virtual bool GetValue(const UObject* ContextObject, bool& OutResult) { return false; }
+	virtual bool GetValue(const UObject* ContextObject, bool& OutResult) const { return false; }
 };
 
 UINTERFACE(NotBlueprintType, meta = (CannotImplementInterfaceInBlueprint))
@@ -32,5 +32,5 @@ class CHOOSER_API IChooserParameterFloat
 	GENERATED_BODY()
 
 public:
-	virtual bool GetValue(const UObject* ContextObject, float& OutResult) { return false; }
+	virtual bool GetValue(const UObject* ContextObject, float& OutResult) const { return false; }
 };
\ No newline at end of file
diff --git a/Engine/Plugins/Experimental/Chooser/Source/ChooserEditor/ChooserEditor.Build.cs b/Engine/Plugins/Experimental/Chooser/Source/ChooserEditor/ChooserEditor.Build.cs
--- a/Engine/Plugins/Experimental/Chooser/Source/ChooserEditor/ChooserEditor.Build.cs
+++ b/Engine/Plugins/Experimental/Chooser/Source/ChooserEditor/ChooserEditor.Build.cs
@@ -22,7 +22,8 @@ namespace UnrealBuildTool.Rules
 					"InputCore",
 					"EditorStyle",
 					"PropertyEditor",
-					"BlueprintGraph"
+					"BlueprintGraph",
+					"GraphEditor"
 					// ... add private dependencies that you statically link with here ...
 				}
 			);
diff --git a/Engine/Plugins/Experimental/Chooser/Source/ChooserEditor/Private/ChooserTableEditor.cpp b/Engine/Plugins/Experimental/Chooser/Source/ChooserEditor/Private/ChooserTableEditor.cpp
--- a/Engine/Plugins/Experimental/Chooser/Source/ChooserEditor/Private/ChooserTableEditor.cpp
+++ b/Engine/Plugins/Experimental/Chooser/Source/ChooserEditor/Private/ChooserTableEditor.cpp
@@ -28,7 +28,7 @@
 #include "ObjectChooser_Asset.h"
 #include "ObjectChooserClassFilter.h"
 #include "ObjectChooserWidgetFactories.h"
-#include "EdGraphSchema_K2.h"
+#include "GraphEditorSettings.h"
 #include "IDetailCustomization.h"
 #include "PropertyCustomizationHelpers.h"
 #include "ScopedTransaction.h"
@@ -448,7 +448,11 @@ public:
 						}
 						ColumnClass = ColumnClass->GetSuperClass();
 					}
-					return ColumnWidget.ToSharedRef();
+					
+					if (ColumnWidget.IsValid())
+					{
+						return ColumnWidget.ToSharedRef();
+					}
 				}
 			}
 		}
@@ -1025,10 +1029,25 @@ void ConvertToText_ContextProperty(const UObject* Object, FText& OutText)
 {
 	if (const PropertyType* ContextProperty = Cast<PropertyType>(Object))
 	{
-		OutText = FText::FromName(ContextProperty->PropertyName);
+		if (ContextProperty->PropertyBindingChain.Num()>0)
+		{
+			OutText = FText::FromName(ContextProperty->PropertyBindingChain.Last());
+		}
 	}
 }
 
+FLinearColor GetBindingColor(UChooserParameterBool_ContextProperty* Property)
+{
+	return GetDefault<UGraphEditorSettings>()->BooleanPinTypeColor;
+}
+	
+FLinearColor GetBindingColor(UChooserParameterFloat_ContextProperty* Property)
+{
+	return GetDefault<UGraphEditorSettings>()->DoublePinTypeColor;
+}
+	
+const UGraphEditorSettings* Settings = GetDefault<UGraphEditorSettings>();
+
 template <typename PropertyType>
 TSharedRef<SWidget> CreatePropertyWidget(UObject* Object, UClass* ContextClass)
 {
@@ -1036,26 +1055,21 @@ TSharedRef<SWidget> CreatePropertyWidget(UObject* Object, UClass* ContextClass)
 	
 	FPropertyBindingWidgetArgs Args;
 	Args.bAllowPropertyBindings = true;
+	
 	Args.OnCanBindProperty = FOnCanBindProperty::CreateLambda([](FProperty* Property)
 	{
-		return Property == nullptr || Property->GetCPPType() == PropertyType::CPPTypeName();
+		return Property == nullptr || PropertyType::CanBind(Property->GetCPPType());
 	});
 	Args.OnCanBindToClass = FOnCanBindToClass::CreateLambda([ContextClass ](UClass* InClass)
 	{
-		return ContextClass->IsChildOf(InClass);
+		return true;
 	});
 	
 	Args.CurrentBindingColor = MakeAttributeLambda([ContextProperty, ContextClass]() {
-		if (ContextProperty.IsValid())
-		{
-			if (const FProperty* Property = FindFProperty<FProperty>(ContextClass, ContextProperty->PropertyName))
-			{
-				const UEdGraphSchema_K2* Schema = GetDefault<UEdGraphSchema_K2>();
-				FEdGraphPinType VariablePinType;
-				Schema->ConvertPropertyToPinType(Property, VariablePinType);
-				return Schema->GetPinTypeColor(VariablePinType);
-			}
-		}
+		 if (ContextProperty.IsValid() && ContextProperty->PropertyBindingChain.Num() > 0)
+		 {
+		 	return GetBindingColor(ContextProperty.Get());
+		 }
 		return FLinearColor::Gray;
 	});
 
@@ -1070,15 +1084,46 @@ TSharedRef<SWidget> CreatePropertyWidget(UObject* Object, UClass* ContextClass)
 			IPropertyAccessEditor& PropertyAccessEditor = IModularFeatures::Get().GetModularFeature<IPropertyAccessEditor>("PropertyAccessEditor");
 			if (ContextProperty.IsValid())
 			{
-				if (InBindingChain.Num() == 2)
-			 	{
-					const FScopedTransaction Transaction(LOCTEXT("Change Property Binding", "Change Property Binding"));
-					ContextProperty->Modify(true);
-					ContextProperty->PropertyName = FName(InBindingChain[1].Field.GetName());
+				const FScopedTransaction Transaction(LOCTEXT("Change Property Binding", "Change Property Binding"));
+				ContextProperty->Modify(true);
+				
+				ContextProperty->PropertyBindingChain.Empty();
+
+				if (InBindingChain.Num() > 1)
+				{
+					for(int i = 1; i<InBindingChain.Num(); i++)
+					{
+						ContextProperty->PropertyBindingChain.Add(InBindingChain[i].Field.GetFName());
+					}
+				}
+				else
+				{
+					ContextProperty->PropertyBindingChain.Empty();
 				}
 			}
 		});
 
+	Args.CurrentBindingToolTipText = MakeAttributeLambda([ContextProperty]()
+			{
+				const FText Bind = LOCTEXT("Bind", "Bind");
+				FText CurrentValue = Bind;
+
+				if (ContextProperty.IsValid() && ContextProperty->PropertyBindingChain.Num()>0)
+				{
+					TArray<FText> BindingChainText;
+                 	BindingChainText.Reserve(ContextProperty->PropertyBindingChain.Num());
+                 
+                 	for (const FName& Name : ContextProperty->PropertyBindingChain)
+                 	{
+                 		BindingChainText.Add(FText::FromName(Name));
+                 	}
+					
+					CurrentValue = FText::Join(LOCTEXT("PropertyPathSeparator","."), BindingChainText);
+				}
+
+				return CurrentValue;	
+			});
+	
 	Args.CurrentBindingText = MakeAttributeLambda([ContextProperty]()
 			{
 				const FText Bind = LOCTEXT("Bind", "Bind");
@@ -1086,7 +1131,24 @@ TSharedRef<SWidget> CreatePropertyWidget(UObject* Object, UClass* ContextClass)
 
 				if (ContextProperty.IsValid())
 				{
-					CurrentValue = FText::FromName(ContextProperty->PropertyName);
+					int BindingChainLength = ContextProperty->PropertyBindingChain.Num();
+					if (BindingChainLength > 0)
+					{
+						if (BindingChainLength == 1)
+						{
+							// single property, just use the property name
+							CurrentValue = FText::FromName(ContextProperty->PropertyBindingChain.Last());
+						}
+						else
+						{
+							// for longer chains always show the last struct/object name, and the final property name (full path in tooltip)
+							CurrentValue = FText::Join(LOCTEXT("PropertyPathSeparator","."),
+								TArray<FText>({
+									FText::FromName(ContextProperty->PropertyBindingChain[BindingChainLength-2]),
+									FText::FromName(ContextProperty->PropertyBindingChain[BindingChainLength-1])
+								}));
+						}
+					}
 				}
 
 				return CurrentValue;
