Add GameplayTag column to chooser

diff --git a/Engine/Plugins/Experimental/Chooser/Chooser.uplugin b/Engine/Plugins/Experimental/Chooser/Chooser.uplugin
--- a/Engine/Plugins/Experimental/Chooser/Chooser.uplugin
+++ b/Engine/Plugins/Experimental/Chooser/Chooser.uplugin
@@ -16,6 +16,10 @@
 	"IsExperimentalVersion" : true,
 	"Installed" : false,
 	"Plugins": [
+		{
+			"Name": "GameplayTagsEditor",
+			"Enabled": true
+		}
 	],
 	"Modules" :
 	[
diff --git a/Engine/Plugins/Experimental/Chooser/Source/Chooser/Chooser.Build.cs b/Engine/Plugins/Experimental/Chooser/Source/Chooser/Chooser.Build.cs
--- a/Engine/Plugins/Experimental/Chooser/Source/Chooser/Chooser.Build.cs
+++ b/Engine/Plugins/Experimental/Chooser/Source/Chooser/Chooser.Build.cs
@@ -19,6 +19,7 @@ namespace UnrealBuildTool.Rules
 					"Core",
 					"CoreUObject",
 					"Engine",
+					"GameplayTags"
 					// ... add private dependencies that you statically link with here ...
 				}
 			);
diff --git a/Engine/Plugins/Experimental/Chooser/Source/Chooser/Internal/GameplayTagColumn.h b/Engine/Plugins/Experimental/Chooser/Source/Chooser/Internal/GameplayTagColumn.h
new file mode 100644
--- /dev/null
+++ b/Engine/Plugins/Experimental/Chooser/Source/Chooser/Internal/GameplayTagColumn.h
@@ -0,0 +1,63 @@
+// Copyright Epic Games, Inc. All Rights Reserved.
+#pragma once
+
+#include "CoreMinimal.h"
+#include "IChooserColumn.h"
+#include "IChooserParameterGameplayTag.h"
+#include "GameplayTagContainer.h"
+#include "GameplayTagColumn.generated.h"
+
+
+UCLASS()
+class CHOOSER_API UChooserParameterGameplayTag_ContextProperty :  public UObject, public IChooserParameterGameplayTag
+{
+	GENERATED_BODY()
+public:
+	UPROPERTY()
+	TArray<FName> PropertyBindingChain;
+	
+	virtual bool GetValue(const UObject* ContextObject, const FGameplayTagContainer*& OutResult) const override;
+
+#if WITH_EDITOR
+	static bool CanBind(const FString& TypeName)
+	{
+		static FString GameplayTagTypeName = "FGameplayTagContainer";
+		return TypeName == GameplayTagTypeName;
+	}
+#endif
+};
+
+UCLASS()
+class CHOOSER_API UChooserColumnGameplayTag : public UObject, public IChooserColumn
+{
+	GENERATED_BODY()
+	public:
+	UChooserColumnGameplayTag() {};
+	UChooserColumnGameplayTag(const FObjectInitializer& ObjectInitializer);
+	
+	UPROPERTY(EditAnywhere, Meta = (EditInlineInterface = "true"), Category = "Input")
+	TScriptInterface<IChooserParameterGameplayTag> InputValue;
+
+	UPROPERTY(EditAnywhere, Category=Runtime)
+	EGameplayContainerMatchType	TagMatchType = EGameplayContainerMatchType::Any;
+
+	UPROPERTY(EditAnywhere, Category=Runtime)
+	// array of results (cells for this column for each row in the table)
+	// should match the length of the Results array 
+	TArray<FGameplayTagContainer> RowValues;
+	
+	virtual void Filter(const UObject* ContextObject, const TArray<uint32>& IndexListIn, TArray<uint32>& IndexListOut) override;
+
+	// todo: macro boilerplate
+	virtual void SetNumRows(uint32 NumRows) override { RowValues.SetNum(NumRows); }
+	virtual void DeleteRows(const TArray<uint32> & RowIndices )
+	{
+		for(uint32 Index : RowIndices)
+		{
+			RowValues.RemoveAt(Index);
+		}
+	}
+	virtual UClass* GetInputValueInterface() override { return UChooserParameterGameplayTag::StaticClass(); };
+	virtual UObject* GetInputValue() override { return InputValue.GetObject(); };
+	virtual void SetInputValue(UObject* Value) override { InputValue = Value; };
+};
\ No newline at end of file
diff --git a/Engine/Plugins/Experimental/Chooser/Source/Chooser/Private/Chooser.cpp b/Engine/Plugins/Experimental/Chooser/Source/Chooser/Private/Chooser.cpp
--- a/Engine/Plugins/Experimental/Chooser/Source/Chooser/Private/Chooser.cpp
+++ b/Engine/Plugins/Experimental/Chooser/Source/Chooser/Private/Chooser.cpp
@@ -1,46 +1,14 @@
 // Copyright Epic Games, Inc. All Rights Reserved.
 #include "Chooser.h"
 #include "ChooserFunctionLibrary.h"
-
-bool ResolvePropertyChain(const void*& Container, UStruct*& StructType, const TArray<FName>& PropertyBindingChain)
-{
-	if (PropertyBindingChain.IsEmpty())
-	{
-		return false;
-	}
-	
-	const int PropertyChainLength = PropertyBindingChain.Num();
-	for(int PropertyChainIndex = 0; PropertyChainIndex < PropertyChainLength - 1; PropertyChainIndex++)
-	{
-		if (const FStructProperty* StructProperty = FindFProperty<FStructProperty>(StructType, PropertyBindingChain[PropertyChainIndex]))
-		{
-			StructType = StructProperty->Struct;
-			Container = StructProperty->ContainerPtrToValuePtr<void>(Container);
-		}
-		else if (const FObjectProperty* ObjectProperty = FindFProperty<FObjectProperty>(StructType, PropertyBindingChain[PropertyChainIndex]))
-		{
-			StructType = ObjectProperty->PropertyClass;
-			Container = *ObjectProperty->ContainerPtrToValuePtr<TObjectPtr<UObject>>(Container);
-			if (Container == nullptr)
-			{
-				return false;
-			}
-		}
-		else
-		{
-			return false;
-		}
-	}
-	
-	return true;
-}
+#include "ChooserPropertyAccess.h"
 
 bool UChooserParameterBool_ContextProperty::GetValue(const UObject* ContextObject, bool& OutResult) const
 {
 	UStruct* StructType = ContextObject->GetClass();
 	const void* Container = ContextObject;
 	
-	if (ResolvePropertyChain(Container, StructType, PropertyBindingChain))
+	if (UE::Chooser::ResolvePropertyChain(Container, StructType, PropertyBindingChain))
 	{
 		if (const FBoolProperty* Property = FindFProperty<FBoolProperty>(StructType, PropertyBindingChain.Last()))
 		{
@@ -57,7 +25,7 @@ bool UChooserParameterFloat_ContextProperty::GetValue(const UObject* ContextObje
 	UStruct* StructType = ContextObject->GetClass();
 	const void* Container = ContextObject;
 	
-	if (ResolvePropertyChain(Container, StructType, PropertyBindingChain))
+	if (UE::Chooser::ResolvePropertyChain(Container, StructType, PropertyBindingChain))
 	{
 		if (const FDoubleProperty* DoubleProperty = FindFProperty<FDoubleProperty>(StructType, PropertyBindingChain.Last()))
 		{
diff --git a/Engine/Plugins/Experimental/Chooser/Source/Chooser/Private/ChooserPropertyAccess.cpp b/Engine/Plugins/Experimental/Chooser/Source/Chooser/Private/ChooserPropertyAccess.cpp
new file mode 100644
--- /dev/null
+++ b/Engine/Plugins/Experimental/Chooser/Source/Chooser/Private/ChooserPropertyAccess.cpp
@@ -0,0 +1,40 @@
+// Copyright Epic Games, Inc. All Rights Reserved.
+
+#include "ChooserPropertyAccess.h"
+#include "UObject/UnrealType.h"
+
+namespace UE::Chooser
+{
+	bool ResolvePropertyChain(const void*& Container, UStruct*& StructType, const TArray<FName>& PropertyBindingChain)
+	{
+		if (PropertyBindingChain.Num() == 0)
+		{
+			return false;
+		}
+	
+		const int PropertyChainLength = PropertyBindingChain.Num();
+		for(int PropertyChainIndex = 0; PropertyChainIndex < PropertyChainLength - 1; PropertyChainIndex++)
+		{
+			if (const FStructProperty* StructProperty = FindFProperty<FStructProperty>(StructType, PropertyBindingChain[PropertyChainIndex]))
+			{
+				StructType = StructProperty->Struct;
+				Container = StructProperty->ContainerPtrToValuePtr<void>(Container);
+			}
+			else if (const FObjectProperty* ObjectProperty = FindFProperty<FObjectProperty>(StructType, PropertyBindingChain[PropertyChainIndex]))
+			{
+				StructType = ObjectProperty->PropertyClass;
+				Container = *ObjectProperty->ContainerPtrToValuePtr<TObjectPtr<UObject>>(Container);
+				if (Container == nullptr)
+				{
+					return false;
+				}
+			}
+			else
+			{
+				return false;
+			}
+		}
+	
+		return true;
+	}
+}
\ No newline at end of file
diff --git a/Engine/Plugins/Experimental/Chooser/Source/Chooser/Private/GameplayTagColumn.cpp b/Engine/Plugins/Experimental/Chooser/Source/Chooser/Private/GameplayTagColumn.cpp
new file mode 100644
--- /dev/null
+++ b/Engine/Plugins/Experimental/Chooser/Source/Chooser/Private/GameplayTagColumn.cpp
@@ -0,0 +1,69 @@
+// Copyright Epic Games, Inc. All Rights Reserved.
+#include "GameplayTagColumn.h"
+#include "ChooserPropertyAccess.h"
+
+bool UChooserParameterGameplayTag_ContextProperty::GetValue(const UObject* ContextObject, const FGameplayTagContainer*& OutResult) const
+{
+	UStruct* StructType = ContextObject->GetClass();
+	const void* Container = ContextObject;
+	
+	if (UE::Chooser::ResolvePropertyChain(Container, StructType, PropertyBindingChain))
+	{
+		if (const FStructProperty* Property = FindFProperty<FStructProperty>(StructType, PropertyBindingChain.Last()))
+		{
+			OutResult = Property->ContainerPtrToValuePtr<FGameplayTagContainer>(Container);
+			return true;
+		}
+	}
+
+	return false;
+}
+
+UChooserColumnGameplayTag::UChooserColumnGameplayTag(const FObjectInitializer& ObjectInitializer)
+{
+	InputValue = ObjectInitializer.CreateDefaultSubobject<UChooserParameterGameplayTag_ContextProperty>(this, "InputValue");
+	InputValue.GetObject()->SetFlags(RF_Transactional);
+}
+
+void UChooserColumnGameplayTag::Filter(const UObject* ContextObject, const TArray<uint32>& IndexListIn, TArray<uint32>& IndexListOut)
+{
+	if (ContextObject && InputValue)
+	{
+		const FGameplayTagContainer* Result = nullptr;
+		InputValue->GetValue(ContextObject,Result);
+
+		for (uint32 Index : IndexListIn)
+		{
+			 
+			if (RowValues.Num() > (int)Index)
+			{
+				if (RowValues[Index].IsEmpty())
+				{
+					IndexListOut.Push(Index);
+				}
+				else
+				{
+					if (TagMatchType == EGameplayContainerMatchType::All)
+					{
+						if (Result->HasAll(RowValues[Index]))
+						{
+							IndexListOut.Push(Index);
+						}
+					}
+					else
+					{
+						if (Result->HasAny(RowValues[Index]))
+						{
+							IndexListOut.Push(Index);
+						}
+					}
+				}
+			}
+		}
+	}
+	else
+	{
+		// passthrough fallback (behaves better during live editing)
+		IndexListOut = IndexListIn;
+	}
+}
\ No newline at end of file
diff --git a/Engine/Plugins/Experimental/Chooser/Source/Chooser/Public/ChooserPropertyAccess.h b/Engine/Plugins/Experimental/Chooser/Source/Chooser/Public/ChooserPropertyAccess.h
new file mode 100644
--- /dev/null
+++ b/Engine/Plugins/Experimental/Chooser/Source/Chooser/Public/ChooserPropertyAccess.h
@@ -0,0 +1,10 @@
+// Copyright Epic Games, Inc. All Rights Reserved.
+
+#pragma once
+
+#include "CoreMinimal.h"
+
+namespace UE::Chooser
+{
+	bool ResolvePropertyChain(const void*& Container, UStruct*& StructType, const TArray<FName>& PropertyBindingChain);
+}
\ No newline at end of file
diff --git a/Engine/Plugins/Experimental/Chooser/Source/Chooser/Public/IChooserParameterGameplayTag.h b/Engine/Plugins/Experimental/Chooser/Source/Chooser/Public/IChooserParameterGameplayTag.h
new file mode 100644
--- /dev/null
+++ b/Engine/Plugins/Experimental/Chooser/Source/Chooser/Public/IChooserParameterGameplayTag.h
@@ -0,0 +1,24 @@
+// Copyright Epic Games, Inc. All Rights Reserved.
+
+#pragma once
+
+#include "CoreMinimal.h"
+#include "UObject/Object.h"
+#include "UObject/Interface.h"
+#include "IChooserParameterGameplayTag.generated.h"
+
+struct FGameplayTagContainer;
+
+UINTERFACE(NotBlueprintType, meta = (CannotImplementInterfaceInBlueprint))
+class CHOOSER_API UChooserParameterGameplayTag : public UInterface
+{
+	GENERATED_BODY()
+};
+
+class CHOOSER_API IChooserParameterGameplayTag
+{
+	GENERATED_BODY()
+
+public:
+	virtual bool GetValue(const UObject* ContextObject, const FGameplayTagContainer*& OutResult) const { return false; }
+};
\ No newline at end of file
diff --git a/Engine/Plugins/Experimental/Chooser/Source/ChooserEditor/ChooserEditor.Build.cs b/Engine/Plugins/Experimental/Chooser/Source/ChooserEditor/ChooserEditor.Build.cs
--- a/Engine/Plugins/Experimental/Chooser/Source/ChooserEditor/ChooserEditor.Build.cs
+++ b/Engine/Plugins/Experimental/Chooser/Source/ChooserEditor/ChooserEditor.Build.cs
@@ -16,6 +16,7 @@ namespace UnrealBuildTool.Rules
 					"Chooser",
 					"UnrealEd",
 					"EditorWidgets",
+					"ToolWidgets",
 					"SlateCore",
 					"Slate",
 					"PropertyEditor",
@@ -23,7 +24,9 @@ namespace UnrealBuildTool.Rules
 					"EditorStyle",
 					"PropertyEditor",
 					"BlueprintGraph",
-					"GraphEditor"
+					"GraphEditor",
+					"GameplayTags",
+					"GameplayTagsEditor"
 					// ... add private dependencies that you statically link with here ...
 				}
 			);
diff --git a/Engine/Plugins/Experimental/Chooser/Source/ChooserEditor/Private/ChooserEditorModule.cpp b/Engine/Plugins/Experimental/Chooser/Source/ChooserEditor/Private/ChooserEditorModule.cpp
--- a/Engine/Plugins/Experimental/Chooser/Source/ChooserEditor/Private/ChooserEditorModule.cpp
+++ b/Engine/Plugins/Experimental/Chooser/Source/ChooserEditor/Private/ChooserEditorModule.cpp
@@ -4,6 +4,7 @@
 #include "IAssetTools.h"
 #include "AssetTypeActions_Chooser.h"
 #include "ChooserTableEditor.h"
+#include "GameplayTagColumnEditor.h"
 
 #define LOCTEXT_NAMESPACE "ChooserEditorModule"
 
@@ -16,6 +17,7 @@ void FModule::StartupModule()
 	AssetTypeActions_ChooserTable = MakeShared<FAssetTypeActions_ChooserTable>();
 	AssetTools.RegisterAssetTypeActions(AssetTypeActions_ChooserTable.ToSharedRef());
 	FChooserTableEditor::RegisterWidgets();
+	RegisterGameplayTagWidgets();
 	
 	InterfacePropertyTypeIdentifier = MakeShared<UE::ChooserEditor::FPropertyTypeIdentifier>();
 	
diff --git a/Engine/Plugins/Experimental/Chooser/Source/ChooserEditor/Private/ChooserTableEditor.cpp b/Engine/Plugins/Experimental/Chooser/Source/ChooserEditor/Private/ChooserTableEditor.cpp
--- a/Engine/Plugins/Experimental/Chooser/Source/ChooserEditor/Private/ChooserTableEditor.cpp
+++ b/Engine/Plugins/Experimental/Chooser/Source/ChooserEditor/Private/ChooserTableEditor.cpp
@@ -2,6 +2,7 @@
 
 #include "ChooserTableEditor.h"
 
+#include "ContextPropertyWidget.h"
 #include "Framework/Docking/TabManager.h"
 #include "Modules/ModuleManager.h"
 #include "Styling/AppStyle.h"
@@ -1024,148 +1025,16 @@ TSharedRef<SWidget> CreateBoolColumnWidget(UObject* Column, int Row)
 	});
 }
 	
-template <typename PropertyType>
-void ConvertToText_ContextProperty(const UObject* Object, FText& OutText)
+TSharedRef<SWidget> CreateBoolPropertyWidget(UObject* Object, UClass* ContextClass)
 {
-	if (const PropertyType* ContextProperty = Cast<PropertyType>(Object))
-	{
-		if (ContextProperty->PropertyBindingChain.Num()>0)
-		{
-			OutText = FText::FromName(ContextProperty->PropertyBindingChain.Last());
-		}
-	}
-}
-
-FLinearColor GetBindingColor(UChooserParameterBool_ContextProperty* Property)
-{
-	return GetDefault<UGraphEditorSettings>()->BooleanPinTypeColor;
+	return CreatePropertyWidget<UChooserParameterBool_ContextProperty>(Object, ContextClass, GetDefault<UGraphEditorSettings>()->BooleanPinTypeColor);
 }
 	
-FLinearColor GetBindingColor(UChooserParameterFloat_ContextProperty* Property)
+TSharedRef<SWidget> CreateFloatPropertyWidget(UObject* Object, UClass* ContextClass)
 {
-	return GetDefault<UGraphEditorSettings>()->DoublePinTypeColor;
+	return CreatePropertyWidget<UChooserParameterFloat_ContextProperty>(Object, ContextClass, GetDefault<UGraphEditorSettings>()->FloatPinTypeColor);
 }
 	
-const UGraphEditorSettings* Settings = GetDefault<UGraphEditorSettings>();
-
-template <typename PropertyType>
-TSharedRef<SWidget> CreatePropertyWidget(UObject* Object, UClass* ContextClass)
-{
-	TWeakObjectPtr<PropertyType> ContextProperty = Cast<PropertyType>(Object);
-	
-	FPropertyBindingWidgetArgs Args;
-	Args.bAllowPropertyBindings = true;
-	
-	Args.OnCanBindProperty = FOnCanBindProperty::CreateLambda([](FProperty* Property)
-	{
-		return Property == nullptr || PropertyType::CanBind(Property->GetCPPType());
-	});
-	Args.OnCanBindToClass = FOnCanBindToClass::CreateLambda([ContextClass ](UClass* InClass)
-	{
-		return true;
-	});
-	
-	Args.CurrentBindingColor = MakeAttributeLambda([ContextProperty, ContextClass]() {
-		 if (ContextProperty.IsValid() && ContextProperty->PropertyBindingChain.Num() > 0)
-		 {
-		 	return GetBindingColor(ContextProperty.Get());
-		 }
-		return FLinearColor::Gray;
-	});
-
-	Args.OnCanAcceptPropertyOrChildren = FOnCanBindProperty::CreateLambda([](FProperty* InProperty)
-		{
-			// Make only editor visible properties visible for binding.
-			return InProperty->HasAnyPropertyFlags(CPF_Edit);
-		});	
-
-	Args.OnAddBinding = FOnAddBinding::CreateLambda([ContextProperty](FName InPropertyName, const TArray<FBindingChainElement>& InBindingChain)
-		{
-			IPropertyAccessEditor& PropertyAccessEditor = IModularFeatures::Get().GetModularFeature<IPropertyAccessEditor>("PropertyAccessEditor");
-			if (ContextProperty.IsValid())
-			{
-				const FScopedTransaction Transaction(LOCTEXT("Change Property Binding", "Change Property Binding"));
-				ContextProperty->Modify(true);
-				
-				ContextProperty->PropertyBindingChain.Empty();
-
-				if (InBindingChain.Num() > 1)
-				{
-					for(int i = 1; i<InBindingChain.Num(); i++)
-					{
-						ContextProperty->PropertyBindingChain.Add(InBindingChain[i].Field.GetFName());
-					}
-				}
-				else
-				{
-					ContextProperty->PropertyBindingChain.Empty();
-				}
-			}
-		});
-
-	Args.CurrentBindingToolTipText = MakeAttributeLambda([ContextProperty]()
-			{
-				const FText Bind = LOCTEXT("Bind", "Bind");
-				FText CurrentValue = Bind;
-
-				if (ContextProperty.IsValid() && ContextProperty->PropertyBindingChain.Num()>0)
-				{
-					TArray<FText> BindingChainText;
-                 	BindingChainText.Reserve(ContextProperty->PropertyBindingChain.Num());
-                 
-                 	for (const FName& Name : ContextProperty->PropertyBindingChain)
-                 	{
-                 		BindingChainText.Add(FText::FromName(Name));
-                 	}
-					
-					CurrentValue = FText::Join(LOCTEXT("PropertyPathSeparator","."), BindingChainText);
-				}
-
-				return CurrentValue;	
-			});
-	
-	Args.CurrentBindingText = MakeAttributeLambda([ContextProperty]()
-			{
-				const FText Bind = LOCTEXT("Bind", "Bind");
-				FText CurrentValue = Bind;
-
-				if (ContextProperty.IsValid())
-				{
-					int BindingChainLength = ContextProperty->PropertyBindingChain.Num();
-					if (BindingChainLength > 0)
-					{
-						if (BindingChainLength == 1)
-						{
-							// single property, just use the property name
-							CurrentValue = FText::FromName(ContextProperty->PropertyBindingChain.Last());
-						}
-						else
-						{
-							// for longer chains always show the last struct/object name, and the final property name (full path in tooltip)
-							CurrentValue = FText::Join(LOCTEXT("PropertyPathSeparator","."),
-								TArray<FText>({
-									FText::FromName(ContextProperty->PropertyBindingChain[BindingChainLength-2]),
-									FText::FromName(ContextProperty->PropertyBindingChain[BindingChainLength-1])
-								}));
-						}
-					}
-				}
-
-				return CurrentValue;
-			});
-
-	Args.CurrentBindingImage = MakeAttributeLambda([]() -> const FSlateBrush*
-		{
-			static FName PropertyIcon(TEXT("Kismet.Tabs.Variables"));
-			return FAppStyle::GetBrush(PropertyIcon);
-		});
-
-	IPropertyAccessEditor& PropertyAccessEditor = IModularFeatures::Get().GetModularFeature<IPropertyAccessEditor>("PropertyAccessEditor");
-	FBindingContextStruct StructInfo;
-	StructInfo.Struct = ContextClass;
-	return PropertyAccessEditor.MakePropertyBindingWidget({StructInfo}, Args);
-}
-
 TSharedRef<SWidget> CreateFloatRangeColumnWidget(UObject* Column, int Row)
 {
 	UChooserColumnFloatRange* FloatRangeColumn = Cast<UChooserColumnFloatRange>(Column);
@@ -1300,8 +1169,8 @@ void FChooserTableEditor::RegisterWidgets()
 	FObjectChooserWidgetFactories::ChooserWidgetCreators.Add(UObjectChooser_Asset::StaticClass(), CreateAssetWidget);
 	FObjectChooserWidgetFactories::ChooserWidgetCreators.Add(UObjectChooser_EvaluateChooser::StaticClass(), CreateEvaluateChooserWidget);
 	
-	FObjectChooserWidgetFactories::ChooserWidgetCreators.Add(UChooserParameterBool_ContextProperty::StaticClass(), CreatePropertyWidget<UChooserParameterBool_ContextProperty>);
-	FObjectChooserWidgetFactories::ChooserWidgetCreators.Add(UChooserParameterFloat_ContextProperty::StaticClass(), CreatePropertyWidget<UChooserParameterFloat_ContextProperty>);
+	FObjectChooserWidgetFactories::ChooserWidgetCreators.Add(UChooserParameterBool_ContextProperty::StaticClass(), CreateBoolPropertyWidget);
+	FObjectChooserWidgetFactories::ChooserWidgetCreators.Add(UChooserParameterFloat_ContextProperty::StaticClass(), CreateFloatPropertyWidget);
 	FObjectChooserWidgetFactories::ChooserTextConverter.Add(UChooserParameterBool_ContextProperty::StaticClass(), ConvertToText_ContextProperty<UChooserParameterBool_ContextProperty>);
 	FObjectChooserWidgetFactories::ChooserTextConverter.Add(UChooserParameterFloat_ContextProperty::StaticClass(), ConvertToText_ContextProperty<UChooserParameterFloat_ContextProperty>);
 
diff --git a/Engine/Plugins/Experimental/Chooser/Source/ChooserEditor/Private/GameplayTagColumnEditor.cpp b/Engine/Plugins/Experimental/Chooser/Source/ChooserEditor/Private/GameplayTagColumnEditor.cpp
new file mode 100644
--- /dev/null
+++ b/Engine/Plugins/Experimental/Chooser/Source/ChooserEditor/Private/GameplayTagColumnEditor.cpp
@@ -0,0 +1,67 @@
+// Copyright Epic Games, Inc. All Rights Reserved.
+
+#include "GameplayTagColumnEditor.h"
+#include "ContextPropertyWidget.h"
+#include "GameplayTagColumn.h"
+#include "ObjectChooserWidgetFactories.h"
+#include "ChooserTableEditor.h"
+#include "SGameplayTagWidget.h"
+#include "SSimpleComboButton.h"
+#include "GraphEditorSettings.h"
+
+#define LOCTEXT_NAMESPACE "GameplayTagColumnEditor"
+
+namespace UE::ChooserEditor
+{
+
+TSharedRef<SWidget> CreateGameplayTagColumnWidget(UObject* Column, int Row)
+{
+	UChooserColumnGameplayTag* GameplayTagColumn = Cast<UChooserColumnGameplayTag>(Column);
+
+	return SNew(SSimpleComboButton)
+		.Text_Lambda([GameplayTagColumn, Row]()
+		{
+			if (Row < GameplayTagColumn->RowValues.Num())
+			{
+				FText Result = FText::FromString(GameplayTagColumn->RowValues[Row].ToStringSimple(false));
+				if (Result.IsEmpty())
+				{
+					Result = LOCTEXT("Any Tag", "[Any]");
+				}
+				return Result;
+			}
+			else
+			{
+				return FText();
+			}
+		})	
+		.OnGetMenuContent_Lambda([GameplayTagColumn, Row]()
+		{
+			if (Row < GameplayTagColumn->RowValues.Num())
+			{
+				TArray<SGameplayTagWidget::FEditableGameplayTagContainerDatum> EditableContainers;
+				EditableContainers.Emplace(GameplayTagColumn, &(GameplayTagColumn->RowValues[Row]));
+				return TSharedRef<SWidget>(SNew(SGameplayTagWidget, EditableContainers));
+			}
+
+			return SNullWidget::NullWidget;
+		}
+	);
+}
+
+TSharedRef<SWidget> CreateGameplayTagPropertyWidget(UObject* Object, UClass* ContextClass)
+{
+	return CreatePropertyWidget<UChooserParameterGameplayTag_ContextProperty>(Object, ContextClass, GetDefault<UGraphEditorSettings>()->StructPinTypeColor);
+}
+	
+void RegisterGameplayTagWidgets()
+{
+	FObjectChooserWidgetFactories::ChooserWidgetCreators.Add(UChooserParameterGameplayTag_ContextProperty::StaticClass(), CreateGameplayTagPropertyWidget);
+	FObjectChooserWidgetFactories::ChooserTextConverter.Add(UChooserParameterGameplayTag_ContextProperty::StaticClass(), ConvertToText_ContextProperty<UChooserParameterGameplayTag_ContextProperty>);
+
+	FChooserTableEditor::ColumnWidgetCreators.Add(UChooserColumnGameplayTag::StaticClass(), CreateGameplayTagColumnWidget);
+}
+	
+}
+
+#undef LOCTEXT_NAMESPACE
diff --git a/Engine/Plugins/Experimental/Chooser/Source/ChooserEditor/Private/GameplayTagColumnEditor.h b/Engine/Plugins/Experimental/Chooser/Source/ChooserEditor/Private/GameplayTagColumnEditor.h
new file mode 100644
--- /dev/null
+++ b/Engine/Plugins/Experimental/Chooser/Source/ChooserEditor/Private/GameplayTagColumnEditor.h
@@ -0,0 +1,10 @@
+// Copyright Epic Games, Inc. All Rights Reserved.
+
+#pragma once
+
+#include "CoreMinimal.h"
+
+namespace UE::ChooserEditor
+{
+	void RegisterGameplayTagWidgets();
+}
\ No newline at end of file
diff --git a/Engine/Plugins/Experimental/Chooser/Source/ChooserEditor/Public/ContextPropertyWidget.h b/Engine/Plugins/Experimental/Chooser/Source/ChooserEditor/Public/ContextPropertyWidget.h
new file mode 100644
--- /dev/null
+++ b/Engine/Plugins/Experimental/Chooser/Source/ChooserEditor/Public/ContextPropertyWidget.h
@@ -0,0 +1,142 @@
+// Copyright Epic Games, Inc. All Rights Reserved.
+#pragma once
+
+#include "ScopedTransaction.h"
+#include "IPropertyAccessEditor.h"
+#include "Features/IModularFeatures.h"
+#include "Styling/AppStyle.h"
+
+namespace UE::ChooserEditor
+{
+	// Filter class for things which implement a particular interface
+	template <typename PropertyType>
+    TSharedRef<SWidget> CreatePropertyWidget(UObject* Object, UClass* ContextClass, const FLinearColor& BindingColor)
+    {
+    	TWeakObjectPtr<PropertyType> ContextProperty = Cast<PropertyType>(Object);
+    	
+    	FPropertyBindingWidgetArgs Args;
+    	Args.bAllowPropertyBindings = true;
+    	
+    	Args.OnCanBindProperty = FOnCanBindProperty::CreateLambda([](FProperty* Property)
+    	{
+    		return Property == nullptr || PropertyType::CanBind(Property->GetCPPType());
+    	});
+    	Args.OnCanBindToClass = FOnCanBindToClass::CreateLambda([ContextClass ](UClass* InClass)
+    	{
+    		return true;
+    	});
+    	
+    	Args.CurrentBindingColor = MakeAttributeLambda([ContextProperty, BindingColor]() {
+    		 if (ContextProperty.IsValid() && ContextProperty->PropertyBindingChain.Num() > 0)
+    		 {
+    		 	return BindingColor;
+    		 }
+    		return FLinearColor::Gray;
+    	});
+    
+    	Args.OnCanAcceptPropertyOrChildren = FOnCanBindProperty::CreateLambda([](FProperty* InProperty)
+    		{
+    			// Make only editor visible properties visible for binding.
+    			return InProperty->HasAnyPropertyFlags(CPF_Edit);
+    		});	
+    
+    	Args.OnAddBinding = FOnAddBinding::CreateLambda([ContextProperty](FName InPropertyName, const TArray<FBindingChainElement>& InBindingChain)
+    		{
+    			IPropertyAccessEditor& PropertyAccessEditor = IModularFeatures::Get().GetModularFeature<IPropertyAccessEditor>("PropertyAccessEditor");
+    			if (ContextProperty.IsValid())
+    			{
+    				const FScopedTransaction Transaction(NSLOCTEXT("ContextPropertyWidget", "Change Property Binding", "Change Property Binding"));
+    				ContextProperty->Modify(true);
+    				
+    				ContextProperty->PropertyBindingChain.Empty();
+    
+    				if (InBindingChain.Num() > 1)
+    				{
+    					for(int i = 1; i<InBindingChain.Num(); i++)
+    					{
+    						ContextProperty->PropertyBindingChain.Add(InBindingChain[i].Field.GetFName());
+    					}
+    				}
+    				else
+    				{
+    					ContextProperty->PropertyBindingChain.Empty();
+    				}
+    			}
+    		});
+    
+    	Args.CurrentBindingToolTipText = MakeAttributeLambda([ContextProperty]()
+    			{
+    				const FText Bind = NSLOCTEXT("ContextPropertyWidget", "Bind", "Bind");
+    				FText CurrentValue = Bind;
+    
+    				if (ContextProperty.IsValid() && ContextProperty->PropertyBindingChain.Num()>0)
+    				{
+    					TArray<FText> BindingChainText;
+                     	BindingChainText.Reserve(ContextProperty->PropertyBindingChain.Num());
+                     
+                     	for (const FName& Name : ContextProperty->PropertyBindingChain)
+                     	{
+                     		BindingChainText.Add(FText::FromName(Name));
+                     	}
+    					
+    					CurrentValue = FText::Join(NSLOCTEXT("ContextPropertyWidget", "PropertyPathSeparator","."), BindingChainText);
+    				}
+    
+    				return CurrentValue;	
+    			});
+    	
+    	Args.CurrentBindingText = MakeAttributeLambda([ContextProperty]()
+    			{
+    				const FText Bind = NSLOCTEXT("ContextPropertyWidget", "Bind", "Bind");
+    				FText CurrentValue = Bind;
+    
+    				if (ContextProperty.IsValid())
+    				{
+    					int BindingChainLength = ContextProperty->PropertyBindingChain.Num();
+    					if (BindingChainLength > 0)
+    					{
+    						if (BindingChainLength == 1)
+    						{
+    							// single property, just use the property name
+    							CurrentValue = FText::FromName(ContextProperty->PropertyBindingChain.Last());
+    						}
+    						else
+    						{
+    							// for longer chains always show the last struct/object name, and the final property name (full path in tooltip)
+    							CurrentValue = FText::Join(NSLOCTEXT("ContextPropertyWidget", "PropertyPathSeparator","."),
+    								TArray<FText>({
+    									FText::FromName(ContextProperty->PropertyBindingChain[BindingChainLength-2]),
+    									FText::FromName(ContextProperty->PropertyBindingChain[BindingChainLength-1])
+    								}));
+    						}
+    					}
+    				}
+    
+    				return CurrentValue;
+    			});
+    
+    	Args.CurrentBindingImage = MakeAttributeLambda([]() -> const FSlateBrush*
+    		{
+    			static FName PropertyIcon(TEXT("Kismet.Tabs.Variables"));
+    			return FAppStyle::GetBrush(PropertyIcon);
+    		});
+    
+    	IPropertyAccessEditor& PropertyAccessEditor = IModularFeatures::Get().GetModularFeature<IPropertyAccessEditor>("PropertyAccessEditor");
+    	FBindingContextStruct StructInfo;
+    	StructInfo.Struct = ContextClass;
+    	return PropertyAccessEditor.MakePropertyBindingWidget({StructInfo}, Args);
+    }
+				
+	template <typename PropertyType>
+	void ConvertToText_ContextProperty(const UObject* Object, FText& OutText)
+	{
+		if (const PropertyType* ContextProperty = Cast<PropertyType>(Object))
+		{
+			if (ContextProperty->PropertyBindingChain.Num()>0)
+			{
+				OutText = FText::FromName(ContextProperty->PropertyBindingChain.Last());
+			}
+		}
+	}
+
+}
\ No newline at end of file
