Chooser code cleanup. Moving all column code to separate files from core chooser code

diff --git a/Engine/Plugins/Experimental/Chooser/Source/Chooser/Internal/BoolColumn.h b/Engine/Plugins/Experimental/Chooser/Source/Chooser/Internal/BoolColumn.h
new file mode 100644
--- /dev/null
+++ b/Engine/Plugins/Experimental/Chooser/Source/Chooser/Internal/BoolColumn.h
@@ -0,0 +1,60 @@
+// Copyright Epic Games, Inc. All Rights Reserved.
+#pragma once
+
+#include "CoreMinimal.h"
+#include "IChooserColumn.h"
+#include "IChooserParameterBool.h"
+#include "BoolColumn.generated.h"
+
+
+UCLASS()
+class CHOOSER_API UChooserParameterBool_ContextProperty :  public UObject, public IChooserParameterBool
+{
+	GENERATED_BODY()
+public:
+
+	UPROPERTY()
+	TArray<FName> PropertyBindingChain;
+	
+	virtual bool GetValue(const UObject* ContextObject, bool& OutResult) const override;
+
+#if WITH_EDITOR
+	static bool CanBind(const FString& TypeName)
+	{
+		static FString BoolTypeName = "bool";
+		return TypeName == BoolTypeName;
+	}
+#endif
+};
+
+UCLASS()
+class CHOOSER_API UChooserColumnBool : public UObject, public IChooserColumn
+{
+	GENERATED_BODY()
+	public:
+	UChooserColumnBool() {};
+	UChooserColumnBool(const FObjectInitializer& ObjectInitializer);
+	
+	UPROPERTY(EditAnywhere, Meta = (EditInlineInterface = "true"), Category = "Input")
+	TScriptInterface<IChooserParameterBool> InputValue;
+
+	UPROPERTY(EditAnywhere, Category=Runtime)
+	// array of results (cells for this column for each row in the table)
+	// should match the length of the Results array 
+	TArray<bool> RowValues;
+	
+	virtual void Filter(const UObject* ContextObject, const TArray<uint32>& IndexListIn, TArray<uint32>& IndexListOut) override;
+
+	// todo: macro boilerplate
+	virtual void SetNumRows(uint32 NumRows) override { RowValues.SetNum(NumRows); }
+	virtual void DeleteRows(const TArray<uint32> & RowIndices )
+	{
+		for(uint32 Index : RowIndices)
+		{
+			RowValues.RemoveAt(Index);
+		}
+	}
+	virtual UClass* GetInputValueInterface() override { return UChooserParameterBool::StaticClass(); };
+	virtual UObject* GetInputValue() override { return InputValue.GetObject(); };
+	virtual void SetInputValue(UObject* Value) override { InputValue = Value; };
+};
\ No newline at end of file
diff --git a/Engine/Plugins/Experimental/Chooser/Source/Chooser/Internal/Chooser.h b/Engine/Plugins/Experimental/Chooser/Source/Chooser/Internal/Chooser.h
--- a/Engine/Plugins/Experimental/Chooser/Source/Chooser/Internal/Chooser.h
+++ b/Engine/Plugins/Experimental/Chooser/Source/Chooser/Internal/Chooser.h
@@ -5,123 +5,9 @@
 #include "UObject/Object.h"
 #include "IObjectChooser.h"
 #include "IChooserColumn.h"
-#include "IChooserColumnParameter.h"
 
 #include "Chooser.generated.h"
 
-UCLASS()
-class CHOOSER_API UChooserParameterBool_ContextProperty :  public UObject, public IChooserParameterBool
-{
-	GENERATED_BODY()
-public:
-
-	UPROPERTY()
-	TArray<FName> PropertyBindingChain;
-	
-	virtual bool GetValue(const UObject* ContextObject, bool& OutResult) const override;
-
-#if WITH_EDITOR
-	static bool CanBind(const FString& TypeName)
-	{
-		static FString BoolTypeName = "bool";
-		return TypeName == BoolTypeName;
-	}
-#endif
-};
-
-UCLASS()
-class CHOOSER_API UChooserParameterFloat_ContextProperty :  public UObject, public IChooserParameterFloat
-{
-	GENERATED_BODY()
-public:
-	UPROPERTY()
-	TArray<FName> PropertyBindingChain;
-	
-	virtual bool GetValue(const UObject* ContextObject, float& OutResult) const override;
-
-#if WITH_EDITOR
-	static bool CanBind(const FString& TypeName)
-	{
-		static FString DoubleTypeName = "double";
-		static FString FloatTypeName = "float";
-		return TypeName == FloatTypeName || TypeName == DoubleTypeName;
-	}
-#endif
-};
-
-UCLASS()
-class CHOOSER_API UChooserColumnBool : public UObject, public IChooserColumn
-{
-	GENERATED_BODY()
-	public:
-	UChooserColumnBool() {};
-	UChooserColumnBool(const FObjectInitializer& ObjectInitializer);
-	
-	UPROPERTY(EditAnywhere, Meta = (EditInlineInterface = "true"), Category = "Input")
-	TScriptInterface<IChooserParameterBool> InputValue;
-
-	UPROPERTY(EditAnywhere, Category=Runtime)
-	// array of results (cells for this column for each row in the table)
-	// should match the length of the Results array 
-	TArray<bool> RowValues;
-	
-	virtual void Filter(const UObject* ContextObject, const TArray<uint32>& IndexListIn, TArray<uint32>& IndexListOut) override;
-
-	// todo: macro boilerplate
-	virtual void SetNumRows(uint32 NumRows) override { RowValues.SetNum(NumRows); }
-	virtual void DeleteRows(const TArray<uint32> & RowIndices )
-	{
-		for(uint32 Index : RowIndices)
-		{
-			RowValues.RemoveAt(Index);
-		}
-	}
-	virtual UClass* GetInputValueInterface() override { return UChooserParameterBool::StaticClass(); };
-	virtual UObject* GetInputValue() override { return InputValue.GetObject(); };
-	virtual void SetInputValue(UObject* Value) override { InputValue = Value; };
-};
-
-USTRUCT()
-struct FChooserFloatRangeRowData
-{
-	GENERATED_BODY()
-	
-	UPROPERTY(EditAnywhere, Category=Runtime)
-	float Min=0;
-	
-	UPROPERTY(EditAnywhere, Category=Runtime)
-	float Max=0;
-};
-
-UCLASS()
-class CHOOSER_API UChooserColumnFloatRange : public UObject, public IChooserColumn
-{
-	GENERATED_BODY()
-	public:
-	UChooserColumnFloatRange() {}
-	UChooserColumnFloatRange(const FObjectInitializer& ObjectInitializer);
-		
-	UPROPERTY(EditAnywhere, Meta = (EditInlineInterface = "true"), Category = "Input")
-	TScriptInterface<IChooserParameterFloat> InputValue;
-	
-	UPROPERTY(EditAnywhere, Category=Runtime)
-	// array of results (cells for this column for each row in the table)
-	// should match the length of the Results array 
-	TArray<FChooserFloatRangeRowData> RowValues;
-	
-	virtual void Filter(const UObject* ContextObject, const TArray<uint32>& IndexListIn, TArray<uint32>& IndexListOut) override;
-	virtual void SetNumRows(uint32 NumRows) { RowValues.SetNum(NumRows); }
-	virtual void DeleteRows(const TArray<uint32> & RowIndices )
-	{
-		for(uint32 Index : RowIndices)
-		{
-			RowValues.RemoveAt(Index);
-		}
-	}
-	virtual UClass* GetInputValueInterface() { return UChooserParameterFloat::StaticClass(); };
-	virtual UObject* GetInputValue() override { return InputValue.GetObject(); };
-	virtual void SetInputValue(UObject* Value) override { InputValue = Value; };
-};
 
 UCLASS(MinimalAPI)
 class UChooserTable : public UObject
diff --git a/Engine/Plugins/Experimental/Chooser/Source/Chooser/Internal/FloatRangeColumn.h b/Engine/Plugins/Experimental/Chooser/Source/Chooser/Internal/FloatRangeColumn.h
new file mode 100644
--- /dev/null
+++ b/Engine/Plugins/Experimental/Chooser/Source/Chooser/Internal/FloatRangeColumn.h
@@ -0,0 +1,69 @@
+// Copyright Epic Games, Inc. All Rights Reserved.
+#pragma once
+
+#include "CoreMinimal.h"
+#include "IChooserColumn.h"
+#include "IChooserParameterFloat.h"
+#include "FloatRangeColumn.generated.h"
+
+UCLASS()
+class CHOOSER_API UChooserParameterFloat_ContextProperty :  public UObject, public IChooserParameterFloat
+{
+	GENERATED_BODY()
+public:
+	UPROPERTY()
+	TArray<FName> PropertyBindingChain;
+	
+	virtual bool GetValue(const UObject* ContextObject, float& OutResult) const override;
+
+#if WITH_EDITOR
+	static bool CanBind(const FString& TypeName)
+	{
+		static FString DoubleTypeName = "double";
+		static FString FloatTypeName = "float";
+		return TypeName == FloatTypeName || TypeName == DoubleTypeName;
+	}
+#endif
+};
+
+USTRUCT()
+struct FChooserFloatRangeRowData
+{
+	GENERATED_BODY()
+	
+	UPROPERTY(EditAnywhere, Category=Runtime)
+	float Min=0;
+	
+	UPROPERTY(EditAnywhere, Category=Runtime)
+	float Max=0;
+};
+
+UCLASS()
+class CHOOSER_API UChooserColumnFloatRange : public UObject, public IChooserColumn
+{
+	GENERATED_BODY()
+	public:
+	UChooserColumnFloatRange() {}
+	UChooserColumnFloatRange(const FObjectInitializer& ObjectInitializer);
+		
+	UPROPERTY(EditAnywhere, Meta = (EditInlineInterface = "true"), Category = "Input")
+	TScriptInterface<IChooserParameterFloat> InputValue;
+	
+	UPROPERTY(EditAnywhere, Category=Runtime)
+	// array of results (cells for this column for each row in the table)
+	// should match the length of the Results array 
+	TArray<FChooserFloatRangeRowData> RowValues;
+	
+	virtual void Filter(const UObject* ContextObject, const TArray<uint32>& IndexListIn, TArray<uint32>& IndexListOut) override;
+	virtual void SetNumRows(uint32 NumRows) { RowValues.SetNum(NumRows); }
+	virtual void DeleteRows(const TArray<uint32> & RowIndices )
+	{
+		for(uint32 Index : RowIndices)
+		{
+			RowValues.RemoveAt(Index);
+		}
+	}
+	virtual UClass* GetInputValueInterface() { return UChooserParameterFloat::StaticClass(); };
+	virtual UObject* GetInputValue() override { return InputValue.GetObject(); };
+	virtual void SetInputValue(UObject* Value) override { InputValue = Value; };
+};
\ No newline at end of file
diff --git a/Engine/Plugins/Experimental/Chooser/Source/Chooser/Private/BoolColumn.cpp b/Engine/Plugins/Experimental/Chooser/Source/Chooser/Private/BoolColumn.cpp
new file mode 100644
--- /dev/null
+++ b/Engine/Plugins/Experimental/Chooser/Source/Chooser/Private/BoolColumn.cpp
@@ -0,0 +1,51 @@
+// Copyright Epic Games, Inc. All Rights Reserved.
+#include "BoolColumn.h"
+#include "ChooserPropertyAccess.h"
+
+bool UChooserParameterBool_ContextProperty::GetValue(const UObject* ContextObject, bool& OutResult) const
+{
+	UStruct* StructType = ContextObject->GetClass();
+	const void* Container = ContextObject;
+	
+	if (UE::Chooser::ResolvePropertyChain(Container, StructType, PropertyBindingChain))
+	{
+		if (const FBoolProperty* Property = FindFProperty<FBoolProperty>(StructType, PropertyBindingChain.Last()))
+		{
+			OutResult = *Property->ContainerPtrToValuePtr<bool>(Container);
+			return true;
+		}
+	}
+
+	return false;
+}
+
+UChooserColumnBool::UChooserColumnBool(const FObjectInitializer& ObjectInitializer)
+{
+	InputValue = ObjectInitializer.CreateDefaultSubobject<UChooserParameterBool_ContextProperty>(this, "InputValue");
+	InputValue.GetObject()->SetFlags(RF_Transactional);
+}
+
+void UChooserColumnBool::Filter(const UObject* ContextObject, const TArray<uint32>& IndexListIn, TArray<uint32>& IndexListOut)
+{
+	if (ContextObject && InputValue)
+	{
+		bool Result = false;
+		InputValue->GetValue(ContextObject,Result);
+		
+		for (uint32 Index : IndexListIn)
+		{
+			if (RowValues.Num() > (int)Index)
+			{
+				if (Result == RowValues[Index])
+				{
+					IndexListOut.Push(Index);
+				}
+			}
+		}
+	}
+	else
+	{
+		// passthrough fallback (behaves better during live editing)
+		IndexListOut = IndexListIn;
+	}
+}
\ No newline at end of file
diff --git a/Engine/Plugins/Experimental/Chooser/Source/Chooser/Private/Chooser.cpp b/Engine/Plugins/Experimental/Chooser/Source/Chooser/Private/Chooser.cpp
--- a/Engine/Plugins/Experimental/Chooser/Source/Chooser/Private/Chooser.cpp
+++ b/Engine/Plugins/Experimental/Chooser/Source/Chooser/Private/Chooser.cpp
@@ -3,115 +3,12 @@
 #include "ChooserFunctionLibrary.h"
 #include "ChooserPropertyAccess.h"
 
-bool UChooserParameterBool_ContextProperty::GetValue(const UObject* ContextObject, bool& OutResult) const
-{
-	UStruct* StructType = ContextObject->GetClass();
-	const void* Container = ContextObject;
-	
-	if (UE::Chooser::ResolvePropertyChain(Container, StructType, PropertyBindingChain))
-	{
-		if (const FBoolProperty* Property = FindFProperty<FBoolProperty>(StructType, PropertyBindingChain.Last()))
-		{
-			OutResult = *Property->ContainerPtrToValuePtr<bool>(Container);
-			return true;
-		}
-	}
-
-	return false;
-}
-
-bool UChooserParameterFloat_ContextProperty::GetValue(const UObject* ContextObject, float& OutResult) const
-{
-	UStruct* StructType = ContextObject->GetClass();
-	const void* Container = ContextObject;
-	
-	if (UE::Chooser::ResolvePropertyChain(Container, StructType, PropertyBindingChain))
-	{
-		if (const FDoubleProperty* DoubleProperty = FindFProperty<FDoubleProperty>(StructType, PropertyBindingChain.Last()))
-		{
-			OutResult = *DoubleProperty->ContainerPtrToValuePtr<double>(Container);
-			return true;
-		}
-		
-		if (const FFloatProperty* FloatProperty = FindFProperty<FFloatProperty>(StructType, PropertyBindingChain.Last()))
-		{
-			OutResult = *FloatProperty->ContainerPtrToValuePtr<float>(Container);
-			return true;
-		}
-	}
-
-	return false;
-}
-
 UChooserTable::UChooserTable(const FObjectInitializer& Initializer)
 	:Super(Initializer)
 {
 
 }
 
-UChooserColumnBool::UChooserColumnBool(const FObjectInitializer& ObjectInitializer)
-{
-	InputValue = ObjectInitializer.CreateDefaultSubobject<UChooserParameterBool_ContextProperty>(this, "InputValue");
-	InputValue.GetObject()->SetFlags(RF_Transactional);
-}
-
-void UChooserColumnBool::Filter(const UObject* ContextObject, const TArray<uint32>& IndexListIn, TArray<uint32>& IndexListOut)
-{
-	if (ContextObject && InputValue)
-	{
-		bool Result = false;
-		InputValue->GetValue(ContextObject,Result);
-		
-		for (uint32 Index : IndexListIn)
-		{
-			if (RowValues.Num() > (int)Index)
-			{
-				if (Result == RowValues[Index])
-				{
-					IndexListOut.Push(Index);
-				}
-			}
-		}
-	}
-	else
-	{
-		// passthrough fallback (behaves better during live editing)
-		IndexListOut = IndexListIn;
-	}
-}
-
-UChooserColumnFloatRange::UChooserColumnFloatRange(const FObjectInitializer& ObjectInitializer)
-{
-	InputValue = ObjectInitializer.CreateDefaultSubobject<UChooserParameterFloat_ContextProperty>(this, "InputValue");
-	InputValue.GetObject()->SetFlags(RF_Transactional);
-}
-
-void UChooserColumnFloatRange::Filter(const UObject* ContextObject, const TArray<uint32>& IndexListIn, TArray<uint32>& IndexListOut)
-{
-	if (ContextObject && InputValue)
-	{
-		float Result = 0.0f;
-		InputValue->GetValue(ContextObject, Result);
-
-		for(uint32 Index : IndexListIn)
-		{
-			if (RowValues.Num() > (int)Index)
-			{
-				const FChooserFloatRangeRowData& RowValue = RowValues[Index];
-				if (Result >= RowValue.Min && Result <= RowValue.Max)
-				{
-					IndexListOut.Push(Index);
-				}
-			}
-		}
-	}
-	else
-	{
-		// passthrough fallback (behaves better during live editing)
-		IndexListOut = IndexListIn;
-	}
-}
-
 static IObjectChooser::EIteratorStatus StaticEvaluateChooser(const UObject* ContextObject, const UChooserTable* Chooser, IObjectChooser::FObjectChooserIteratorCallback Callback)
 {
 	if (Chooser == nullptr)
diff --git a/Engine/Plugins/Experimental/Chooser/Source/Chooser/Private/FloatRangeColumn.cpp b/Engine/Plugins/Experimental/Chooser/Source/Chooser/Private/FloatRangeColumn.cpp
new file mode 100644
--- /dev/null
+++ b/Engine/Plugins/Experimental/Chooser/Source/Chooser/Private/FloatRangeColumn.cpp
@@ -0,0 +1,58 @@
+// Copyright Epic Games, Inc. All Rights Reserved.
+#include "FloatRangeColumn.h"
+#include "ChooserPropertyAccess.h"
+
+bool UChooserParameterFloat_ContextProperty::GetValue(const UObject* ContextObject, float& OutResult) const
+{
+	UStruct* StructType = ContextObject->GetClass();
+	const void* Container = ContextObject;
+	
+	if (UE::Chooser::ResolvePropertyChain(Container, StructType, PropertyBindingChain))
+	{
+		if (const FDoubleProperty* DoubleProperty = FindFProperty<FDoubleProperty>(StructType, PropertyBindingChain.Last()))
+		{
+			OutResult = *DoubleProperty->ContainerPtrToValuePtr<double>(Container);
+			return true;
+		}
+		
+		if (const FFloatProperty* FloatProperty = FindFProperty<FFloatProperty>(StructType, PropertyBindingChain.Last()))
+		{
+			OutResult = *FloatProperty->ContainerPtrToValuePtr<float>(Container);
+			return true;
+		}
+	}
+
+	return false;
+}
+
+UChooserColumnFloatRange::UChooserColumnFloatRange(const FObjectInitializer& ObjectInitializer)
+{
+	InputValue = ObjectInitializer.CreateDefaultSubobject<UChooserParameterFloat_ContextProperty>(this, "InputValue");
+	InputValue.GetObject()->SetFlags(RF_Transactional);
+}
+
+void UChooserColumnFloatRange::Filter(const UObject* ContextObject, const TArray<uint32>& IndexListIn, TArray<uint32>& IndexListOut)
+{
+	if (ContextObject && InputValue)
+	{
+		float Result = 0.0f;
+		InputValue->GetValue(ContextObject, Result);
+
+		for(uint32 Index : IndexListIn)
+		{
+			if (RowValues.Num() > (int)Index)
+			{
+				const FChooserFloatRangeRowData& RowValue = RowValues[Index];
+				if (Result >= RowValue.Min && Result <= RowValue.Max)
+				{
+					IndexListOut.Push(Index);
+				}
+			}
+		}
+	}
+	else
+	{
+		// passthrough fallback (behaves better during live editing)
+		IndexListOut = IndexListIn;
+	}
+}
\ No newline at end of file
diff --git a/Engine/Plugins/Experimental/Chooser/Source/Chooser/Public/IChooserColumnParameter.h b/Engine/Plugins/Experimental/Chooser/Source/Chooser/Public/IChooserColumnParameter.h
deleted file mode 100644
diff --git a/Engine/Plugins/Experimental/Chooser/Source/Chooser/Public/IChooserParameterBool.h b/Engine/Plugins/Experimental/Chooser/Source/Chooser/Public/IChooserParameterBool.h
new file mode 100644
--- /dev/null
+++ b/Engine/Plugins/Experimental/Chooser/Source/Chooser/Public/IChooserParameterBool.h
@@ -0,0 +1,22 @@
+// Copyright Epic Games, Inc. All Rights Reserved.
+
+#pragma once
+
+#include "CoreMinimal.h"
+#include "UObject/Object.h"
+#include "UObject/Interface.h"
+#include "IChooserParameterBool.generated.h"
+
+UINTERFACE(NotBlueprintType, meta = (CannotImplementInterfaceInBlueprint))
+class CHOOSER_API UChooserParameterBool : public UInterface
+{
+	GENERATED_BODY()
+};
+
+class CHOOSER_API IChooserParameterBool
+{
+	GENERATED_BODY()
+
+public:
+	virtual bool GetValue(const UObject* ContextObject, bool& OutResult) const { return false; }
+};
diff --git a/Engine/Plugins/Experimental/Chooser/Source/Chooser/Public/IChooserParameterFloat.h b/Engine/Plugins/Experimental/Chooser/Source/Chooser/Public/IChooserParameterFloat.h
new file mode 100644
--- /dev/null
+++ b/Engine/Plugins/Experimental/Chooser/Source/Chooser/Public/IChooserParameterFloat.h
@@ -0,0 +1,22 @@
+// Copyright Epic Games, Inc. All Rights Reserved.
+
+#pragma once
+
+#include "CoreMinimal.h"
+#include "UObject/Object.h"
+#include "UObject/Interface.h"
+#include "IChooserParameterFloat.generated.h"
+
+UINTERFACE(NotBlueprintType, meta = (CannotImplementInterfaceInBlueprint))
+class CHOOSER_API UChooserParameterFloat : public UInterface
+{
+	GENERATED_BODY()
+};
+
+class CHOOSER_API IChooserParameterFloat
+{
+	GENERATED_BODY()
+
+public:
+	virtual bool GetValue(const UObject* ContextObject, float& OutResult) const { return false; }
+};
\ No newline at end of file
diff --git a/Engine/Plugins/Experimental/Chooser/Source/ChooserEditor/Private/BoolColumnEditor.cpp b/Engine/Plugins/Experimental/Chooser/Source/ChooserEditor/Private/BoolColumnEditor.cpp
new file mode 100644
--- /dev/null
+++ b/Engine/Plugins/Experimental/Chooser/Source/ChooserEditor/Private/BoolColumnEditor.cpp
@@ -0,0 +1,52 @@
+// Copyright Epic Games, Inc. All Rights Reserved.
+
+#include "BoolColumnEditor.h"
+#include "BoolColumn.h"
+#include "ContextPropertyWidget.h"
+#include "ObjectChooserWidgetFactories.h"
+#include "ChooserTableEditor.h"
+#include "Widgets/Input/SCheckBox.h"
+#include "GraphEditorSettings.h"
+
+#define LOCTEXT_NAMESPACE "BoolColumnEditor"
+
+namespace UE::ChooserEditor
+{
+
+TSharedRef<SWidget> CreateBoolColumnWidget(UObject* Column, int Row)
+{
+	UChooserColumnBool* BoolColumn = Cast<UChooserColumnBool>(Column);
+
+	return SNew (SCheckBox)
+	.OnCheckStateChanged_Lambda([BoolColumn,Row](ECheckBoxState State)
+	{
+		if (Row < BoolColumn->RowValues.Num())
+		{
+			const FScopedTransaction Transaction(LOCTEXT("Change Bool Value", "Change Bool Value"));
+			BoolColumn->Modify(true);
+			BoolColumn->RowValues[Row] = (State == ECheckBoxState::Checked);
+		}
+	})
+	.IsChecked_Lambda([BoolColumn, Row]()
+	{
+		const bool value = (Row < BoolColumn->RowValues.Num()) ? BoolColumn->RowValues[Row] : false;
+		return value ? ECheckBoxState::Checked : ECheckBoxState::Unchecked;
+	});
+}
+	
+TSharedRef<SWidget> CreateBoolPropertyWidget(UObject* Object, UClass* ContextClass)
+{
+	return CreatePropertyWidget<UChooserParameterBool_ContextProperty>(Object, ContextClass, GetDefault<UGraphEditorSettings>()->BooleanPinTypeColor);
+}
+	
+void RegisterBoolWidgets()
+{
+	FObjectChooserWidgetFactories::ChooserWidgetCreators.Add(UChooserParameterBool_ContextProperty::StaticClass(), CreateBoolPropertyWidget);
+	FObjectChooserWidgetFactories::ChooserTextConverter.Add(UChooserParameterBool_ContextProperty::StaticClass(), ConvertToText_ContextProperty<UChooserParameterBool_ContextProperty>);
+
+	FChooserTableEditor::ColumnWidgetCreators.Add(UChooserColumnBool::StaticClass(), CreateBoolColumnWidget);
+}
+	
+}
+
+#undef LOCTEXT_NAMESPACE
diff --git a/Engine/Plugins/Experimental/Chooser/Source/ChooserEditor/Private/BoolColumnEditor.h b/Engine/Plugins/Experimental/Chooser/Source/ChooserEditor/Private/BoolColumnEditor.h
new file mode 100644
--- /dev/null
+++ b/Engine/Plugins/Experimental/Chooser/Source/ChooserEditor/Private/BoolColumnEditor.h
@@ -0,0 +1,10 @@
+// Copyright Epic Games, Inc. All Rights Reserved.
+
+#pragma once
+
+#include "CoreMinimal.h"
+
+namespace UE::ChooserEditor
+{
+	void RegisterBoolWidgets();
+}
\ No newline at end of file
diff --git a/Engine/Plugins/Experimental/Chooser/Source/ChooserEditor/Private/ChooserEditorModule.cpp b/Engine/Plugins/Experimental/Chooser/Source/ChooserEditor/Private/ChooserEditorModule.cpp
--- a/Engine/Plugins/Experimental/Chooser/Source/ChooserEditor/Private/ChooserEditorModule.cpp
+++ b/Engine/Plugins/Experimental/Chooser/Source/ChooserEditor/Private/ChooserEditorModule.cpp
@@ -4,6 +4,8 @@
 #include "IAssetTools.h"
 #include "AssetTypeActions_Chooser.h"
 #include "ChooserTableEditor.h"
+#include "BoolColumnEditor.h"
+#include "FloatRangeColumnEditor.h"
 #include "GameplayTagColumnEditor.h"
 
 #define LOCTEXT_NAMESPACE "ChooserEditorModule"
@@ -18,6 +20,8 @@ void FModule::StartupModule()
 	AssetTools.RegisterAssetTypeActions(AssetTypeActions_ChooserTable.ToSharedRef());
 	FChooserTableEditor::RegisterWidgets();
 	RegisterGameplayTagWidgets();
+	RegisterFloatRangeWidgets();
+	RegisterBoolWidgets();
 	
 	InterfacePropertyTypeIdentifier = MakeShared<UE::ChooserEditor::FPropertyTypeIdentifier>();
 	
@@ -25,7 +29,7 @@ void FModule::StartupModule()
 	
 	PropertyModule.RegisterCustomPropertyTypeLayout(
 		FInterfaceProperty::StaticClass()->GetFName(),
-		FOnGetPropertyTypeCustomizationInstance::CreateLambda([] { return MakeShared<UE::ChooserEditor::FInterfacePropertyTypeCustomization>(UChooserParameterBool::StaticClass()); }),
+		FOnGetPropertyTypeCustomizationInstance::CreateLambda([] { return MakeShared<UE::ChooserEditor::FInterfacePropertyTypeCustomization>(); }),
 		InterfacePropertyTypeIdentifier);
 }
 		 
diff --git a/Engine/Plugins/Experimental/Chooser/Source/ChooserEditor/Private/ChooserTableEditor.cpp b/Engine/Plugins/Experimental/Chooser/Source/ChooserEditor/Private/ChooserTableEditor.cpp
--- a/Engine/Plugins/Experimental/Chooser/Source/ChooserEditor/Private/ChooserTableEditor.cpp
+++ b/Engine/Plugins/Experimental/Chooser/Source/ChooserEditor/Private/ChooserTableEditor.cpp
@@ -2,14 +2,12 @@
 
 #include "ChooserTableEditor.h"
 
-#include "ContextPropertyWidget.h"
 #include "Framework/Docking/TabManager.h"
 #include "Modules/ModuleManager.h"
 #include "Styling/AppStyle.h"
 #include "Widgets/Docking/SDockTab.h"
 #include "PropertyEditorModule.h"
 #include "IDetailsView.h"
-#include "Editor.h"
 #include "Framework/MultiBox/MultiBoxBuilder.h"
 #include "Widgets/Input/SHyperlink.h"
 #include "Widgets/Input/SNumericEntryBox.h"
@@ -39,29 +37,13 @@
 namespace UE::ChooserEditor
 {
 
-const FName FChooserTableEditor::ToolkitFName( TEXT( "GenericAssetEditor" ) );
+const FName FChooserTableEditor::ToolkitFName( TEXT( "ChooserTableEditor" ) );
 const FName FChooserTableEditor::PropertiesTabId( TEXT( "ChooserEditor_Properties" ) );
 const FName FChooserTableEditor::TableTabId( TEXT( "ChooserEditor_Table" ) );
 
-
-void ConvertObjectToText(UObject* Object, FText& OutText)
-{
-	UClass* Class = Object->GetClass();
-	while (Class)
-	{
-		if (auto TextConverter = FObjectChooserWidgetFactories::ChooserTextConverter.Find(Class))
-		{
-			(*TextConverter)(Object, OutText);
-			break;
-		}
-		Class = Class->GetSuperClass();
-	}
-}
-
-
 void FChooserTableEditor::RegisterTabSpawners(const TSharedRef<class FTabManager>& InTabManager)
 {
-	WorkspaceMenuCategory = InTabManager->AddLocalWorkspaceMenuCategory(LOCTEXT("WorkspaceMenu_GenericAssetEditor", "Asset Editor"));
+	WorkspaceMenuCategory = InTabManager->AddLocalWorkspaceMenuCategory(LOCTEXT("WorkspaceMenu_ChooserTableEditor", "Chooser Table Editor"));
 
 	FAssetEditorToolkit::RegisterTabSpawners(InTabManager);
 
@@ -80,6 +62,7 @@ void FChooserTableEditor::UnregisterTabSpawners(const TSharedRef<class FTabManag
 {
 	FAssetEditorToolkit::UnregisterTabSpawners(InTabManager);
 
+	InTabManager->UnregisterTabSpawner( TableTabId );
 	InTabManager->UnregisterTabSpawner( PropertiesTabId );
 }
 
@@ -87,7 +70,6 @@ const FName FChooserTableEditor::ChooserEditorAppIdentifier( TEXT( "ChooserEdito
 
 FChooserTableEditor::~FChooserTableEditor()
 {
-	GEditor->GetEditorSubsystem<UImportSubsystem>()->OnAssetPostImport.RemoveAll(this);
 	FCoreUObjectDelegates::OnObjectsReplaced.RemoveAll(this);
 	
 	DetailsView.Reset();
@@ -97,7 +79,6 @@ FChooserTableEditor::~FChooserTableEditor()
 void FChooserTableEditor::InitEditor( const EToolkitMode::Type Mode, const TSharedPtr< class IToolkitHost >& InitToolkitHost, const TArray<UObject*>& ObjectsToEdit, FGetDetailsViewObjects GetDetailsViewObjects )
 {
 	EditingObjects = ObjectsToEdit;
-	GEditor->GetEditorSubsystem<UImportSubsystem>()->OnAssetPostImport.AddSP(this, &FChooserTableEditor::HandleAssetPostImport);
 	FCoreUObjectDelegates::OnObjectsReplaced.AddSP(this, &FChooserTableEditor::OnObjectsReplaced);
 
 	FPropertyEditorModule& PropertyEditorModule = FModuleManager::GetModuleChecked<FPropertyEditorModule>( "PropertyEditor" );
@@ -133,8 +114,6 @@ void FChooserTableEditor::InitEditor( const EToolkitMode::Type Mode, const TShar
 
 	RegenerateMenusAndToolbars();
 
-	UChooserTable* Chooser = Cast<UChooserTable>(EditingObjects[0]);
-	
 	SelectRootProperties();
 }
 
@@ -328,7 +307,7 @@ public:
 		Operation->ChooserEditor = InEditor;
 		Operation->RowIndex = InRowIndex;
 		Operation->DefaultHoverText = LOCTEXT("Chooser Row", "Chooser Row");
-		ConvertObjectToText(InEditor->GetChooser()->Results[InRowIndex].GetObject(), Operation->DefaultHoverText);
+		FObjectChooserWidgetFactories::ConvertToText(InEditor->GetChooser()->Results[InRowIndex].GetObject(), Operation->DefaultHoverText);
 		Operation->CurrentHoverText = Operation->DefaultHoverText;
 			
 		Operation->Construct();
@@ -741,7 +720,6 @@ TSharedRef<SDockTab> FChooserTableEditor::SpawnTableTab( const FSpawnTabArgs& Ar
 					if (SelectedItem)
 					{
 						SelectedRows.SetNum(0);
-						// hack test details
 						UChooserTable* Chooser = Cast<UChooserTable>(EditingObjects[0]);
 						// Get the list of objects to edit the details of
 						TObjectPtr<UChooserRowDetails> Selection = NewObject<UChooserRowDetails>();
@@ -805,15 +783,6 @@ void FChooserTableEditor::UpdateTableRows()
 	}
 }
 
-void FChooserTableEditor::HandleAssetPostImport(UFactory* InFactory, UObject* InObject)
-{
-	if (EditingObjects.Contains(InObject))
-	{
-		// The details panel likely needs to be refreshed if an asset was imported again
-		DetailsView->SetObjects(EditingObjects);
-	}
-}
-
 void FChooserTableEditor::OnObjectsReplaced(const TMap<UObject*, UObject*>& ReplacementMap)
 {
 	bool bChangedAny = false;
@@ -839,7 +808,7 @@ void FChooserTableEditor::OnObjectsReplaced(const TMap<UObject*, UObject*>& Repl
 
 FString FChooserTableEditor::GetWorldCentricTabPrefix() const
 {
-	return LOCTEXT("WorldCentricTabPrefix", "Generic Asset ").ToString();
+	return LOCTEXT("WorldCentricTabPrefix", "Chooser Table Asset ").ToString();
 }
 
 TSharedRef<FChooserTableEditor> FChooserTableEditor::CreateEditor( const EToolkitMode::Type Mode, const TSharedPtr< IToolkitHost >& InitToolkitHost, UObject* ObjectToEdit, FGetDetailsViewObjects GetDetailsViewObjects )
@@ -860,62 +829,6 @@ TSharedRef<FChooserTableEditor> FChooserTableEditor::CreateEditor( const EToolki
 	return NewEditor;
 }
 
-void FChooserTableEditor::PostRegenerateMenusAndToolbars()
-{
-	// Find the common denominator class of the assets we're editing
-	TArray<UClass*> ClassList;
-	for (UObject* Obj : EditingObjects)
-	{
-		check(Obj);
-		ClassList.Add(Obj->GetClass());
-	}
-
-	UClass* CommonDenominatorClass = UClass::FindCommonBase(ClassList);
-	const bool bNotAllSame = (EditingObjects.Num() > 0) && (EditingObjects[0]->GetClass() != CommonDenominatorClass);
-
-	// Provide a hyperlink to view that native class
-	if (CommonDenominatorClass)
-	{
-		TWeakObjectPtr<UClass> WeakClassPtr(CommonDenominatorClass);
-		auto OnNavigateToClassCode = [WeakClassPtr]()
-		{
-			if (UClass* StrongClassPtr = WeakClassPtr.Get())
-			{
-				if (FSourceCodeNavigation::CanNavigateToClass(StrongClassPtr))
-				{
-					FSourceCodeNavigation::NavigateToClass(StrongClassPtr);
-				}
-			}
-		};
-
-		// build and attach the menu overlay
-		TSharedRef<SHorizontalBox> MenuOverlayBox = SNew(SHorizontalBox)
-			+ SHorizontalBox::Slot()
-			.AutoWidth()
-			.VAlign(VAlign_Center)
-			[
-				SNew(STextBlock)
-				.ColorAndOpacity(FSlateColor::UseSubduedForeground())
-				.ShadowOffset(FVector2D::UnitVector)
-				.Text(bNotAllSame ? LOCTEXT("ChooserTableEditor_AssetType_Varied", "Common Asset Type: ") : LOCTEXT("ChooserTableEditor_AssetType", "Asset Type: "))
-			]
-			+SHorizontalBox::Slot()
-			.AutoWidth()
-			.VAlign(VAlign_Center)
-			.Padding(0.0f, 0.0f, 8.0f, 0.0f)
-			[
-				SNew(SHyperlink)
-				.Style(FAppStyle::Get(), "Common.GotoNativeCodeHyperlink")
-				.OnNavigate_Lambda(OnNavigateToClassCode)
-				.Text(FText::FromName(CommonDenominatorClass->GetFName()))
-				.ToolTipText(FText::Format(LOCTEXT("GoToCode_ToolTip", "Click to open this source file in {0}"), FSourceCodeNavigation::GetSelectedSourceCodeIDE()))
-			];
-	
-		SetMenuOverlay(MenuOverlayBox);
-	}
-}
-
-	
 void FChooserTableEditor::SelectColumn(int Index)
 {
 	UChooserTable* Chooser = GetChooser();
@@ -950,7 +863,7 @@ TSharedRef<SWidget> CreateAssetWidget(UObject* Object, UClass* ContextClass)
 	UChooserTable* Chooser = Object->GetTypedOuter<UChooserTable>();
 	
 	return SNew(SObjectPropertyEntryBox)
-		.AllowedClass(Chooser->OutputObjectType!=nullptr ? Chooser->OutputObjectType.Get() : UObject::StaticClass())
+		.AllowedClass((Chooser!=nullptr && Chooser->OutputObjectType!=nullptr) ? Chooser->OutputObjectType.Get() : UObject::StaticClass())
 		.ObjectPath_Lambda([DIAsset](){ return DIAsset->Asset ? DIAsset->Asset.GetPath() : "";})
 		.OnObjectChanged_Lambda([DIAsset](const FAssetData& AssetData)
 		{
@@ -999,102 +912,11 @@ void ConvertToText_EvaluateChooser(const UObject* Object, FText& OutText)
 TSharedRef<SWidget> CreateObjectWidget(UObject* Object, UClass* ContextClass)
 {
 	FText ObjectText = FText::FromString(Object->GetName());
-	ConvertObjectToText(Object, ObjectText);
+	FObjectChooserWidgetFactories::ConvertToText(Object, ObjectText);
 	
 	return SNew (STextBlock).Text(ObjectText); // could make this use Text_Lambda, to have it update correctly?
 }
 
-TSharedRef<SWidget> CreateBoolColumnWidget(UObject* Column, int Row)
-{
-	UChooserColumnBool* BoolColumn = Cast<UChooserColumnBool>(Column);
-
-	return SNew (SCheckBox)
-	.OnCheckStateChanged_Lambda([BoolColumn,Row](ECheckBoxState State)
-	{
-		if (Row < BoolColumn->RowValues.Num())
-		{
-			const FScopedTransaction Transaction(LOCTEXT("Change Bool Value", "Change Bool Value"));
-			BoolColumn->Modify(true);
-			BoolColumn->RowValues[Row] = (State == ECheckBoxState::Checked);
-		}
-	})
-	.IsChecked_Lambda([BoolColumn, Row]()
-	{
-		const bool value = (Row < BoolColumn->RowValues.Num()) ? BoolColumn->RowValues[Row] : false;
-		return value ? ECheckBoxState::Checked : ECheckBoxState::Unchecked;
-	});
-}
-	
-TSharedRef<SWidget> CreateBoolPropertyWidget(UObject* Object, UClass* ContextClass)
-{
-	return CreatePropertyWidget<UChooserParameterBool_ContextProperty>(Object, ContextClass, GetDefault<UGraphEditorSettings>()->BooleanPinTypeColor);
-}
-	
-TSharedRef<SWidget> CreateFloatPropertyWidget(UObject* Object, UClass* ContextClass)
-{
-	return CreatePropertyWidget<UChooserParameterFloat_ContextProperty>(Object, ContextClass, GetDefault<UGraphEditorSettings>()->FloatPinTypeColor);
-}
-	
-TSharedRef<SWidget> CreateFloatRangeColumnWidget(UObject* Column, int Row)
-{
-	UChooserColumnFloatRange* FloatRangeColumn = Cast<UChooserColumnFloatRange>(Column);
-
-	return SNew(SHorizontalBox)
-	+ SHorizontalBox::Slot().MaxWidth(10)
-	[
-		SNew(STextBlock).Text(LOCTEXT("FloatRangeLeft", "("))
-	]
-	+ SHorizontalBox::Slot().FillWidth(1.0)
-	[
-		SNew(SNumericEntryBox<float>)
-		.MaxValue_Lambda([FloatRangeColumn, Row]()
-		{
-			return (Row < FloatRangeColumn->RowValues.Num()) ? FloatRangeColumn->RowValues[Row].Max : 0;
-		})
-		.Value_Lambda([FloatRangeColumn, Row]()
-		{
-			return (Row < FloatRangeColumn->RowValues.Num()) ? FloatRangeColumn->RowValues[Row].Min : 0;
-		})
-		.OnValueCommitted_Lambda([FloatRangeColumn, Row](float NewValue, ETextCommit::Type CommitType)
-		{
-			if (Row < FloatRangeColumn->RowValues.Num())
-			{
-				const FScopedTransaction Transaction(LOCTEXT("Edit Min Value", "Edit Min Value"));
-				FloatRangeColumn->Modify(true);
-				FloatRangeColumn->RowValues[Row].Min = NewValue;
-			}
-		})
-	]
-	+ SHorizontalBox::Slot().MaxWidth(10)
-	[
-		SNew(STextBlock).Text(LOCTEXT("FloatRangeComma", " ,"))
-	]
-	+ SHorizontalBox::Slot().FillWidth(1.0)
-	[
-		SNew(SNumericEntryBox<float>)
-		.MinValue_Lambda([FloatRangeColumn, Row]()
-		{
-			return (Row < FloatRangeColumn->RowValues.Num()) ? FloatRangeColumn->RowValues[Row].Min : 0;
-		})
-		.Value_Lambda([FloatRangeColumn, Row]()
-		{
-			return (Row < FloatRangeColumn->RowValues.Num()) ? FloatRangeColumn->RowValues[Row].Max : 0;
-		})
-		.OnValueCommitted_Lambda([FloatRangeColumn, Row](float NewValue, ETextCommit::Type CommitType)
-		{
-			if (Row < FloatRangeColumn->RowValues.Num())
-			{
-				const FScopedTransaction Transaction(LOCTEXT("Edit Max", "Edit Max Value"));
-				FloatRangeColumn->Modify(true);
-				FloatRangeColumn->RowValues[Row].Max = NewValue;
-			}
-		})
-	]
-	+ SHorizontalBox::Slot().MaxWidth(10)
-	[
-		SNew(STextBlock).Text(LOCTEXT("FloatRangeRight", " )"))
-	];
-}
 
 class FChooserRowDetails : public IDetailCustomization
 {
@@ -1150,7 +972,7 @@ void FChooserRowDetails::CustomizeDetails(IDetailLayoutBuilder& DetailBuilder)
 				FText DisplayName = LOCTEXT("No Input Value", "No Input Value");
 				if (UObject* InputValue = Column.GetInterface()->GetInputValue())
 				{
-					ConvertObjectToText(InputValue, DisplayName);
+					FObjectChooserWidgetFactories::ConvertToText(InputValue, DisplayName);
 				}
 				NewColumnProperty.DisplayName(DisplayName);
 				NewColumnProperty.ShowPropertyButtons(false); // hide array add button
@@ -1169,13 +991,6 @@ void FChooserTableEditor::RegisterWidgets()
 	FObjectChooserWidgetFactories::ChooserWidgetCreators.Add(UObjectChooser_Asset::StaticClass(), CreateAssetWidget);
 	FObjectChooserWidgetFactories::ChooserWidgetCreators.Add(UObjectChooser_EvaluateChooser::StaticClass(), CreateEvaluateChooserWidget);
 	
-	FObjectChooserWidgetFactories::ChooserWidgetCreators.Add(UChooserParameterBool_ContextProperty::StaticClass(), CreateBoolPropertyWidget);
-	FObjectChooserWidgetFactories::ChooserWidgetCreators.Add(UChooserParameterFloat_ContextProperty::StaticClass(), CreateFloatPropertyWidget);
-	FObjectChooserWidgetFactories::ChooserTextConverter.Add(UChooserParameterBool_ContextProperty::StaticClass(), ConvertToText_ContextProperty<UChooserParameterBool_ContextProperty>);
-	FObjectChooserWidgetFactories::ChooserTextConverter.Add(UChooserParameterFloat_ContextProperty::StaticClass(), ConvertToText_ContextProperty<UChooserParameterFloat_ContextProperty>);
-
-	ColumnWidgetCreators.Add(UChooserColumnBool::StaticClass(), CreateBoolColumnWidget);
-	ColumnWidgetCreators.Add(UChooserColumnFloatRange::StaticClass(), CreateFloatRangeColumnWidget);
 
 	FPropertyEditorModule& PropertyModule = FModuleManager::LoadModuleChecked<FPropertyEditorModule>("PropertyEditor");
 	PropertyModule.RegisterCustomClassLayout("ChooserRowDetails", FOnGetDetailCustomizationInstance::CreateStatic(&FChooserRowDetails::MakeInstance));	
diff --git a/Engine/Plugins/Experimental/Chooser/Source/ChooserEditor/Private/ChooserTableEditor.h b/Engine/Plugins/Experimental/Chooser/Source/ChooserEditor/Private/ChooserTableEditor.h
--- a/Engine/Plugins/Experimental/Chooser/Source/ChooserEditor/Private/ChooserTableEditor.h
+++ b/Engine/Plugins/Experimental/Chooser/Source/ChooserEditor/Private/ChooserTableEditor.h
@@ -71,9 +71,6 @@ namespace UE::ChooserEditor
 		virtual void NotifyPreChange( FProperty* PropertyAboutToChange ) override;
 		virtual void NotifyPostChange( const FPropertyChangedEvent& PropertyChangedEvent, FProperty* PropertyThatChanged) override;
 
-		/** FAssetEditorToolkit interface */
-		virtual void PostRegenerateMenusAndToolbars() override;
-
 		UChooserTable* GetChooser() { return Cast<UChooserTable>(EditingObjects[0]); }
 	
 		/** Used to show or hide certain properties */
@@ -101,9 +98,6 @@ namespace UE::ChooserEditor
 	
 		TSharedRef<ITableRow> GenerateTableRow(TSharedPtr<FChooserTableRow> InItem, const TSharedRef<STableViewBase>& OwnerTable);
 
-		/** Handles when an asset is imported */
-		void HandleAssetPostImport(class UFactory* InFactory, UObject* InObject);
-
 		/** Called when objects need to be swapped out for new versions, like after a blueprint recompile. */
 		void OnObjectsReplaced(const TMap<UObject*, UObject*>& ReplacementMap);
 
diff --git a/Engine/Plugins/Experimental/Chooser/Source/ChooserEditor/Private/FloatRangeColumnEditor.cpp b/Engine/Plugins/Experimental/Chooser/Source/ChooserEditor/Private/FloatRangeColumnEditor.cpp
new file mode 100644
--- /dev/null
+++ b/Engine/Plugins/Experimental/Chooser/Source/ChooserEditor/Private/FloatRangeColumnEditor.cpp
@@ -0,0 +1,92 @@
+// Copyright Epic Games, Inc. All Rights Reserved.
+
+#include "FloatRangeColumnEditor.h"
+#include "FloatRangeColumn.h"
+#include "ContextPropertyWidget.h"
+#include "ObjectChooserWidgetFactories.h"
+#include "ChooserTableEditor.h"
+#include "Widgets/Input/SNumericEntryBox.h"
+#include "GraphEditorSettings.h"
+
+#define LOCTEXT_NAMESPACE "FloatRangeColumnEditor"
+
+namespace UE::ChooserEditor
+{
+TSharedRef<SWidget> CreateFloatPropertyWidget(UObject* Object, UClass* ContextClass)
+{
+	return CreatePropertyWidget<UChooserParameterFloat_ContextProperty>(Object, ContextClass, GetDefault<UGraphEditorSettings>()->FloatPinTypeColor);
+}
+	
+TSharedRef<SWidget> CreateFloatRangeColumnWidget(UObject* Column, int Row)
+{
+	UChooserColumnFloatRange* FloatRangeColumn = Cast<UChooserColumnFloatRange>(Column);
+
+	return SNew(SHorizontalBox)
+	+ SHorizontalBox::Slot().MaxWidth(10)
+	[
+		SNew(STextBlock).Text(LOCTEXT("FloatRangeLeft", "("))
+	]
+	+ SHorizontalBox::Slot().FillWidth(1.0)
+	[
+		SNew(SNumericEntryBox<float>)
+		.MaxValue_Lambda([FloatRangeColumn, Row]()
+		{
+			return (Row < FloatRangeColumn->RowValues.Num()) ? FloatRangeColumn->RowValues[Row].Max : 0;
+		})
+		.Value_Lambda([FloatRangeColumn, Row]()
+		{
+			return (Row < FloatRangeColumn->RowValues.Num()) ? FloatRangeColumn->RowValues[Row].Min : 0;
+		})
+		.OnValueCommitted_Lambda([FloatRangeColumn, Row](float NewValue, ETextCommit::Type CommitType)
+		{
+			if (Row < FloatRangeColumn->RowValues.Num())
+			{
+				const FScopedTransaction Transaction(LOCTEXT("Edit Min Value", "Edit Min Value"));
+				FloatRangeColumn->Modify(true);
+				FloatRangeColumn->RowValues[Row].Min = NewValue;
+			}
+		})
+	]
+	+ SHorizontalBox::Slot().MaxWidth(10)
+	[
+		SNew(STextBlock).Text(LOCTEXT("FloatRangeComma", " ,"))
+	]
+	+ SHorizontalBox::Slot().FillWidth(1.0)
+	[
+		SNew(SNumericEntryBox<float>)
+		.MinValue_Lambda([FloatRangeColumn, Row]()
+		{
+			return (Row < FloatRangeColumn->RowValues.Num()) ? FloatRangeColumn->RowValues[Row].Min : 0;
+		})
+		.Value_Lambda([FloatRangeColumn, Row]()
+		{
+			return (Row < FloatRangeColumn->RowValues.Num()) ? FloatRangeColumn->RowValues[Row].Max : 0;
+		})
+		.OnValueCommitted_Lambda([FloatRangeColumn, Row](float NewValue, ETextCommit::Type CommitType)
+		{
+			if (Row < FloatRangeColumn->RowValues.Num())
+			{
+				const FScopedTransaction Transaction(LOCTEXT("Edit Max", "Edit Max Value"));
+				FloatRangeColumn->Modify(true);
+				FloatRangeColumn->RowValues[Row].Max = NewValue;
+			}
+		})
+	]
+	+ SHorizontalBox::Slot().MaxWidth(10)
+	[
+		SNew(STextBlock).Text(LOCTEXT("FloatRangeRight", " )"))
+	];
+}
+
+	
+void RegisterFloatRangeWidgets()
+{
+	FObjectChooserWidgetFactories::ChooserWidgetCreators.Add(UChooserParameterFloat_ContextProperty::StaticClass(), CreateFloatPropertyWidget);
+	FObjectChooserWidgetFactories::ChooserTextConverter.Add(UChooserParameterFloat_ContextProperty::StaticClass(), ConvertToText_ContextProperty<UChooserParameterFloat_ContextProperty>);
+
+	FChooserTableEditor::ColumnWidgetCreators.Add(UChooserColumnFloatRange::StaticClass(), CreateFloatRangeColumnWidget);
+}
+	
+}
+
+#undef LOCTEXT_NAMESPACE
diff --git a/Engine/Plugins/Experimental/Chooser/Source/ChooserEditor/Private/FloatRangeColumnEditor.h b/Engine/Plugins/Experimental/Chooser/Source/ChooserEditor/Private/FloatRangeColumnEditor.h
new file mode 100644
--- /dev/null
+++ b/Engine/Plugins/Experimental/Chooser/Source/ChooserEditor/Private/FloatRangeColumnEditor.h
@@ -0,0 +1,10 @@
+// Copyright Epic Games, Inc. All Rights Reserved.
+
+#pragma once
+
+#include "CoreMinimal.h"
+
+namespace UE::ChooserEditor
+{
+	void RegisterFloatRangeWidgets();
+}
\ No newline at end of file
diff --git a/Engine/Plugins/Experimental/Chooser/Source/ChooserEditor/Private/InterfacePropertyTypeCustomization.cpp b/Engine/Plugins/Experimental/Chooser/Source/ChooserEditor/Private/InterfacePropertyTypeCustomization.cpp
--- a/Engine/Plugins/Experimental/Chooser/Source/ChooserEditor/Private/InterfacePropertyTypeCustomization.cpp
+++ b/Engine/Plugins/Experimental/Chooser/Source/ChooserEditor/Private/InterfacePropertyTypeCustomization.cpp
@@ -41,7 +41,7 @@ void FInterfacePropertyTypeCustomization::CustomizeHeader(TSharedRef<IPropertyHa
 		}
 	}
 	
-	TSharedPtr<SWidget> Widget = FObjectChooserWidgetFactories::CreateWidget(InterfaceType, PropertyValue->GetObject(), ContextClass,
+	TSharedPtr<SWidget> Widget = FObjectChooserWidgetFactories::CreateWidget(IntProperty->InterfaceClass, PropertyValue->GetObject(), ContextClass,
 			FOnClassPicked::CreateLambda([this, PropertyHandle](UClass* ChosenClass)
 			{
 				TArray<void*> RawData;
diff --git a/Engine/Plugins/Experimental/Chooser/Source/ChooserEditor/Private/InterfacePropertyTypeCustomization.h b/Engine/Plugins/Experimental/Chooser/Source/ChooserEditor/Private/InterfacePropertyTypeCustomization.h
--- a/Engine/Plugins/Experimental/Chooser/Source/ChooserEditor/Private/InterfacePropertyTypeCustomization.h
+++ b/Engine/Plugins/Experimental/Chooser/Source/ChooserEditor/Private/InterfacePropertyTypeCustomization.h
@@ -11,13 +11,11 @@ namespace UE::ChooserEditor
 class FInterfacePropertyTypeCustomization : public IPropertyTypeCustomization
 {
 public:
-	FInterfacePropertyTypeCustomization(UClass* Interface) : InterfaceType(Interface) {}
+	FInterfacePropertyTypeCustomization(){}
 private:
 	// IPropertyTypeCustomization interface
 	virtual void CustomizeHeader(TSharedRef<IPropertyHandle> PropertyHandle, FDetailWidgetRow& HeaderRow, IPropertyTypeCustomizationUtils& CustomizationUtils) override;
 	virtual void CustomizeChildren(TSharedRef<IPropertyHandle> PropertyHandle, IDetailChildrenBuilder& ChildBuilder, IPropertyTypeCustomizationUtils& CustomizationUtils) override;
-
-	UClass* InterfaceType;
 };
 
 class FPropertyTypeIdentifier : public IPropertyTypeIdentifier
diff --git a/Engine/Plugins/Experimental/Chooser/Source/ChooserEditor/Private/ObjectChooserWidgetFactories.cpp b/Engine/Plugins/Experimental/Chooser/Source/ChooserEditor/Private/ObjectChooserWidgetFactories.cpp
--- a/Engine/Plugins/Experimental/Chooser/Source/ChooserEditor/Private/ObjectChooserWidgetFactories.cpp
+++ b/Engine/Plugins/Experimental/Chooser/Source/ChooserEditor/Private/ObjectChooserWidgetFactories.cpp
@@ -15,7 +15,7 @@
 namespace UE::ChooserEditor
 {
 
-void ConvertToText(UObject* Object, FText& OutText)
+void FObjectChooserWidgetFactories::ConvertToText(UObject* Object, FText& OutText)
 {
 	UClass* Class = Object->GetClass();
 	while (Class)
diff --git a/Engine/Plugins/Experimental/Chooser/Source/ChooserEditor/Public/ObjectChooserWidgetFactories.h b/Engine/Plugins/Experimental/Chooser/Source/ChooserEditor/Public/ObjectChooserWidgetFactories.h
--- a/Engine/Plugins/Experimental/Chooser/Source/ChooserEditor/Public/ObjectChooserWidgetFactories.h
+++ b/Engine/Plugins/Experimental/Chooser/Source/ChooserEditor/Public/ObjectChooserWidgetFactories.h
@@ -20,6 +20,7 @@ namespace UE::ChooserEditor
 
 		static TSharedPtr<SWidget> CreateWidget(UObject* Value, UClass* ContextClass);
 		static TSharedPtr<SWidget> CreateWidget(UClass* InterfaceType, UObject* Value, UClass* ContextClass, const FOnClassPicked& CreateClassCallback, TSharedPtr<SBorder>* InnerWidget = nullptr);
+		static void ConvertToText(UObject* Object, FText& OutText);
 		
 		static void RegisterWidgets();
 	};
