Chooser Enum Column

diff --git a/Engine/Plugins/Experimental/Chooser/Source/Chooser/Internal/BoolColumn.h b/Engine/Plugins/Experimental/Chooser/Source/Chooser/Internal/BoolColumn.h
--- a/Engine/Plugins/Experimental/Chooser/Source/Chooser/Internal/BoolColumn.h
+++ b/Engine/Plugins/Experimental/Chooser/Source/Chooser/Internal/BoolColumn.h
@@ -4,9 +4,9 @@
 #include "CoreMinimal.h"
 #include "IChooserColumn.h"
 #include "IChooserParameterBool.h"
+#include "ChooserPropertyAccess.h"
 #include "BoolColumn.generated.h"
 
-
 UCLASS()
 class CHOOSER_API UChooserParameterBool_ContextProperty :  public UObject, public IChooserParameterBool
 {
@@ -19,10 +19,15 @@ public:
 	virtual bool GetValue(const UObject* ContextObject, bool& OutResult) const override;
 
 #if WITH_EDITOR
-	static bool CanBind(const FString& TypeName)
+	static bool CanBind(const FProperty& Property)
 	{
 		static FString BoolTypeName = "bool";
-		return TypeName == BoolTypeName;
+		return Property.GetCPPType() == BoolTypeName;
+	}
+
+	void SetBinding(const TArray<FBindingChainElement>& InBindingChain)
+	{
+		UE::Chooser::CopyPropertyChain(InBindingChain, PropertyBindingChain);
 	}
 #endif
 };
diff --git a/Engine/Plugins/Experimental/Chooser/Source/Chooser/Internal/EnumColumn.h b/Engine/Plugins/Experimental/Chooser/Source/Chooser/Internal/EnumColumn.h
new file mode 100644
--- /dev/null
+++ b/Engine/Plugins/Experimental/Chooser/Source/Chooser/Internal/EnumColumn.h
@@ -0,0 +1,127 @@
+// Copyright Epic Games, Inc. All Rights Reserved.
+#pragma once
+
+#include "CoreMinimal.h"
+#include "IChooserColumn.h"
+#include "IChooserParameterEnum.h"
+#include "EnumColumn.generated.h"
+
+struct FBindingChainElement;
+
+UCLASS()
+class CHOOSER_API UChooserParameterEnum_ContextProperty : public UObject, public IChooserParameterEnum
+{
+	GENERATED_BODY()
+public:
+	UPROPERTY()
+	TArray<FName> PropertyBindingChain;
+
+	virtual bool GetValue(const UObject* ContextObject, uint8& OutResult) const override;
+
+#if WITH_EDITOR
+	static bool CanBind(const FProperty& Property)
+	{
+		if (Property.IsA<FEnumProperty>())
+		{
+			return true;
+		}
+
+		if (const FByteProperty* ByteProperty = CastField<FByteProperty>(&Property))
+		{
+			return ByteProperty->Enum != nullptr;
+		}
+
+		return false;
+	}
+
+	void SetBinding(const TArray<FBindingChainElement>& InBindingChain);
+
+	virtual const UEnum* GetEnum() const override { return Enum; }
+
+	virtual FSimpleMulticastDelegate& OnEnumChanged() { return EnumChanged; }
+#endif
+
+private:
+#if WITH_EDITORONLY_DATA
+	UPROPERTY()
+	TObjectPtr<const UEnum> Enum = nullptr;
+
+	FSimpleMulticastDelegate EnumChanged;
+#endif
+};
+
+
+UENUM()
+enum class EChooserEnumComparison : uint8
+{
+	Equal UMETA(DisplayName = "Value =="),
+	NotEqual UMETA(DisplayName = "Value !="),
+	GreaterThan UMETA(DisplayName = "Value >"),
+	GreaterThanEqual UMETA(DisplayName = "Value >="),
+	LessThan UMETA(DisplayName = "Value <"),
+	LessThanEqual UMETA(DisplayName = "Value <="),
+};
+
+USTRUCT()
+struct FChooserEnumRowData
+{
+	GENERATED_BODY()
+
+	UPROPERTY(EditAnywhere, Category = Runtime)
+	EChooserEnumComparison Comparison = EChooserEnumComparison::Equal;
+
+	UPROPERTY(EditAnywhere, Category = Runtime)
+	uint8 Value = 0;
+
+	bool Evaluate(const uint8 LeftHandSide) const;
+};
+
+UCLASS()
+class CHOOSER_API UChooserColumnEnum : public UObject, public IChooserColumn
+{
+	GENERATED_BODY()
+public:
+	UChooserColumnEnum() = default;
+	UChooserColumnEnum(const FObjectInitializer& ObjectInitializer);
+
+	UPROPERTY(EditAnywhere, Meta = (EditInlineInterface = "true"), Category = "Input")
+	TScriptInterface<IChooserParameterEnum> InputValue;
+
+	UPROPERTY(EditAnywhere, Category = Runtime)
+	// array of results (cells for this column for each row in the table)
+	// should match the length of the Results array
+	TArray<FChooserEnumRowData> RowValues;
+
+	virtual void Filter(const UObject* ContextObject, const TArray<uint32>& IndexListIn, TArray<uint32>& IndexListOut) override;
+	virtual void SetNumRows(uint32 NumRows) { RowValues.SetNum(NumRows); }
+	virtual void DeleteRows(const TArray<uint32>& RowIndices)
+	{
+		for (uint32 Index : RowIndices)
+		{
+			RowValues.RemoveAt(Index);
+		}
+	}
+	virtual UClass* GetInputValueInterface() { return UChooserParameterEnum::StaticClass(); };
+	virtual UObject* GetInputValue() override { return InputValue.GetObject(); };
+	virtual void SetInputValue(UObject* Value) override
+	{
+		InputValue = Value;
+#if WITH_EDITOR
+		InputChanged();
+#endif
+	};
+
+#if WITH_EDITOR
+	FSimpleMulticastDelegate OnEnumChanged;
+	void InputChanged()
+	{
+		InputValue->OnEnumChanged().AddLambda([this](){ OnEnumChanged.Broadcast(); });
+		OnEnumChanged.Broadcast();
+	}
+
+	virtual void PostLoad() override
+	{
+		InputChanged();
+	}
+#endif
+};
\ No newline at end of file
diff --git a/Engine/Plugins/Experimental/Chooser/Source/Chooser/Internal/FloatRangeColumn.h b/Engine/Plugins/Experimental/Chooser/Source/Chooser/Internal/FloatRangeColumn.h
--- a/Engine/Plugins/Experimental/Chooser/Source/Chooser/Internal/FloatRangeColumn.h
+++ b/Engine/Plugins/Experimental/Chooser/Source/Chooser/Internal/FloatRangeColumn.h
@@ -4,6 +4,7 @@
 #include "CoreMinimal.h"
 #include "IChooserColumn.h"
 #include "IChooserParameterFloat.h"
+#include "ChooserPropertyAccess.h"
 #include "FloatRangeColumn.generated.h"
 
 UCLASS()
@@ -17,12 +18,18 @@ public:
 	virtual bool GetValue(const UObject* ContextObject, float& OutResult) const override;
 
 #if WITH_EDITOR
-	static bool CanBind(const FString& TypeName)
+	static bool CanBind(const FProperty& Property)
 	{
 		static FString DoubleTypeName = "double";
 		static FString FloatTypeName = "float";
+		const FString& TypeName = Property.GetCPPType();
 		return TypeName == FloatTypeName || TypeName == DoubleTypeName;
 	}
+	
+	void SetBinding(const TArray<FBindingChainElement>& InBindingChain)
+	{
+		UE::Chooser::CopyPropertyChain(InBindingChain, PropertyBindingChain);
+	}
 #endif
 };
 
diff --git a/Engine/Plugins/Experimental/Chooser/Source/Chooser/Internal/GameplayTagColumn.h b/Engine/Plugins/Experimental/Chooser/Source/Chooser/Internal/GameplayTagColumn.h
--- a/Engine/Plugins/Experimental/Chooser/Source/Chooser/Internal/GameplayTagColumn.h
+++ b/Engine/Plugins/Experimental/Chooser/Source/Chooser/Internal/GameplayTagColumn.h
@@ -4,9 +4,11 @@
 #include "CoreMinimal.h"
 #include "IChooserColumn.h"
 #include "IChooserParameterGameplayTag.h"
+#include "ChooserPropertyAccess.h"
 #include "GameplayTagContainer.h"
 #include "GameplayTagColumn.generated.h"
 
+struct FBindingChainElement;
 
 UCLASS()
 class CHOOSER_API UChooserParameterGameplayTag_ContextProperty :  public UObject, public IChooserParameterGameplayTag
@@ -19,10 +21,15 @@ public:
 	virtual bool GetValue(const UObject* ContextObject, const FGameplayTagContainer*& OutResult) const override;
 
 #if WITH_EDITOR
-	static bool CanBind(const FString& TypeName)
+	static bool CanBind(const FProperty& Property)
 	{
-		static FString GameplayTagTypeName = "FGameplayTagContainer";
-		return TypeName == GameplayTagTypeName;
+		static FString TypeName = "FGameplayTagContainer";
+		return Property.GetCPPType() == TypeName;
+	}
+
+	void SetBinding(const TArray<FBindingChainElement>& InBindingChain)
+	{
+		UE::Chooser::CopyPropertyChain(InBindingChain, PropertyBindingChain);
 	}
 #endif
 };
diff --git a/Engine/Plugins/Experimental/Chooser/Source/Chooser/Private/ChooserPropertyAccess.cpp b/Engine/Plugins/Experimental/Chooser/Source/Chooser/Private/ChooserPropertyAccess.cpp
--- a/Engine/Plugins/Experimental/Chooser/Source/Chooser/Private/ChooserPropertyAccess.cpp
+++ b/Engine/Plugins/Experimental/Chooser/Source/Chooser/Private/ChooserPropertyAccess.cpp
@@ -3,6 +3,10 @@
 #include "ChooserPropertyAccess.h"
 #include "UObject/UnrealType.h"
 
+#if WITH_EDITOR
+#include "IPropertyAccessEditor.h"
+#endif
+
 namespace UE::Chooser
 {
 	bool ResolvePropertyChain(const void*& Container, UStruct*& StructType, const TArray<FName>& PropertyBindingChain)
@@ -37,4 +41,19 @@ namespace UE::Chooser
 	
 		return true;
 	}
+
+
+#if WITH_EDITOR
+	void CopyPropertyChain(const TArray<FBindingChainElement>& InBindingChain, TArray<FName>& OutPropertyBindingChain)
+	{
+		OutPropertyBindingChain.Empty();
+
+		for (int32 i = 1; i < InBindingChain.Num(); ++i)
+		{
+			OutPropertyBindingChain.Emplace(InBindingChain[i].Field.GetFName());
+		}
+	}
+#endif
+
+	
 }
\ No newline at end of file
diff --git a/Engine/Plugins/Experimental/Chooser/Source/Chooser/Private/EnumColumn.cpp b/Engine/Plugins/Experimental/Chooser/Source/Chooser/Private/EnumColumn.cpp
new file mode 100644
--- /dev/null
+++ b/Engine/Plugins/Experimental/Chooser/Source/Chooser/Private/EnumColumn.cpp
@@ -0,0 +1,124 @@
+// Copyright Epic Games, Inc. All Rights Reserved.
+#include "EnumColumn.h"
+#include "ChooserPropertyAccess.h"
+
+#if WITH_EDITOR
+#include "IPropertyAccessEditor.h"
+#endif
+
+bool UChooserParameterEnum_ContextProperty::GetValue(const UObject* ContextObject, uint8& OutResult) const
+{
+	UStruct* StructType = ContextObject->GetClass();
+	const void* Container = ContextObject;
+
+	if (UE::Chooser::ResolvePropertyChain(Container, StructType, PropertyBindingChain))
+	{
+		if (const FEnumProperty* EnumProperty = FindFProperty<FEnumProperty>(StructType, PropertyBindingChain.Last()))
+		{
+			OutResult = *EnumProperty->ContainerPtrToValuePtr<uint8>(Container);
+			return true;
+		}
+
+		if (const FByteProperty* ByteProperty = FindFProperty<FByteProperty>(StructType, PropertyBindingChain.Last()))
+		{
+			if (ByteProperty->IsEnum())
+			{
+				OutResult = *ByteProperty->ContainerPtrToValuePtr<uint8>(Container);
+				return true;
+			}
+		}
+	}
+
+	return false;
+}
+
+#if WITH_EDITOR
+
+void UChooserParameterEnum_ContextProperty::SetBinding(const TArray<FBindingChainElement>& InBindingChain)
+{
+	const UEnum* PreviousEnum = Enum;
+	Enum = nullptr;
+
+	UE::Chooser::CopyPropertyChain(InBindingChain, PropertyBindingChain);
+
+	const FField* Field = InBindingChain.Last().Field.ToField();
+	if (const FEnumProperty* EnumProperty = CastField<FEnumProperty>(Field))
+	{
+		Enum = EnumProperty->GetEnum();
+	}
+	else if (const FByteProperty* ByteProperty = CastField<FByteProperty>(Field))
+	{
+		Enum = ByteProperty->Enum;
+	}
+
+	if (Enum != PreviousEnum)
+	{
+		// Our enum type has changed! Need to refresh the UI to update enum value pickers.
+		EnumChanged.Broadcast();
+	}
+}
+
+#endif // WITH_EDITOR
+
+UChooserColumnEnum::UChooserColumnEnum(const FObjectInitializer& ObjectInitializer)
+{
+#if WITH_EDITOR
+	InputValue = ObjectInitializer.CreateDefaultSubobject<UChooserParameterEnum_ContextProperty>(this, "InputValue");
+	InputValue.GetObject()->SetFlags(RF_Transactional);
+	InputChanged();
+#endif
+}
+
+bool FChooserEnumRowData::Evaluate(const uint8 LeftHandSide) const
+{
+	switch (Comparison)
+	{
+	case EChooserEnumComparison::Equal:
+		return LeftHandSide == Value;
+
+	case EChooserEnumComparison::NotEqual:
+		return LeftHandSide != Value;
+
+	case EChooserEnumComparison::GreaterThan:
+		return LeftHandSide > Value;
+
+	case EChooserEnumComparison::GreaterThanEqual:
+		return LeftHandSide >= Value;
+
+	case EChooserEnumComparison::LessThan:
+		return LeftHandSide < Value;
+
+	case EChooserEnumComparison::LessThanEqual:
+		return LeftHandSide <= Value;
+
+	default:
+		checkf(false, TEXT("Unsupported comparison type for enum comparison operation!"));
+		return false;
+	}
+}
+
+void UChooserColumnEnum::Filter(const UObject* ContextObject, const TArray<uint32>& IndexListIn, TArray<uint32>& IndexListOut)
+{
+	uint8 Result = 0;
+	if (ContextObject != nullptr &&
+		InputValue != nullptr &&
+		InputValue->GetValue(ContextObject, Result))
+	{
+		for (const uint32 Index : IndexListIn)
+		{
+			if (RowValues.IsValidIndex(Index))
+			{
+				const FChooserEnumRowData& RowValue = RowValues[Index];
+				if (RowValue.Evaluate(Result))
+				{
+					IndexListOut.Emplace(Index);
+				}
+			}
+		}
+	}
+	else
+	{
+		// passthrough fallback (behaves better during live editing)
+		IndexListOut = IndexListIn;
+	}
+}
\ No newline at end of file
diff --git a/Engine/Plugins/Experimental/Chooser/Source/Chooser/Public/ChooserPropertyAccess.h b/Engine/Plugins/Experimental/Chooser/Source/Chooser/Public/ChooserPropertyAccess.h
--- a/Engine/Plugins/Experimental/Chooser/Source/Chooser/Public/ChooserPropertyAccess.h
+++ b/Engine/Plugins/Experimental/Chooser/Source/Chooser/Public/ChooserPropertyAccess.h
@@ -4,7 +4,14 @@
 
 #include "CoreMinimal.h"
 
+#if WITH_EDITOR
+struct FBindingChainElement;
+#endif
+
 namespace UE::Chooser
 {
-	bool ResolvePropertyChain(const void*& Container, UStruct*& StructType, const TArray<FName>& PropertyBindingChain);
+	CHOOSER_API bool ResolvePropertyChain(const void*& Container, UStruct*& StructType, const TArray<FName>& PropertyBindingChain);
+#if WITH_EDITOR
+	CHOOSER_API void CopyPropertyChain(const TArray<FBindingChainElement>& InBindingChain, TArray<FName>& OutPropertyBindingChain);
+#endif
 }
\ No newline at end of file
diff --git a/Engine/Plugins/Experimental/Chooser/Source/Chooser/Public/IChooserParameterEnum.h b/Engine/Plugins/Experimental/Chooser/Source/Chooser/Public/IChooserParameterEnum.h
new file mode 100644
--- /dev/null
+++ b/Engine/Plugins/Experimental/Chooser/Source/Chooser/Public/IChooserParameterEnum.h
@@ -0,0 +1,27 @@
+// Copyright Epic Games, Inc. All Rights Reserved.
+
+#pragma once
+
+#include "CoreMinimal.h"
+#include "UObject/Object.h"
+#include "UObject/Interface.h"
+#include "IChooserParameterEnum.generated.h"
+
+UINTERFACE(NotBlueprintType, meta = (CannotImplementInterfaceInBlueprint))
+class CHOOSER_API UChooserParameterEnum : public UInterface
+{
+	GENERATED_BODY()
+};
+
+class CHOOSER_API IChooserParameterEnum
+{
+	GENERATED_BODY()
+
+public:
+	virtual bool GetValue(const UObject* ContextObject, uint8& OutResult) const { return false; }
+
+#if WITH_EDITOR
+	virtual const UEnum* GetEnum() const { return nullptr; }
+	virtual FSimpleMulticastDelegate& OnEnumChanged() = 0;
+#endif
+};
\ No newline at end of file
diff --git a/Engine/Plugins/Experimental/Chooser/Source/ChooserEditor/Private/ChooserEditorModule.cpp b/Engine/Plugins/Experimental/Chooser/Source/ChooserEditor/Private/ChooserEditorModule.cpp
--- a/Engine/Plugins/Experimental/Chooser/Source/ChooserEditor/Private/ChooserEditorModule.cpp
+++ b/Engine/Plugins/Experimental/Chooser/Source/ChooserEditor/Private/ChooserEditorModule.cpp
@@ -5,6 +5,7 @@
 #include "AssetTypeActions_Chooser.h"
 #include "ChooserTableEditor.h"
 #include "BoolColumnEditor.h"
+#include "EnumColumnEditor.h"
 #include "FloatRangeColumnEditor.h"
 #include "GameplayTagColumnEditor.h"
 
@@ -22,6 +23,7 @@ void FModule::StartupModule()
 	RegisterGameplayTagWidgets();
 	RegisterFloatRangeWidgets();
 	RegisterBoolWidgets();
+	RegisterEnumWidgets();
 	
 	InterfacePropertyTypeIdentifier = MakeShared<UE::ChooserEditor::FPropertyTypeIdentifier>();
 	
diff --git a/Engine/Plugins/Experimental/Chooser/Source/ChooserEditor/Private/ChooserTableEditor.cpp b/Engine/Plugins/Experimental/Chooser/Source/ChooserEditor/Private/ChooserTableEditor.cpp
--- a/Engine/Plugins/Experimental/Chooser/Source/ChooserEditor/Private/ChooserTableEditor.cpp
+++ b/Engine/Plugins/Experimental/Chooser/Source/ChooserEditor/Private/ChooserTableEditor.cpp
@@ -512,9 +512,8 @@ void FChooserTableEditor::UpdateTableColumns()
 
 				MenuBuilder.AddMenuEntry(LOCTEXT("Column Properties","Properties"),LOCTEXT("Delete Column ToolTip", "Remove this column and all its data from the table"),FSlateIcon(),
 					FUIAction(
-						FExecuteAction::CreateLambda([this,ColumnIndex, ColumnId, &Column]()
+						FExecuteAction::CreateLambda([this,Chooser,ColumnIndex, ColumnId, &Column]()
 						{
-							UChooserTable* Chooser = Cast<UChooserTable>(EditingObjects[0]);
 							DetailsView->SetObject(Chooser->Columns[ColumnIndex].GetObject());
 							SelectedColumn = ColumnId;
 						})
@@ -523,10 +522,9 @@ void FChooserTableEditor::UpdateTableColumns()
 
 				MenuBuilder.AddMenuEntry(LOCTEXT("Delete Column","Delete"),LOCTEXT("Delete Column ToolTip", "Remove this column and all its data from the table"),FSlateIcon(),
 					FUIAction(
-						FExecuteAction::CreateLambda([this,ColumnIndex, &Column]()
+						FExecuteAction::CreateLambda([this,Chooser,ColumnIndex, &Column]()
 						{
 							const FScopedTransaction Transaction(LOCTEXT("Delete Column Transaction", "Delete Column"));
-							UChooserTable* Chooser = Cast<UChooserTable>(EditingObjects[0]);
 							Chooser->Modify(true);
 							Chooser->Columns.RemoveAt(ColumnIndex);
 							UpdateTableColumns();
diff --git a/Engine/Plugins/Experimental/Chooser/Source/ChooserEditor/Private/EnumColumnEditor.cpp b/Engine/Plugins/Experimental/Chooser/Source/ChooserEditor/Private/EnumColumnEditor.cpp
new file mode 100644
--- /dev/null
+++ b/Engine/Plugins/Experimental/Chooser/Source/ChooserEditor/Private/EnumColumnEditor.cpp
@@ -0,0 +1,137 @@
+// Copyright Epic Games, Inc. All Rights Reserved.
+
+#include "EnumColumnEditor.h"
+#include "EnumColumn.h"
+#include "ContextPropertyWidget.h"
+#include "ObjectChooserWidgetFactories.h"
+#include "ChooserTableEditor.h"
+#include "GraphEditorSettings.h"
+#include "SEnumCombo.h"
+#include "../../../../../Developer/TextureFormatOodle/Sdks/2.9.8/include/oodle2tex.h"
+
+#define LOCTEXT_NAMESPACE "EnumColumnEditor"
+
+namespace UE::ChooserEditor
+{
+	// Wrapper widget for EnumComboBox which will reconstruct the combo box when the Enum has changed
+	class SEnumCell : public SCompoundWidget
+	{
+	public:
+		SLATE_BEGIN_ARGS(SEnumCell)
+			: _RowIndex(-1)
+		{}
+
+		SLATE_ARGUMENT(UChooserColumnEnum*, EnumColumn)
+		SLATE_ATTRIBUTE(int, RowIndex);
+                
+		SLATE_END_ARGS()
+
+		TSharedRef<SWidget> CreateEnumComboBox()
+		{
+			if (const UChooserColumnEnum* EnumColumnPointer = EnumColumn.Get())
+			{
+				if (EnumColumnPointer->InputValue)
+				{
+					if (const UEnum* Enum = EnumColumnPointer->InputValue->GetEnum())
+					{
+						return SNew(SEnumComboBox, Enum)
+							.CurrentValue_Lambda([this]()
+							{
+								int Row = RowIndex.Get();
+								return EnumColumn.Get()->RowValues.IsValidIndex(Row) ? static_cast<int32>(EnumColumn->RowValues[Row].Value) : 0;
+							})
+							.OnEnumSelectionChanged_Lambda([this](int32 EnumValue, ESelectInfo::Type)
+							{
+								int Row = RowIndex.Get();
+								if (EnumColumn->RowValues.IsValidIndex(Row))
+								{
+									const FScopedTransaction Transaction(LOCTEXT("Edit RHS", "Edit Enum Value"));
+									EnumColumn->Modify(true);
+									EnumColumn->RowValues[Row].Value = static_cast<uint8>(EnumValue);
+								}
+							});
+					}
+				}
+			}
+			
+			return SNullWidget::NullWidget;
+		}
+
+		void UpdateEnumComboBox()
+		{
+			ChildSlot[ CreateEnumComboBox()	];
+		}
+
+		void Construct( const FArguments& InArgs)
+		{
+			RowIndex = InArgs._RowIndex;
+			EnumColumn = InArgs._EnumColumn;
+
+			if (UChooserColumnEnum* EnumColumnPointer = EnumColumn.Get())
+			{
+				EnumChangedHandle = EnumColumnPointer->OnEnumChanged.AddSP(this, &SEnumCell::UpdateEnumComboBox);
+			}
+			
+			UpdateEnumComboBox();
+		}
+
+		~SEnumCell()
+		{
+			if (UChooserColumnEnum* EnumColumnPointer = EnumColumn.Get())
+			{
+				EnumColumnPointer->OnEnumChanged.Remove(EnumChangedHandle);
+			}
+		}
+
+	private:
+		TWeakObjectPtr<UChooserColumnEnum> EnumColumn;
+		TAttribute<int> RowIndex;
+		FDelegateHandle EnumChangedHandle;
+	};
+
+	TSharedRef<SWidget> CreateEnumColumnWidget(UObject* Column, int Row)
+	{
+		UChooserColumnEnum* EnumColumn = CastChecked<UChooserColumnEnum>(Column);
+		
+		return SNew(SHorizontalBox)
+		+ SHorizontalBox::Slot().MaxWidth(95.0f)
+		[
+			SNew(SEnumComboBox, StaticEnum<EChooserEnumComparison>())
+			.CurrentValue_Lambda([EnumColumn, Row]()
+			{
+				return EnumColumn->RowValues.IsValidIndex(Row) ? static_cast<int32>(EnumColumn->RowValues[Row].Comparison) : 0;
+			})
+			.OnEnumSelectionChanged_Lambda([EnumColumn, Row](int32 EnumValue, ESelectInfo::Type)
+			{
+				if (EnumColumn->RowValues.IsValidIndex(Row))
+				{
+					const FScopedTransaction Transaction(LOCTEXT("Edit Comparison", "Edit Comparison Operation"));
+					EnumColumn->Modify(true);
+					EnumColumn->RowValues[Row].Comparison = static_cast<EChooserEnumComparison>(EnumValue);
+				}
+			})
+		]
+		+ SHorizontalBox::Slot()
+		[
+			SNew(SEnumCell).EnumColumn(EnumColumn).RowIndex(Row)
+		];
+	}
+
+	
+
+TSharedRef<SWidget> CreateEnumPropertyWidget(UObject* Object, UClass* ContextClass)
+{
+	return CreatePropertyWidget<UChooserParameterEnum_ContextProperty>(Object, ContextClass, GetDefault<UGraphEditorSettings>()->BytePinTypeColor);
+}
+	
+void RegisterEnumWidgets()
+{
+	FObjectChooserWidgetFactories::ChooserWidgetCreators.Add(UChooserParameterEnum_ContextProperty::StaticClass(), CreateEnumPropertyWidget);
+	FObjectChooserWidgetFactories::ChooserTextConverter.Add(UChooserParameterEnum_ContextProperty::StaticClass(), ConvertToText_ContextProperty<UChooserParameterEnum_ContextProperty>);
+
+	FChooserTableEditor::ColumnWidgetCreators.Add(UChooserColumnEnum::StaticClass(), CreateEnumColumnWidget);
+}
+	
+}
+
+#undef LOCTEXT_NAMESPACE
diff --git a/Engine/Plugins/Experimental/Chooser/Source/ChooserEditor/Private/EnumColumnEditor.h b/Engine/Plugins/Experimental/Chooser/Source/ChooserEditor/Private/EnumColumnEditor.h
new file mode 100644
--- /dev/null
+++ b/Engine/Plugins/Experimental/Chooser/Source/ChooserEditor/Private/EnumColumnEditor.h
@@ -0,0 +1,10 @@
+// Copyright Epic Games, Inc. All Rights Reserved.
+
+#pragma once
+
+#include "CoreMinimal.h"
+
+namespace UE::ChooserEditor
+{
+	void RegisterEnumWidgets();
+}
\ No newline at end of file
diff --git a/Engine/Plugins/Experimental/Chooser/Source/ChooserEditor/Public/ContextPropertyWidget.h b/Engine/Plugins/Experimental/Chooser/Source/ChooserEditor/Public/ContextPropertyWidget.h
--- a/Engine/Plugins/Experimental/Chooser/Source/ChooserEditor/Public/ContextPropertyWidget.h
+++ b/Engine/Plugins/Experimental/Chooser/Source/ChooserEditor/Public/ContextPropertyWidget.h
@@ -19,7 +19,7 @@ namespace UE::ChooserEditor
     	
     	Args.OnCanBindProperty = FOnCanBindProperty::CreateLambda([](FProperty* Property)
     	{
-    		return Property == nullptr || PropertyType::CanBind(Property->GetCPPType());
+    		return Property == nullptr || PropertyType::CanBind(*Property);
     	});
     	Args.OnCanBindToClass = FOnCanBindToClass::CreateLambda([ContextClass ](UClass* InClass)
     	{
@@ -42,25 +42,11 @@ namespace UE::ChooserEditor
     
     	Args.OnAddBinding = FOnAddBinding::CreateLambda([ContextProperty](FName InPropertyName, const TArray<FBindingChainElement>& InBindingChain)
     		{
-    			IPropertyAccessEditor& PropertyAccessEditor = IModularFeatures::Get().GetModularFeature<IPropertyAccessEditor>("PropertyAccessEditor");
     			if (ContextProperty.IsValid())
     			{
     				const FScopedTransaction Transaction(NSLOCTEXT("ContextPropertyWidget", "Change Property Binding", "Change Property Binding"));
     				ContextProperty->Modify(true);
-    				
-    				ContextProperty->PropertyBindingChain.Empty();
-    
-    				if (InBindingChain.Num() > 1)
-    				{
-    					for(int i = 1; i<InBindingChain.Num(); i++)
-    					{
-    						ContextProperty->PropertyBindingChain.Add(InBindingChain[i].Field.GetFName());
-    					}
-    				}
-    				else
-    				{
-    					ContextProperty->PropertyBindingChain.Empty();
-    				}
+    				ContextProperty->SetBinding(InBindingChain);
     			}
     		});
     
