Removing unnecessary and incorrect include

diff --git a/Engine/Plugins/Experimental/Chooser/Source/ChooserEditor/Private/EnumColumnEditor.cpp b/Engine/Plugins/Experimental/Chooser/Source/ChooserEditor/Private/EnumColumnEditor.cpp
--- a/Engine/Plugins/Experimental/Chooser/Source/ChooserEditor/Private/EnumColumnEditor.cpp
+++ b/Engine/Plugins/Experimental/Chooser/Source/ChooserEditor/Private/EnumColumnEditor.cpp
@@ -7,7 +7,6 @@
 #include "ChooserTableEditor.h"
 #include "GraphEditorSettings.h"
 #include "SEnumCombo.h"
-#include "../../../../../Developer/TextureFormatOodle/Sdks/2.9.8/include/oodle2tex.h"
 
 #define LOCTEXT_NAMESPACE "EnumColumnEditor"
 
