Proxy Table prototype

diff --git a/Engine/Plugins/Experimental/Chooser/Chooser.uplugin b/Engine/Plugins/Experimental/Chooser/Chooser.uplugin
--- a/Engine/Plugins/Experimental/Chooser/Chooser.uplugin
+++ b/Engine/Plugins/Experimental/Chooser/Chooser.uplugin
@@ -12,7 +12,7 @@
 	"SupportURL" : "",
 	"EnabledByDefault" : false,
 	"CanContainContent" : false,
-	"IsBetaVersion" : true,
+	"IsBetaVersion" : false,
 	"IsExperimentalVersion" : true,
 	"Installed" : false,
 	"Plugins": [
@@ -32,6 +32,16 @@
 			"Name" : "ChooserEditor",
 			"Type" : "Editor",
 			"LoadingPhase" : "Default"
-		}
+		},
+		{
+			"Name" : "ProxyTable",
+			"Type" : "Runtime",
+			"LoadingPhase" : "Default"
+		},
+		{
+			"Name" : "ProxyTableEditor",
+			"Type" : "Editor",
+			"LoadingPhase" : "Default"
+		},
 	]
 }
diff --git a/Engine/Plugins/Experimental/Chooser/Source/ProxyTable/Internal/ProxyTable.h b/Engine/Plugins/Experimental/Chooser/Source/ProxyTable/Internal/ProxyTable.h
new file mode 100644
--- /dev/null
+++ b/Engine/Plugins/Experimental/Chooser/Source/ProxyTable/Internal/ProxyTable.h
@@ -0,0 +1,81 @@
+// Copyright Epic Games, Inc. All Rights Reserved.
+#pragma once
+
+#include "CoreMinimal.h"
+#include "UObject/Object.h"
+#include "IObjectChooser.h"
+#include "ChooserPropertyAccess.h"
+#include "IChooserParameterProxyTable.h"
+
+#include "ProxyTable.generated.h"
+
+
+USTRUCT()
+struct FProxyEntry
+{
+	GENERATED_BODY()
+	UPROPERTY(EditAnywhere, Category = "Data")
+	FName Key;
+	UPROPERTY(EditAnywhere, Meta = (EditInlineInterface = "true"), Category = "Data")
+	TScriptInterface<IObjectChooser> Value;
+};
+
+UCLASS(MinimalAPI,BlueprintType)
+class UProxyTable : public UObject
+{
+	GENERATED_UCLASS_BODY()
+public:
+	UProxyTable() {}
+
+	UPROPERTY(EditAnywhere, Category = "Hidden")
+	TArray<FProxyEntry> Entries;
+	
+	UPROPERTY(EditAnywhere, Category = "Inheritance")
+	TArray<TObjectPtr<UProxyTable>> InheritEntriesFrom;
+};
+
+
+struct FBindingChainElement;
+
+UCLASS()
+class PROXYTABLE_API UChooserParameterProxyTable_ContextProperty :  public UObject, public IChooserParameterProxyTable
+{
+	GENERATED_BODY()
+public:
+
+	UPROPERTY()
+	TArray<FName> PropertyBindingChain;
+	
+	virtual bool GetValue(const UObject* ContextObject, const UProxyTable*& OutResult) const override;
+
+#if WITH_EDITOR
+	static bool CanBind(const FProperty& Property)
+	{
+		return Property.GetCPPType() == "UProxyTable*";
+	}
+
+	void SetBinding(const TArray<FBindingChainElement>& InBindingChain)
+	{
+		UE::Chooser::CopyPropertyChain(InBindingChain, PropertyBindingChain);
+	}
+#endif
+};
+
+
+
+UCLASS()
+class PROXYTABLE_API UObjectChooser_LookupProxy : public UObject, public IObjectChooser
+{
+	GENERATED_BODY()
+	virtual UObject* ChooseObject(const UObject* ContextObject) const final override;
+	
+	public:
+	UObjectChooser_LookupProxy() {};
+	UObjectChooser_LookupProxy(const FObjectInitializer& ObjectInitializer);
+	
+	UPROPERTY(EditAnywhere, Meta = (EditInlineInterface = "true"), Category = "Input")
+	TScriptInterface<IChooserParameterProxyTable> ProxyTable;
+	
+	UPROPERTY(EditAnywhere, Category="Parameters")
+	FName Key;
+};
diff --git a/Engine/Plugins/Experimental/Chooser/Source/ProxyTable/Private/ProxyTable.cpp b/Engine/Plugins/Experimental/Chooser/Source/ProxyTable/Private/ProxyTable.cpp
new file mode 100644
--- /dev/null
+++ b/Engine/Plugins/Experimental/Chooser/Source/ProxyTable/Private/ProxyTable.cpp
@@ -0,0 +1,77 @@
+// Copyright Epic Games, Inc. All Rights Reserved.
+#include "ProxyTable.h"
+
+UObjectChooser_LookupProxy::UObjectChooser_LookupProxy(const FObjectInitializer& ObjectInitializer)
+{
+	// ProxyTable = ObjectInitializer.CreateDefaultSubobject<UChooserParameterProxyTable_ContextProperty>(this, "InputValue");
+	// ProxyTable.GetObject()->SetFlags(RF_Transactional);
+}
+
+static UObject* FindProxyObject(const UProxyTable* Table, FName Key, const UObject* ContextObject)
+{
+	for (const FProxyEntry& Entry : Table->Entries)
+	{
+		if (Entry.Key == Key)
+		{
+			if (const IObjectChooser* EntryValue = Entry.Value.GetInterface())
+			{
+				return EntryValue->ChooseObject(ContextObject);
+			}
+		}
+	}
+
+	// search parent tables (uncooked data only)
+	for (const TObjectPtr<UProxyTable> ParentTable : Table->InheritEntriesFrom)
+	{
+		if (UObject* Value = FindProxyObject(ParentTable, Key, ContextObject))
+		{
+			return Value;
+		}
+	}
+
+	return nullptr;
+}
+
+UObject* UObjectChooser_LookupProxy::ChooseObject(const UObject* ContextObject) const
+{
+	if (ProxyTable)
+	{
+		const UProxyTable* Table;
+		if (ProxyTable->GetValue(ContextObject, Table))
+		{
+			if (Table)
+			{
+				if (UObject* Value = FindProxyObject(Table, Key, ContextObject))
+				{
+					return Value;
+				}
+			}
+		}
+				
+	}
+	
+	return nullptr;
+}
+
+UProxyTable::UProxyTable(const FObjectInitializer& Initializer)
+	:Super(Initializer)
+{
+
+}
+
+bool UChooserParameterProxyTable_ContextProperty::GetValue(const UObject* ContextObject, const UProxyTable*& OutResult) const
+{
+	UStruct* StructType = ContextObject->GetClass();
+	const void* Container = ContextObject;
+	
+	if (UE::Chooser::ResolvePropertyChain(Container, StructType, PropertyBindingChain))
+	{
+		if (const FObjectProperty* Property = FindFProperty<FObjectProperty>(StructType, PropertyBindingChain.Last()))
+		{
+			OutResult = *Property->ContainerPtrToValuePtr<UProxyTable*>(Container);
+			return true;
+		}
+	}
+
+	return false;
+}
\ No newline at end of file
diff --git a/Engine/Plugins/Experimental/Chooser/Source/ProxyTable/Private/ProxyTableModule.cpp b/Engine/Plugins/Experimental/Chooser/Source/ProxyTable/Private/ProxyTableModule.cpp
new file mode 100644
--- /dev/null
+++ b/Engine/Plugins/Experimental/Chooser/Source/ProxyTable/Private/ProxyTableModule.cpp
@@ -0,0 +1,5 @@
+// Copyright Epic Games, Inc. All Rights Reserved.
+
+#include "Modules/ModuleManager.h"
+
+IMPLEMENT_MODULE(FDefaultModuleImpl, ProxyTableModule);
\ No newline at end of file
diff --git a/Engine/Plugins/Experimental/Chooser/Source/ProxyTable/ProxyTable.Build.cs b/Engine/Plugins/Experimental/Chooser/Source/ProxyTable/ProxyTable.Build.cs
new file mode 100644
--- /dev/null
+++ b/Engine/Plugins/Experimental/Chooser/Source/ProxyTable/ProxyTable.Build.cs
@@ -0,0 +1,28 @@
+// Copyright Epic Games, Inc. All Rights Reserved.
+
+namespace UnrealBuildTool.Rules
+{
+	public class ProxyTable : ModuleRules
+	{
+		public ProxyTable(ReadOnlyTargetRules Target) : base(Target)
+		{			
+			PublicDependencyModuleNames.AddRange(
+				new string[]
+				{
+					// ... add other public dependencies that you statically link with here ...
+				}
+			);
+
+			PrivateDependencyModuleNames.AddRange(
+				new string[]
+				{
+					"Core",
+					"CoreUObject",
+					"Engine",
+					"Chooser"
+					// ... add private dependencies that you statically link with here ...
+				}
+			);
+		}
+	}
+}
\ No newline at end of file
diff --git a/Engine/Plugins/Experimental/Chooser/Source/ProxyTable/Public/IChooserParameterProxyTable.h b/Engine/Plugins/Experimental/Chooser/Source/ProxyTable/Public/IChooserParameterProxyTable.h
new file mode 100644
--- /dev/null
+++ b/Engine/Plugins/Experimental/Chooser/Source/ProxyTable/Public/IChooserParameterProxyTable.h
@@ -0,0 +1,24 @@
+// Copyright Epic Games, Inc. All Rights Reserved.
+
+#pragma once
+
+#include "CoreMinimal.h"
+#include "UObject/Object.h"
+#include "UObject/Interface.h"
+#include "IChooserParameterProxyTable.generated.h"
+
+class UProxyTable;
+
+UINTERFACE(NotBlueprintType, meta = (CannotImplementInterfaceInBlueprint))
+class PROXYTABLE_API UChooserParameterProxyTable : public UInterface
+{
+	GENERATED_BODY()
+};
+
+class PROXYTABLE_API IChooserParameterProxyTable
+{
+	GENERATED_BODY()
+
+public:
+	virtual bool GetValue(const UObject* ContextObject, const UProxyTable*& OutResult) const { return false; }
+};
\ No newline at end of file
diff --git a/Engine/Plugins/Experimental/Chooser/Source/ProxyTableEditor/Private/AssetTypeActions_ProxyTable.cpp b/Engine/Plugins/Experimental/Chooser/Source/ProxyTableEditor/Private/AssetTypeActions_ProxyTable.cpp
new file mode 100644
--- /dev/null
+++ b/Engine/Plugins/Experimental/Chooser/Source/ProxyTableEditor/Private/AssetTypeActions_ProxyTable.cpp
@@ -0,0 +1,27 @@
+// Copyright Epic Games, Inc. All Rights Reserved.
+
+#include "AssetTypeActions_ProxyTable.h"
+#include "ProxyTableEditor.h"
+
+#define LOCTEXT_NAMESPACE "AssetTypeActions_ProxyTable" 
+
+namespace UE::ProxyTableEditor
+{
+
+void FAssetTypeActions_ProxyTable::OpenAssetEditor(const TArray<UObject*>& InObjects, TSharedPtr<class IToolkitHost> EditWithinLevelEditor)
+{
+	FProxyTableEditor::CreateEditor(EToolkitMode::Standalone, EditWithinLevelEditor, InObjects);
+}
+
+const TArray<FText>& FAssetTypeActions_ProxyTable::GetSubMenus() const
+{
+	static const TArray<FText> SubMenus
+	{
+		LOCTEXT("DataInterfacesSubMenu", "Data Interfaces")
+	};
+	return SubMenus;
+}
+
+}
+
+#undef LOCTEXT_NAMESPACE
\ No newline at end of file
diff --git a/Engine/Plugins/Experimental/Chooser/Source/ProxyTableEditor/Private/AssetTypeActions_ProxyTable.h b/Engine/Plugins/Experimental/Chooser/Source/ProxyTableEditor/Private/AssetTypeActions_ProxyTable.h
new file mode 100644
--- /dev/null
+++ b/Engine/Plugins/Experimental/Chooser/Source/ProxyTableEditor/Private/AssetTypeActions_ProxyTable.h
@@ -0,0 +1,24 @@
+// Copyright Epic Games, Inc. All Rights Reserved.
+
+#pragma once
+
+#include "AssetTypeActions_Base.h"
+#include "ProxyTable.h"
+
+namespace UE::ProxyTableEditor
+{
+	
+class FAssetTypeActions_ProxyTable : public FAssetTypeActions_Base
+{
+public:
+	// IAssetTypeActions Implementation
+	virtual FText GetName() const override { return NSLOCTEXT("AssetTypeActions", "AssetTypeActions_ProxyTable", "Proxy Table"); }
+	virtual FColor GetTypeColor() const override { return FColor(128,64,64); }
+	virtual UClass* GetSupportedClass() const override { return UProxyTable::StaticClass(); }
+	virtual void OpenAssetEditor( const TArray<UObject*>& InObjects, TSharedPtr<class IToolkitHost> EditWithinLevelEditor = TSharedPtr<IToolkitHost>() ) override;
+	virtual bool CanFilter() override { return true; }
+	virtual const TArray<FText>& GetSubMenus() const override;
+	virtual uint32 GetCategories() override { return EAssetTypeCategories::Animation; }
+};
+
+}
diff --git a/Engine/Plugins/Experimental/Chooser/Source/ProxyTableEditor/Private/ProxyTableEditor.cpp b/Engine/Plugins/Experimental/Chooser/Source/ProxyTableEditor/Private/ProxyTableEditor.cpp
new file mode 100644
--- /dev/null
+++ b/Engine/Plugins/Experimental/Chooser/Source/ProxyTableEditor/Private/ProxyTableEditor.cpp
@@ -0,0 +1,679 @@
+// Copyright Epic Games, Inc. All Rights Reserved.
+
+#include "ProxyTableEditor.h"
+
+#include "Framework/Docking/TabManager.h"
+#include "Modules/ModuleManager.h"
+#include "Styling/AppStyle.h"
+#include "Widgets/Docking/SDockTab.h"
+#include "PropertyEditorModule.h"
+#include "IDetailsView.h"
+#include "Editor.h"
+#include "Framework/MultiBox/MultiBoxBuilder.h"
+#include "Widgets/Input/SHyperlink.h"
+#include "Widgets/Input/SNumericEntryBox.h"
+#include "Widgets/Views/SListView.h"
+#include "Widgets/Input/SComboButton.h"
+#include "Widgets/Input/SEditableTextBox.h"
+#include "SAssetDropTarget.h"
+#include "SClassViewer.h"
+#include "SourceCodeNavigation.h"
+#include "ProxyTable.h"
+#include "ClassViewerFilter.h"
+#include "IPropertyAccessEditor.h"
+#include "LandscapeRender.h"
+#include "Widgets/Layout/SWidgetSwitcher.h"
+#include "DragAndDrop/DecoratedDragDropOp.h"
+#include "GraphEditorSettings.h"
+#include "ObjectChooserClassFilter.h"
+#include "ScopedTransaction.h"
+#include "ObjectChooserWidgetFactories.h"
+#include "ContextPropertyWidget.h"
+
+#define LOCTEXT_NAMESPACE "ProxyTableEditor"
+
+namespace UE::ProxyTableEditor
+{
+
+const FName FProxyTableEditor::ToolkitFName( TEXT( "GenericAssetEditor" ) );
+const FName FProxyTableEditor::PropertiesTabId( TEXT( "ProxyEditor_Properties" ) );
+const FName FProxyTableEditor::TableTabId( TEXT( "ProxyEditor_Table" ) );
+
+void FProxyTableEditor::RegisterTabSpawners(const TSharedRef<class FTabManager>& InTabManager)
+{
+	WorkspaceMenuCategory = InTabManager->AddLocalWorkspaceMenuCategory(LOCTEXT("WorkspaceMenu_GenericAssetEditor", "Asset Editor"));
+
+	FAssetEditorToolkit::RegisterTabSpawners(InTabManager);
+
+	InTabManager->RegisterTabSpawner( PropertiesTabId, FOnSpawnTab::CreateSP(this, &FProxyTableEditor::SpawnPropertiesTab) )
+		.SetDisplayName( LOCTEXT("PropertiesTab", "Details") )
+		.SetGroup(WorkspaceMenuCategory.ToSharedRef())
+		.SetIcon(FSlateIcon("EditorStyle", "LevelEditor.Tabs.Details"));
+	
+	InTabManager->RegisterTabSpawner( TableTabId, FOnSpawnTab::CreateSP(this, &FProxyTableEditor::SpawnTableTab) )
+		.SetDisplayName( LOCTEXT("TableTab", "Proxy Table") )
+		.SetGroup(WorkspaceMenuCategory.ToSharedRef())
+		.SetIcon(FSlateIcon("EditorStyle", "LevelEditor.Tabs.Details"));
+}
+
+void FProxyTableEditor::UnregisterTabSpawners(const TSharedRef<class FTabManager>& InTabManager)
+{
+	FAssetEditorToolkit::UnregisterTabSpawners(InTabManager);
+
+	InTabManager->UnregisterTabSpawner( TableTabId );
+	InTabManager->UnregisterTabSpawner( PropertiesTabId );
+}
+
+const FName FProxyTableEditor::ProxyEditorAppIdentifier( TEXT( "ProxyEditorApp" ) );
+
+FProxyTableEditor::~FProxyTableEditor()
+{
+	GEditor->GetEditorSubsystem<UImportSubsystem>()->OnAssetPostImport.RemoveAll(this);
+	FCoreUObjectDelegates::OnObjectsReplaced.RemoveAll(this);
+	
+	DetailsView.Reset();
+}
+
+
+void FProxyTableEditor::InitEditor( const EToolkitMode::Type Mode, const TSharedPtr< class IToolkitHost >& InitToolkitHost, const TArray<UObject*>& ObjectsToEdit, FGetDetailsViewObjects GetDetailsViewObjects )
+{
+	EditingObjects = ObjectsToEdit;
+	FCoreUObjectDelegates::OnObjectsReplaced.AddSP(this, &FProxyTableEditor::OnObjectsReplaced);
+
+	FPropertyEditorModule& PropertyEditorModule = FModuleManager::GetModuleChecked<FPropertyEditorModule>( "PropertyEditor" );
+	FDetailsViewArgs DetailsViewArgs;
+	DetailsViewArgs.NotifyHook = this;
+	DetailsViewArgs.NameAreaSettings = FDetailsViewArgs::HideNameArea;
+	DetailsView = PropertyEditorModule.CreateDetailView( DetailsViewArgs );
+	const TSharedRef<FTabManager::FLayout> StandaloneDefaultLayout = FTabManager::NewLayout( "Standalone_ProxyTableEditor_Layout_v1" )
+	->AddArea
+	(
+		FTabManager::NewPrimaryArea() ->SetOrientation(Orient_Vertical)
+		->Split
+		(
+			FTabManager::NewSplitter()
+			->Split
+			(
+				FTabManager::NewStack()
+				->SetSizeCoefficient(0.7)
+				->AddTab( TableTabId, ETabState::OpenedTab )
+			)
+			->Split
+			(
+				FTabManager::NewStack()
+				->SetSizeCoefficient(0.3)
+				->AddTab( PropertiesTabId, ETabState::OpenedTab )
+			)
+		)
+	);
+
+	const bool bCreateDefaultStandaloneMenu = true;
+	const bool bCreateDefaultToolbar = true;
+	FAssetEditorToolkit::InitAssetEditor( Mode, InitToolkitHost, FProxyTableEditor::ProxyEditorAppIdentifier, StandaloneDefaultLayout, bCreateDefaultStandaloneMenu, bCreateDefaultToolbar, ObjectsToEdit );
+
+	RegenerateMenusAndToolbars();
+
+	SelectRootProperties();
+}
+
+FName FProxyTableEditor::GetToolkitFName() const
+{
+	return ToolkitFName;
+}
+
+FText FProxyTableEditor::GetBaseToolkitName() const
+{
+	return LOCTEXT("AppLabel", "Proxy Table Editor");
+}
+
+void FProxyTableEditor::PostUndo(bool bSuccess)
+{
+	UpdateTableRows();
+}
+	
+void FProxyTableEditor::PostRedo(bool bSuccess)
+{
+	UpdateTableRows();
+}
+
+	
+void FProxyTableEditor::NotifyPreChange(FProperty* PropertyAboutToChange)
+{
+}
+
+void FProxyTableEditor::NotifyPostChange(const FPropertyChangedEvent& PropertyChangedEvent, FProperty* PropertyThatChanged)
+{
+}
+
+	
+FText FProxyTableEditor::GetToolkitName() const
+{
+	const TArray<UObject*>& EditingObjs = GetEditingObjects();
+
+	check( EditingObjs.Num() > 0 );
+
+	FFormatNamedArguments Args;
+	Args.Add( TEXT("ToolkitName"), GetBaseToolkitName() );
+
+	if( EditingObjs.Num() == 1 )
+	{
+		const UObject* EditingObject = EditingObjs[ 0 ];
+		return FText::FromString(EditingObject->GetName());
+	}
+	else
+	{
+		UClass* SharedBaseClass = nullptr;
+		for( int32 x = 0; x < EditingObjs.Num(); ++x )
+		{
+			UObject* Obj = EditingObjs[ x ];
+			check( Obj );
+
+			UClass* ObjClass = Cast<UClass>(Obj);
+			if (ObjClass == nullptr)
+			{
+				ObjClass = Obj->GetClass();
+			}
+			check( ObjClass );
+
+			// Initialize with the class of the first object we encounter.
+			if( SharedBaseClass == nullptr )
+			{
+				SharedBaseClass = ObjClass;
+			}
+
+			// If we've encountered an object that's not a subclass of the current best baseclass,
+			// climb up a step in the class hierarchy.
+			while( !ObjClass->IsChildOf( SharedBaseClass ) )
+			{
+				SharedBaseClass = SharedBaseClass->GetSuperClass();
+			}
+		}
+
+		check(SharedBaseClass);
+
+		Args.Add( TEXT("NumberOfObjects"), EditingObjs.Num() );
+		Args.Add( TEXT("ClassName"), FText::FromString( SharedBaseClass->GetName() ) );
+		return FText::Format( LOCTEXT("ToolkitTitle_EditingMultiple", "{NumberOfObjects} {ClassName} - {ToolkitName}"), Args );
+	}
+}
+
+FText FProxyTableEditor::GetToolkitToolTipText() const
+{
+	const TArray<UObject*>& EditingObjs = GetEditingObjects();
+
+	check( EditingObjs.Num() > 0 );
+
+	FFormatNamedArguments Args;
+	Args.Add( TEXT("ToolkitName"), GetBaseToolkitName() );
+
+	if( EditingObjs.Num() == 1 )
+	{
+		const UObject* EditingObject = EditingObjs[ 0 ];
+		return FAssetEditorToolkit::GetToolTipTextForObject(EditingObject);
+	}
+	else
+	{
+		UClass* SharedBaseClass = NULL;
+		for( int32 x = 0; x < EditingObjs.Num(); ++x )
+		{
+			UObject* Obj = EditingObjs[ x ];
+			check( Obj );
+
+			UClass* ObjClass = Cast<UClass>(Obj);
+			if (ObjClass == nullptr)
+			{
+				ObjClass = Obj->GetClass();
+			}
+			check( ObjClass );
+
+			// Initialize with the class of the first object we encounter.
+			if( SharedBaseClass == nullptr )
+			{
+				SharedBaseClass = ObjClass;
+			}
+
+			// If we've encountered an object that's not a subclass of the current best baseclass,
+			// climb up a step in the class hierarchy.
+			while( !ObjClass->IsChildOf( SharedBaseClass ) )
+			{
+				SharedBaseClass = SharedBaseClass->GetSuperClass();
+			}
+		}
+
+		check(SharedBaseClass);
+
+		Args.Add( TEXT("NumberOfObjects"), EditingObjs.Num() );
+		Args.Add( TEXT("ClassName"), FText::FromString( SharedBaseClass->GetName() ) );
+		return FText::Format( LOCTEXT("ToolkitTitle_EditingMultipleToolTip", "{NumberOfObjects} {ClassName} - {ToolkitName}"), Args );
+	}
+}
+
+FLinearColor FProxyTableEditor::GetWorldCentricTabColorScale() const
+{
+	return FLinearColor( 0.5f, 0.0f, 0.0f, 0.5f );
+}
+
+void FProxyTableEditor::SetPropertyVisibilityDelegate(FIsPropertyVisible InVisibilityDelegate)
+{
+	DetailsView->SetIsPropertyVisibleDelegate(InVisibilityDelegate);
+	DetailsView->ForceRefresh();
+}
+
+void FProxyTableEditor::SetPropertyEditingEnabledDelegate(FIsPropertyEditingEnabled InPropertyEditingDelegate)
+{
+	DetailsView->SetIsPropertyEditingEnabledDelegate(InPropertyEditingDelegate);
+	DetailsView->ForceRefresh();
+}
+
+TSharedRef<SDockTab> FProxyTableEditor::SpawnPropertiesTab( const FSpawnTabArgs& Args )
+{
+	check( Args.GetTabId() == PropertiesTabId );
+
+	return SNew(SDockTab)
+		.Label( LOCTEXT("GenericDetailsTitle", "Details") )
+		.TabColorScale( GetTabColorScale() )
+		.OnCanCloseTab_Lambda([]() { return false; })
+		[
+			DetailsView.ToSharedRef()
+		];
+}
+
+
+class FProxyRowDragDropOp : public FDecoratedDragDropOp
+{
+public:
+	DRAG_DROP_OPERATOR_TYPE(FWidgetTemplateDragDropOp, FDecoratedDragDropOp)
+
+	FProxyTableEditor* ChooserEditor;
+	uint32 RowIndex;
+
+	/** Constructs the drag drop operation */
+	static TSharedRef<FProxyRowDragDropOp> New(FProxyTableEditor* InEditor, uint32 InRowIndex)
+	{
+		TSharedRef<FProxyRowDragDropOp> Operation = MakeShareable(new FProxyRowDragDropOp());
+		Operation->ChooserEditor = InEditor;
+		Operation->RowIndex = InRowIndex;
+		Operation->DefaultHoverText = LOCTEXT("Proxy Row", "Proxy Row");
+		UE::ChooserEditor::FObjectChooserWidgetFactories::ConvertToText(InEditor->GetProxyTable()->Entries[InRowIndex].Value.GetObject(), Operation->DefaultHoverText);
+		Operation->CurrentHoverText = Operation->DefaultHoverText;
+			
+		Operation->Construct();
+	
+		return Operation;
+	};
+};
+
+class SProxyRowHandle : public SCompoundWidget
+{
+public:
+	SLATE_BEGIN_ARGS(SProxyRowHandle)
+	{}
+	SLATE_DEFAULT_SLOT(FArguments, Content)
+	SLATE_ARGUMENT(FProxyTableEditor*, ProxyEditor)
+	SLATE_ARGUMENT(uint32, RowIndex)
+	SLATE_END_ARGS()
+
+	void Construct(const FArguments& InArgs)
+	{
+		ProxyEditor = InArgs._ProxyEditor;
+		RowIndex = InArgs._RowIndex;
+
+		ChildSlot
+		[
+			SNew(SBox) .Padding(0.0f) .HAlign(HAlign_Center) .VAlign(VAlign_Center) .WidthOverride(16.0f)
+			[
+				SNew(SImage)
+				.Image(FCoreStyle::Get().GetBrush("VerticalBoxDragIndicatorShort"))
+			]
+		];
+	}
+
+	FReply OnMouseButtonDown(const FGeometry& MyGeometry, const FPointerEvent& MouseEvent) override
+	{
+		return FReply::Handled().DetectDrag(SharedThis(this), EKeys::LeftMouseButton);
+	};
+
+	FReply OnDragDetected(const FGeometry& MyGeometry, const FPointerEvent& MouseEvent) override
+	{
+		TSharedRef<FProxyRowDragDropOp> DragDropOp = FProxyRowDragDropOp::New(ProxyEditor, RowIndex);
+		return FReply::Handled().BeginDragDrop(DragDropOp);
+	}
+
+private:
+	FProxyTableEditor* ProxyEditor = nullptr;
+	uint32 RowIndex;
+};
+
+
+class SProxyTableRow : public SMultiColumnTableRow<TSharedPtr<FProxyTableEditor::FProxyTableRow>>
+{
+public:
+	SLATE_BEGIN_ARGS(SProxyTableRow) {}
+		/** The list item for this row */
+		SLATE_ARGUMENT(TSharedPtr<FProxyTableEditor::FProxyTableRow>, Entry)
+		SLATE_ARGUMENT(UProxyTable*, ProxyTable)
+		SLATE_ARGUMENT(FProxyTableEditor*, Editor)
+	SLATE_END_ARGS()
+
+
+	void Construct(const FArguments& Args, const TSharedRef<STableViewBase>& OwnerTableView)
+	{
+		RowIndex = Args._Entry;
+		ProxyTable = Args._ProxyTable;
+		Editor = Args._Editor;
+
+		SMultiColumnTableRow<TSharedPtr<FProxyTableEditor::FProxyTableRow>>::Construct(
+			FSuperRowType::FArguments(),
+			OwnerTableView
+		);
+	}
+
+	/** Overridden from SMultiColumnTableRow.  Generates a widget for this column of the list view. */
+	virtual TSharedRef<SWidget> GenerateWidgetForColumn(const FName& ColumnName) override
+	{
+		static FName Handles = "Handles";
+		static FName Key = "Key";
+		static FName Value = "Value";
+		
+		if (RowIndex->RowIndex < ProxyTable->Entries.Num())
+		{
+			if (ColumnName == Handles)
+			{
+				// row drag handle
+				
+				return SNew(SProxyRowHandle).ProxyEditor(Editor).RowIndex(RowIndex->RowIndex);
+			}
+			else if (ColumnName == Value) 
+			{
+				UObject* RowValue = ProxyTable->Entries[RowIndex->RowIndex].Value.GetObject();
+				TSharedPtr<SWidget> ResultWidget = UE::ChooserEditor::FObjectChooserWidgetFactories::CreateWidget(UObjectChooser::StaticClass(), RowValue, nullptr/*ProxyTable->ContextObjectType*/,
+				FOnClassPicked::CreateLambda([this, RowIndex=RowIndex->RowIndex](UClass* ChosenClass)
+				{
+					const FScopedTransaction Transaction(LOCTEXT("Change Row Result Type", "Change Row Result Type"));
+					UObject* RowValue = NewObject<UObject>(ProxyTable, ChosenClass, NAME_None, RF_Transactional);
+					ProxyTable->Modify(true);
+					ProxyTable->Entries[RowIndex].Value = RowValue;
+					UE::ChooserEditor::FObjectChooserWidgetFactories::CreateWidget(UObjectChooser::StaticClass(), RowValue, nullptr/*Chooser->ContextObjectType*/, FOnClassPicked(), &CacheBorder);
+				}),
+				&CacheBorder
+				);
+				return ResultWidget.ToSharedRef();
+			}
+			else if (ColumnName == Key)
+			{
+				return SNew(SEditableTextBox)
+					.Text_Lambda([this](){ return ProxyTable->Entries.Num() > RowIndex->RowIndex ?  FText::FromName(ProxyTable->Entries[RowIndex->RowIndex].Key) : FText();})
+					.OnTextCommitted_Lambda([this](const FText& Text, ETextCommit::Type CommitType)
+					{
+						if (ProxyTable->Entries.Num() > RowIndex->RowIndex)
+						{
+							ProxyTable->Entries[RowIndex->RowIndex].Key = FName(Text.ToString());
+						}
+					});
+			}
+		}
+		else if (RowIndex->RowIndex == ProxyTable->Entries.Num())
+        {
+			// on the row past the end, show an Add button in the result column
+			if (ColumnName == Key)
+			{
+				return Editor->GetCreateRowComboButton().ToSharedRef();
+			}
+		}
+		return SNullWidget::NullWidget;
+	}
+
+private:
+	TSharedPtr<FProxyTableEditor::FProxyTableRow> RowIndex;
+	UProxyTable* ProxyTable;
+	FProxyTableEditor* Editor;
+	TSharedPtr<SBorder> CacheBorder;
+};
+
+
+TSharedRef<ITableRow> FProxyTableEditor::GenerateTableRow(TSharedPtr<FProxyTableRow> InItem, const TSharedRef<STableViewBase>& OwnerTable)
+{
+	UProxyTable* ProxyTable = Cast<UProxyTable>(EditingObjects[0]);
+	
+	return SNew(SProxyTableRow, OwnerTable)
+		.Entry(InItem).ProxyTable(ProxyTable).Editor(this);
+}
+
+FReply FProxyTableEditor::SelectRootProperties()
+{
+	if( DetailsView.IsValid() )
+	{
+		// Make sure details window is pointing to our object
+		DetailsView->SetObjects( EditingObjects );
+	}
+
+	return FReply::Handled();
+}
+
+
+void FProxyTableEditor::UpdateTableColumns()
+{
+	HeaderRow->ClearColumns();
+	HeaderRow->AddColumn(SHeaderRow::Column("Handles")
+					.ManualWidth(30)
+					.HeaderContent()
+					[					
+						SNew(SButton).OnClicked_Raw(this, &FProxyTableEditor::SelectRootProperties)
+					]);
+	
+	HeaderRow->AddColumn(SHeaderRow::Column("Key")
+					.DefaultLabel(LOCTEXT("KeyColumnName", "Key"))
+					.ManualWidth(500));
+
+	HeaderRow->AddColumn(SHeaderRow::Column("Value")
+					.DefaultLabel(LOCTEXT("ValueColumnName", "Value"))
+					.ManualWidth(500));
+
+}
+
+TSharedRef<SDockTab> FProxyTableEditor::SpawnTableTab( const FSpawnTabArgs& Args )
+{
+	check( Args.GetTabId() == TableTabId );
+
+	UProxyTable* ProxyTable = Cast<UProxyTable>(EditingObjects[0]);
+	
+	CreateRowComboButton = SNew(SComboButton)
+		.ComboButtonStyle(FAppStyle::Get(), "SimpleComboButton")
+		.ButtonContent()
+		[
+			SNew(STextBlock).Text(LOCTEXT("AddRow", "+ Add Row"))
+		]
+		.OnGetMenuContent_Lambda([this]()
+		{
+			FClassViewerInitializationOptions Options;
+			Options.ClassFilters.Add(MakeShared<UE::ChooserEditor::FInterfaceClassFilter>(UObjectChooser::StaticClass()) );
+			
+			// Add class filter for columns here
+			TSharedRef<SWidget> Widget = FModuleManager::LoadModuleChecked<FClassViewerModule>("ClassViewer").CreateClassViewer(Options,
+				FOnClassPicked::CreateLambda([this](UClass* ChosenClass)
+				{
+					CreateRowComboButton->SetIsOpen(false);
+					UProxyTable* ProxyTable = Cast<UProxyTable>(EditingObjects[0]);
+					const FScopedTransaction Transaction(LOCTEXT("Add Row Transaction", "Add Row"));
+					ProxyTable->Modify(true);
+					UObject* ResultObject = NewObject<UObject>(ProxyTable, ChosenClass, NAME_None, RF_Transactional);
+					FProxyEntry Entry;
+					Entry.Value = TScriptInterface<IObjectChooser>(ResultObject);
+					ProxyTable->Entries.Add(Entry);
+					UpdateTableRows();
+				}));
+			return Widget;
+		});
+
+	HeaderRow = SNew(SHeaderRow);
+
+	UpdateTableColumns();
+	UpdateTableRows();
+
+	TableView = SNew(SListView<TSharedPtr<FProxyTableRow>>)
+    			.ListItemsSource(&TableRows)
+				.OnKeyDownHandler_Lambda([this](const FGeometry&, const FKeyEvent& Event)
+				{
+					if (Event.GetKey() == EKeys::Delete)
+					{
+						const FScopedTransaction Transaction(LOCTEXT("Delete Row Transaction", "Delete Row"));
+						UProxyTable* ProxyTable = Cast<UProxyTable>(EditingObjects[0]);
+						ProxyTable->Modify(true);
+						// delete selected rows.
+						TArray<uint32> RowsToDelete;
+						for(auto& SelectedRow:SelectedRows)
+						{
+							RowsToDelete.Add(SelectedRow->Row);
+						}
+						// sort indices in reverse
+						RowsToDelete.Sort([](int32 A, int32 B){ return A>B; });
+						for(uint32 RowIndex : RowsToDelete)
+						{
+							ProxyTable->Entries.RemoveAt(RowIndex);
+						}
+
+						UpdateTableRows();
+						
+						return FReply::Handled();
+					}
+					return FReply::Unhandled();
+				}
+				)
+				.OnSelectionChanged_Lambda([this](TSharedPtr<FProxyTableRow> SelectedItem,  ESelectInfo::Type SelectInfo)
+				{
+					if (SelectedItem)
+					{
+						SelectedRows.SetNum(0);
+						UProxyTable* ProxyTable = Cast<UProxyTable>(EditingObjects[0]);
+						// Get the list of objects to edit the details of
+						TObjectPtr<UProxyRowDetails> Selection = NewObject<UProxyRowDetails>();
+						Selection->ProxyTable = ProxyTable;
+						Selection->Row = SelectedItem->RowIndex;
+						Selection->SetFlags(RF_Transactional); 
+						SelectedRows.Add(Selection);
+											
+						if( DetailsView.IsValid() )
+						{
+							// Make sure details window is pointing to our object
+							//DetailsView->SetObjects( ... )
+						}
+					}
+				})
+    			.OnGenerateRow_Raw(this, &FProxyTableEditor::GenerateTableRow)
+				.HeaderRow(HeaderRow);
+	
+	return SNew(SDockTab)
+		.Label( LOCTEXT("ProxtTableTitle", "Proxy Table") )
+		.TabColorScale( GetTabColorScale() )
+		.OnCanCloseTab_Lambda([]() { return false; })
+		[
+			TableView.ToSharedRef()
+		];
+}
+
+void FProxyTableEditor::UpdateTableRows()
+{
+	UProxyTable* ProxyTable = Cast<UProxyTable>(EditingObjects[0]);
+	int32 OldNum = TableRows.Num();
+	int32 NewNum = ProxyTable->Entries.Num();
+
+	// Sync the TableRows array which drives the ui table to match the number of results.
+	// Add 1 at the end, for the "Add Row" control
+	TableRows.SetNum(NewNum + 1);
+	
+	for(int32 i = OldNum; i < NewNum+1; i++)
+	{
+		TableRows[i] = MakeShared<FProxyTableRow>(i);
+	}
+
+	if (TableView.IsValid())
+	{
+		TableView->RebuildList();
+	}
+}
+
+void FProxyTableEditor::OnObjectsReplaced(const TMap<UObject*, UObject*>& ReplacementMap)
+{
+	bool bChangedAny = false;
+
+	// Refresh our details view if one of the objects replaced was in the map. This gets called before the reinstance GC fixup, so we might as well fixup EditingObjects now too
+	for (int32 i = 0; i < EditingObjects.Num(); i++)
+	{
+		// UObject* SourceObject = EditingObjects[i];
+		// UObject* ReplacedObject = ReplacementMap.FindRef(SourceObject);
+		//
+		// if (ReplacedObject && ReplacedObject != SourceObject)
+		// {
+		// 	EditingObjects[i] = ReplacedObject;
+		// 	bChangedAny = true;
+		// }
+	}
+
+	if (bChangedAny)
+	{
+		DetailsView->SetObjects(EditingObjects);
+	}
+}
+
+FString FProxyTableEditor::GetWorldCentricTabPrefix() const
+{
+	return LOCTEXT("WorldCentricTabPrefix", "Generic Asset ").ToString();
+}
+	
+TSharedRef<FProxyTableEditor> FProxyTableEditor::CreateEditor( const EToolkitMode::Type Mode, const TSharedPtr< IToolkitHost >& InitToolkitHost, UObject* ObjectToEdit, FGetDetailsViewObjects GetDetailsViewObjects )
+{
+	TSharedRef< FProxyTableEditor > NewEditor( new FProxyTableEditor() );
+
+	TArray<UObject*> ObjectsToEdit;
+	ObjectsToEdit.Add( ObjectToEdit );
+	NewEditor->InitEditor( Mode, InitToolkitHost, ObjectsToEdit, GetDetailsViewObjects );
+
+	return NewEditor;
+}
+
+TSharedRef<FProxyTableEditor> FProxyTableEditor::CreateEditor( const EToolkitMode::Type Mode, const TSharedPtr< IToolkitHost >& InitToolkitHost, const TArray<UObject*>& ObjectsToEdit, FGetDetailsViewObjects GetDetailsViewObjects )
+{
+	TSharedRef< FProxyTableEditor > NewEditor( new FProxyTableEditor() );
+	NewEditor->InitEditor( Mode, InitToolkitHost, ObjectsToEdit, GetDetailsViewObjects );
+	return NewEditor;
+}
+
+/// Result widgets
+
+TSharedRef<SWidget> CreateLookupProxyWidget(UObject* Object, UClass* ContextObject)
+{
+	UObjectChooser_LookupProxy* LookupProxy = Cast<UObjectChooser_LookupProxy>(Object);
+	
+	return SNew(SHorizontalBox)
+		+ SHorizontalBox::Slot()
+		[
+			UE::ChooserEditor::CreatePropertyWidget<UChooserParameterProxyTable_ContextProperty>(LookupProxy->ProxyTable.GetObject(), ContextObject, GetDefault<UGraphEditorSettings>()->ObjectPinTypeColor)
+		]
+		+ SHorizontalBox::Slot()
+		[
+			SNew(SEditableTextBox)
+			.Text_Lambda([LookupProxy]() { return FText::FromName(LookupProxy->Key);})
+			.OnTextChanged_Lambda([LookupProxy](const FText& NewText)
+			{
+				FScopedTransaction ScopedTransaction(LOCTEXT("Change LookupProxy Key Name", "Change LookupProxy Key Name"));
+				LookupProxy->Modify(true);
+				LookupProxy->Key = FName(NewText.ToString());
+			})
+		];
+}
+
+void ConvertToText_LookupProxy(const UObject* Object, FText& OutText)
+{
+	const UObjectChooser_LookupProxy* LookupProxy = Cast<UObjectChooser_LookupProxy>(Object);
+	OutText = FText::FromString(LookupProxy->Key.ToString());
+}
+	
+void FProxyTableEditor::RegisterWidgets()
+{
+	UE::ChooserEditor::FObjectChooserWidgetFactories::ChooserTextConverter.Add(UObjectChooser_LookupProxy::StaticClass(), ConvertToText_LookupProxy);
+	UE::ChooserEditor::FObjectChooserWidgetFactories::ChooserWidgetCreators.Add(UObjectChooser_LookupProxy::StaticClass(), CreateLookupProxyWidget);
+}
+	
+}
+
+#undef LOCTEXT_NAMESPACE
diff --git a/Engine/Plugins/Experimental/Chooser/Source/ProxyTableEditor/Private/ProxyTableEditor.h b/Engine/Plugins/Experimental/Chooser/Source/ProxyTableEditor/Private/ProxyTableEditor.h
new file mode 100644
--- /dev/null
+++ b/Engine/Plugins/Experimental/Chooser/Source/ProxyTableEditor/Private/ProxyTableEditor.h
@@ -0,0 +1,137 @@
+// Copyright Epic Games, Inc. All Rights Reserved.
+
+#pragma once
+
+#include "CoreMinimal.h"
+#include "Toolkits/IToolkitHost.h"
+#include "Toolkits/AssetEditorToolkit.h"
+#include "Editor/PropertyEditor/Public/PropertyEditorDelegates.h"
+#include "Widgets/Views/SListView.h"
+#include "Widgets/Views/STableViewBase.h"
+#include "ProxyTable.h"
+#include "EditorUndoClient.h"
+#include "Misc/NotifyHook.h"
+#include "ProxyTableEditor.generated.h"
+
+class SComboButton;
+class SEditableText;
+class IDetailsView;
+
+// Class used for chooser editor details customization
+UCLASS()
+class UProxyRowDetails : public UObject
+{
+	GENERATED_BODY()
+public:
+	UPROPERTY(EditAnywhere, Instanced, Category="Hidden")
+	TObjectPtr<UProxyTable> ProxyTable;
+	int Row;
+};
+
+namespace UE::ProxyTableEditor
+{
+	class FProxyTableEditor : public FAssetEditorToolkit, public FSelfRegisteringEditorUndoClient, public FNotifyHook
+	{
+	public:
+		/** Delegate that, given an array of assets, returns an array of objects to use in the details view of an FSimpleAssetEditor */
+		DECLARE_DELEGATE_RetVal_OneParam(TArray<UObject*>, FGetDetailsViewObjects, const TArray<UObject*>&);
+
+		virtual void RegisterTabSpawners(const TSharedRef<class FTabManager>& TabManager) override;
+		virtual void UnregisterTabSpawners(const TSharedRef<class FTabManager>& TabManager) override;
+
+
+		/**
+		* Edits the specified asset object
+		*
+		* @param	Mode					Asset editing mode for this editor (standalone or world-centric)
+		* @param	InitToolkitHost			When Mode is WorldCentric, this is the level editor instance to spawn this editor within
+		* @param	ObjectsToEdit			The object to edit
+		* @param	GetDetailsViewObjects	If bound, a delegate to get the array of objects to use in the details view; uses ObjectsToEdit if not bound
+		*/
+		void InitEditor( const EToolkitMode::Type Mode, const TSharedPtr< class IToolkitHost >& InitToolkitHost, const TArray<UObject*>& ObjectsToEdit, FGetDetailsViewObjects GetDetailsViewObjects );
+
+		/** Destructor */
+		virtual ~FProxyTableEditor();
+
+		/** IToolkit interface */
+		virtual FName GetToolkitFName() const override;
+		virtual FText GetBaseToolkitName() const override;
+		virtual FText GetToolkitName() const override;
+		virtual FText GetToolkitToolTipText() const override;
+		virtual FString GetWorldCentricTabPrefix() const override;
+		virtual FLinearColor GetWorldCentricTabColorScale() const override;
+		virtual bool IsPrimaryEditor() const override { return true; }
+		virtual bool IsSimpleAssetEditor() const override { return false; }
+		
+		/** FEditorUndoClient Interface */
+		virtual void PostUndo(bool bSuccess) override;
+		virtual void PostRedo(bool bSuccess) override;
+		
+		/** Begin FNotifyHook Interface */
+		virtual void NotifyPreChange( FProperty* PropertyAboutToChange ) override;
+		virtual void NotifyPostChange( const FPropertyChangedEvent& PropertyChangedEvent, FProperty* PropertyThatChanged) override;
+
+		UProxyTable* GetProxyTable() { return Cast<UProxyTable>(EditingObjects[0]); }
+	
+		/** Used to show or hide certain properties */
+		void SetPropertyVisibilityDelegate(FIsPropertyVisible InVisibilityDelegate);
+		/** Can be used to disable the details view making it read-only */
+		void SetPropertyEditingEnabledDelegate(FIsPropertyEditingEnabled InPropertyEditingDelegate);
+	
+		struct FProxyTableRow
+		{
+			FProxyTableRow(int32 i) { RowIndex = i; }
+			int32 RowIndex;
+		};
+
+		void UpdateTableColumns();
+		void UpdateTableRows();
+	private:
+
+		FReply SelectRootProperties();
+	
+		/** Create the properties tab and its content */
+		TSharedRef<SDockTab> SpawnPropertiesTab( const FSpawnTabArgs& Args );
+		/** Create the table tab and its content */
+		TSharedRef<SDockTab> SpawnTableTab( const FSpawnTabArgs& Args );
+	
+		TSharedRef<ITableRow> GenerateTableRow(TSharedPtr<FProxyTableRow> InItem, const TSharedRef<STableViewBase>& OwnerTable);
+
+		/** Called when objects need to be swapped out for new versions, like after a blueprint recompile. */
+		void OnObjectsReplaced(const TMap<UObject*, UObject*>& ReplacementMap);
+
+		/** Details view */
+		TSharedPtr< class IDetailsView > DetailsView;
+
+		/** App Identifier. */
+		static const FName ProxyEditorAppIdentifier;
+
+		/**	The tab ids for all the tabs used */
+		static const FName PropertiesTabId;
+		static const FName TableTabId;
+
+		/** The objects open within this editor */
+		TArray<UObject*> EditingObjects;
+
+		TArray<TObjectPtr<UProxyRowDetails>> SelectedRows;
+
+		TArray<TSharedPtr<FProxyTableRow>> TableRows;
+	
+		TSharedPtr<SComboButton> CreateRowComboButton;
+		
+		TSharedPtr<SHeaderRow> HeaderRow;
+		TSharedPtr<SListView<TSharedPtr<FProxyTableRow>>> TableView;
+		
+	public:
+		TSharedPtr<SComboButton>& GetCreateRowComboButton() { return CreateRowComboButton; };
+
+		/** The name given to all instances of this type of editor */
+		static const FName ToolkitFName;
+
+		static TSharedRef<FProxyTableEditor> CreateEditor( const EToolkitMode::Type Mode, const TSharedPtr< IToolkitHost >& InitToolkitHost, UObject* ObjectToEdit, FGetDetailsViewObjects GetDetailsViewObjects = FGetDetailsViewObjects() );
+
+		static TSharedRef<FProxyTableEditor> CreateEditor( const EToolkitMode::Type Mode, const TSharedPtr< IToolkitHost >& InitToolkitHost, const TArray<UObject*>& ObjectsToEdit, FGetDetailsViewObjects GetDetailsViewObjects = FGetDetailsViewObjects() );
+
+		static void RegisterWidgets();
+	};
+}
diff --git a/Engine/Plugins/Experimental/Chooser/Source/ProxyTableEditor/Private/ProxyTableEditorModule.cpp b/Engine/Plugins/Experimental/Chooser/Source/ProxyTableEditor/Private/ProxyTableEditorModule.cpp
new file mode 100644
--- /dev/null
+++ b/Engine/Plugins/Experimental/Chooser/Source/ProxyTableEditor/Private/ProxyTableEditorModule.cpp
@@ -0,0 +1,36 @@
+// Copyright Epic Games, Inc. All Rights Reserved.
+
+#include "ProxyTableEditorModule.h"
+#include "IAssetTools.h"
+#include "AssetTypeActions_ProxyTable.h"
+#include "ProxyTableEditor.h"
+
+#define LOCTEXT_NAMESPACE "ChooserEditorModule"
+
+namespace UE::ProxyTableEditor
+{
+
+void FModule::StartupModule()
+{
+	IAssetTools& AssetTools = FModuleManager::LoadModuleChecked<FAssetToolsModule>("AssetTools").Get();
+	
+	AssetTypeActions_ProxyTable = MakeShared<FAssetTypeActions_ProxyTable>();
+	AssetTools.RegisterAssetTypeActions(AssetTypeActions_ProxyTable.ToSharedRef());
+	FProxyTableEditor::RegisterWidgets();
+}
+		 
+
+void FModule::ShutdownModule()
+{
+	if(FModuleManager::Get().IsModuleLoaded("AssetTools"))
+	{
+		IAssetTools& AssetTools = FModuleManager::GetModuleChecked<FAssetToolsModule>("AssetTools").Get();
+		AssetTools.UnregisterAssetTypeActions(AssetTypeActions_ProxyTable.ToSharedRef());
+	}
+}
+
+}
+
+IMPLEMENT_MODULE(UE::ProxyTableEditor::FModule, ProxyTableEditor);
+
+#undef LOCTEXT_NAMESPACE
\ No newline at end of file
diff --git a/Engine/Plugins/Experimental/Chooser/Source/ProxyTableEditor/Private/ProxyTableEditorModule.h b/Engine/Plugins/Experimental/Chooser/Source/ProxyTableEditor/Private/ProxyTableEditorModule.h
new file mode 100644
--- /dev/null
+++ b/Engine/Plugins/Experimental/Chooser/Source/ProxyTableEditor/Private/ProxyTableEditorModule.h
@@ -0,0 +1,23 @@
+// Copyright Epic Games, Inc. All Rights Reserved.
+
+#pragma once
+
+#include "CoreMinimal.h"
+#include "Modules/ModuleInterface.h"
+
+namespace UE::ProxyTableEditor
+{
+
+class FAssetTypeActions_ChooserTable;
+class FAssetTypeActions_ProxyTable;
+
+class FModule : public IModuleInterface
+{
+private:
+	virtual void StartupModule() override;
+	virtual void ShutdownModule() override;
+
+	TSharedPtr<FAssetTypeActions_ProxyTable> AssetTypeActions_ProxyTable;
+};
+
+}
\ No newline at end of file
diff --git a/Engine/Plugins/Experimental/Chooser/Source/ProxyTableEditor/Private/ProxyTableFactory.cpp b/Engine/Plugins/Experimental/Chooser/Source/ProxyTableEditor/Private/ProxyTableFactory.cpp
new file mode 100644
--- /dev/null
+++ b/Engine/Plugins/Experimental/Chooser/Source/ProxyTableEditor/Private/ProxyTableFactory.cpp
@@ -0,0 +1,21 @@
+﻿// Copyright Epic Games, Inc. All Rights Reserved.
+
+#include "ProxyTableFactory.h"
+#include "ProxyTable.h"
+
+UProxyTableFactory::UProxyTableFactory()
+{
+	bCreateNew = true;
+	bEditAfterNew = true;
+	SupportedClass = UProxyTable::StaticClass();
+}
+
+bool UProxyTableFactory::ConfigureProperties()
+{
+	return true;
+}
+
+UObject* UProxyTableFactory::FactoryCreateNew(UClass* Class, UObject* InParent, FName Name, EObjectFlags Flags, UObject* Context, FFeedbackContext* Warn, FName CallingContext)
+{
+	return NewObject<UProxyTable>(InParent, Class, Name, Flags | RF_Public | RF_Standalone | RF_Transactional | RF_LoadCompleted);
+}
\ No newline at end of file
diff --git a/Engine/Plugins/Experimental/Chooser/Source/ProxyTableEditor/Private/ProxyTableFactory.h b/Engine/Plugins/Experimental/Chooser/Source/ProxyTableEditor/Private/ProxyTableFactory.h
new file mode 100644
--- /dev/null
+++ b/Engine/Plugins/Experimental/Chooser/Source/ProxyTableEditor/Private/ProxyTableFactory.h
@@ -0,0 +1,22 @@
+﻿// Copyright Epic Games, Inc. All Rights Reserved.
+
+#pragma once
+
+#include "CoreMinimal.h"
+#include "Factories/Factory.h"
+#include "ProxyTableFactory.generated.h"
+
+UCLASS(MinimalAPI)
+class UProxyTableFactory : public UFactory
+{
+	GENERATED_BODY()
+
+	UProxyTableFactory();
+	
+	virtual bool ConfigureProperties() override;
+	virtual UObject* FactoryCreateNew(UClass* Class, UObject* InParent, FName Name, EObjectFlags Flags, UObject* Context, FFeedbackContext* Warn, FName CallingContext) override;
+	virtual UObject* FactoryCreateNew(UClass* Class, UObject* InParent, FName Name, EObjectFlags Flags, UObject* Context, FFeedbackContext* Warn) override
+	{
+		return FactoryCreateNew(Class, InParent, Name, Flags, Context, Warn, NAME_None);
+	}
+};
\ No newline at end of file
diff --git a/Engine/Plugins/Experimental/Chooser/Source/ProxyTableEditor/ProxyTableEditor.Build.cs b/Engine/Plugins/Experimental/Chooser/Source/ProxyTableEditor/ProxyTableEditor.Build.cs
new file mode 100644
--- /dev/null
+++ b/Engine/Plugins/Experimental/Chooser/Source/ProxyTableEditor/ProxyTableEditor.Build.cs
@@ -0,0 +1,37 @@
+// Copyright Epic Games, Inc. All Rights Reserved.
+
+namespace UnrealBuildTool.Rules
+{
+	public class ProxyTableEditor : ModuleRules
+	{
+		public ProxyTableEditor(ReadOnlyTargetRules Target) : base(Target)
+		{
+			PrivateDependencyModuleNames.AddRange(
+				new string[]
+				{
+					"Core",
+					"CoreUObject",
+					"Engine",
+					"AssetTools",
+					"Chooser",
+					"ChooserEditor",
+					"ProxyTable",
+					"UnrealEd",
+					"EditorWidgets",
+					"ToolWidgets",
+					"SlateCore",
+					"Slate",
+					"PropertyEditor",
+					"InputCore",
+					"EditorStyle",
+					"PropertyEditor",
+					"BlueprintGraph",
+					"GraphEditor",
+					"GameplayTags",
+					"GameplayTagsEditor",
+					// ... add private dependencies that you statically link with here ...
+				}
+			);
+		}
+	}
+}
\ No newline at end of file
