PR #9870: Chooser - fix for missing `PostLoad` (Contributed by [REDACTED_PERSON])

diff --git a/Engine/Plugins/Experimental/Chooser/Source/Chooser/Internal/EnumColumn.h b/Engine/Plugins/Experimental/Chooser/Source/Chooser/Internal/EnumColumn.h
--- a/Engine/Plugins/Experimental/Chooser/Source/Chooser/Internal/EnumColumn.h
+++ b/Engine/Plugins/Experimental/Chooser/Source/Chooser/Internal/EnumColumn.h
@@ -121,6 +121,7 @@ public:
 
 	virtual void PostLoad() override
 	{
+		Super::PostLoad();
 		InputChanged();
 	}
 #endif
