Follow up changes for proxy table prototype

diff --git a/Engine/Plugins/Experimental/Chooser/Source/ProxyTable/Private/ProxyTable.cpp b/Engine/Plugins/Experimental/Chooser/Source/ProxyTable/Private/ProxyTable.cpp
--- a/Engine/Plugins/Experimental/Chooser/Source/ProxyTable/Private/ProxyTable.cpp
+++ b/Engine/Plugins/Experimental/Chooser/Source/ProxyTable/Private/ProxyTable.cpp
@@ -3,8 +3,8 @@
 
 UObjectChooser_LookupProxy::UObjectChooser_LookupProxy(const FObjectInitializer& ObjectInitializer)
 {
-	// ProxyTable = ObjectInitializer.CreateDefaultSubobject<UChooserParameterProxyTable_ContextProperty>(this, "InputValue");
-	// ProxyTable.GetObject()->SetFlags(RF_Transactional);
+	ProxyTable = ObjectInitializer.CreateDefaultSubobject<UChooserParameterProxyTable_ContextProperty>(this, "InputValue");
+	ProxyTable.GetObject()->SetFlags(RF_Transactional);
 }
 
 static UObject* FindProxyObject(const UProxyTable* Table, FName Key, const UObject* ContextObject)
