#UE Fix monolithic ProxyTable

diff --git a/Engine/Plugins/Experimental/Chooser/Source/ProxyTable/Private/ProxyTableModule.cpp b/Engine/Plugins/Experimental/Chooser/Source/ProxyTable/Private/ProxyTableModule.cpp
--- a/Engine/Plugins/Experimental/Chooser/Source/ProxyTable/Private/ProxyTableModule.cpp
+++ b/Engine/Plugins/Experimental/Chooser/Source/ProxyTable/Private/ProxyTableModule.cpp
@@ -2,4 +2,4 @@
 
 #include "Modules/ModuleManager.h"
 
-IMPLEMENT_MODULE(FDefaultModuleImpl, ProxyTableModule);
\ No newline at end of file
+IMPLEMENT_MODULE(FDefaultModuleImpl, ProxyTable);
\ No newline at end of file
