Add readable Display names to Chooser related types

diff --git a/Engine/Plugins/Experimental/Chooser/Source/Chooser/Internal/BoolColumn.h b/Engine/Plugins/Experimental/Chooser/Source/Chooser/Internal/BoolColumn.h
--- a/Engine/Plugins/Experimental/Chooser/Source/Chooser/Internal/BoolColumn.h
+++ b/Engine/Plugins/Experimental/Chooser/Source/Chooser/Internal/BoolColumn.h
@@ -7,7 +7,7 @@
 #include "ChooserPropertyAccess.h"
 #include "BoolColumn.generated.h"
 
-UCLASS()
+UCLASS(DisplayName = "Bool Property Binding")
 class CHOOSER_API UChooserParameterBool_ContextProperty :  public UObject, public IChooserParameterBool
 {
 	GENERATED_BODY()
diff --git a/Engine/Plugins/Experimental/Chooser/Source/Chooser/Internal/Chooser.h b/Engine/Plugins/Experimental/Chooser/Source/Chooser/Internal/Chooser.h
--- a/Engine/Plugins/Experimental/Chooser/Source/Chooser/Internal/Chooser.h
+++ b/Engine/Plugins/Experimental/Chooser/Source/Chooser/Internal/Chooser.h
@@ -31,7 +31,7 @@ public:
 	TObjectPtr<UClass> OutputObjectType;
 };
 
-UCLASS()
+UCLASS(DisplayName = "Evaluate Chooser")
 class CHOOSER_API UObjectChooser_EvaluateChooser : public UObject, public IObjectChooser
 {
 	GENERATED_BODY()
diff --git a/Engine/Plugins/Experimental/Chooser/Source/Chooser/Internal/EnumColumn.h b/Engine/Plugins/Experimental/Chooser/Source/Chooser/Internal/EnumColumn.h
--- a/Engine/Plugins/Experimental/Chooser/Source/Chooser/Internal/EnumColumn.h
+++ b/Engine/Plugins/Experimental/Chooser/Source/Chooser/Internal/EnumColumn.h
@@ -8,7 +8,7 @@
 
 struct FBindingChainElement;
 
-UCLASS()
+UCLASS(DisplayName = "Enum Property Binding")
 class CHOOSER_API UChooserParameterEnum_ContextProperty : public UObject, public IChooserParameterEnum
 {
 	GENERATED_BODY()
diff --git a/Engine/Plugins/Experimental/Chooser/Source/Chooser/Internal/FloatRangeColumn.h b/Engine/Plugins/Experimental/Chooser/Source/Chooser/Internal/FloatRangeColumn.h
--- a/Engine/Plugins/Experimental/Chooser/Source/Chooser/Internal/FloatRangeColumn.h
+++ b/Engine/Plugins/Experimental/Chooser/Source/Chooser/Internal/FloatRangeColumn.h
@@ -7,7 +7,7 @@
 #include "ChooserPropertyAccess.h"
 #include "FloatRangeColumn.generated.h"
 
-UCLASS()
+UCLASS(DisplayName = "Float Property Binding")
 class CHOOSER_API UChooserParameterFloat_ContextProperty :  public UObject, public IChooserParameterFloat
 {
 	GENERATED_BODY()
diff --git a/Engine/Plugins/Experimental/Chooser/Source/Chooser/Internal/GameplayTagColumn.h b/Engine/Plugins/Experimental/Chooser/Source/Chooser/Internal/GameplayTagColumn.h
--- a/Engine/Plugins/Experimental/Chooser/Source/Chooser/Internal/GameplayTagColumn.h
+++ b/Engine/Plugins/Experimental/Chooser/Source/Chooser/Internal/GameplayTagColumn.h
@@ -10,7 +10,7 @@
 
 struct FBindingChainElement;
 
-UCLASS()
+UCLASS(DisplayName = "Gameplay Tags Property Binding")
 class CHOOSER_API UChooserParameterGameplayTag_ContextProperty :  public UObject, public IChooserParameterGameplayTag
 {
 	GENERATED_BODY()
diff --git a/Engine/Plugins/Experimental/Chooser/Source/Chooser/Internal/ObjectChooser_Asset.h b/Engine/Plugins/Experimental/Chooser/Source/Chooser/Internal/ObjectChooser_Asset.h
--- a/Engine/Plugins/Experimental/Chooser/Source/Chooser/Internal/ObjectChooser_Asset.h
+++ b/Engine/Plugins/Experimental/Chooser/Source/Chooser/Internal/ObjectChooser_Asset.h
@@ -6,7 +6,7 @@
 #include "IObjectChooser.h"
 #include "ObjectChooser_Asset.generated.h"
 
-UCLASS()
+UCLASS(DisplayName = "Choose Asset")
 class CHOOSER_API UObjectChooser_Asset : public UObject, public IObjectChooser
 {
 	GENERATED_BODY()
diff --git a/Engine/Plugins/Experimental/Chooser/Source/ChooserEditor/Private/ChooserTableEditor.cpp b/Engine/Plugins/Experimental/Chooser/Source/ChooserEditor/Private/ChooserTableEditor.cpp
--- a/Engine/Plugins/Experimental/Chooser/Source/ChooserEditor/Private/ChooserTableEditor.cpp
+++ b/Engine/Plugins/Experimental/Chooser/Source/ChooserEditor/Private/ChooserTableEditor.cpp
@@ -531,15 +531,13 @@ void FChooserTableEditor::UpdateTableColumns()
 						})
 						));
 				
-					FClassViewerInitializationOptions Options;
-					Options.ClassFilters.Add(MakeShared<FInterfaceClassFilter>(Column->GetInputValueInterface()));
-
 				MenuBuilder.AddSubMenu(LOCTEXT("Input Type", "Input Type"),
 					LOCTEXT("InputTypeToolTip", "Change input parameter type"),
 					FNewMenuDelegate::CreateLambda([this, &Column](FMenuBuilder& Builder)
 					{
 						FClassViewerInitializationOptions Options;
 						Options.ClassFilters.Add(MakeShared<FInterfaceClassFilter>(Column->GetInputValueInterface()));
+						Options.NameTypeToDisplay = EClassViewerNameTypeToDisplay::DisplayName;
 						
 						// Add class filter for columns here
 						TSharedRef<SWidget> Widget = FModuleManager::LoadModuleChecked<FClassViewerModule>("ClassViewer").CreateClassViewer(Options, FOnClassPicked::CreateLambda([this, &Column](UClass* ChosenClass)
@@ -622,6 +620,7 @@ TSharedRef<SDockTab> FChooserTableEditor::SpawnTableTab( const FSpawnTabArgs& Ar
 	{
 		FClassViewerInitializationOptions Options;
 		Options.ClassFilters.Add(MakeShared<FInterfaceClassFilter>(UChooserColumn::StaticClass()));
+		Options.NameTypeToDisplay = EClassViewerNameTypeToDisplay::DisplayName;
 		
 		// Add class filter for columns here
 		TSharedRef<SWidget> Widget = FModuleManager::LoadModuleChecked<FClassViewerModule>("ClassViewer").CreateClassViewer(Options, FOnClassPicked::CreateLambda([this](UClass* ChosenClass)
@@ -657,6 +656,7 @@ TSharedRef<SDockTab> FChooserTableEditor::SpawnTableTab( const FSpawnTabArgs& Ar
 		{
 			FClassViewerInitializationOptions Options;
 			Options.ClassFilters.Add(MakeShared<FInterfaceClassFilter>(UObjectChooser::StaticClass()) );
+			Options.NameTypeToDisplay = EClassViewerNameTypeToDisplay::DisplayName;
 			
 			// Add class filter for columns here
 			TSharedRef<SWidget> Widget = FModuleManager::LoadModuleChecked<FClassViewerModule>("ClassViewer").CreateClassViewer(Options,
diff --git a/Engine/Plugins/Experimental/Chooser/Source/ChooserEditor/Private/ObjectChooserWidgetFactories.cpp b/Engine/Plugins/Experimental/Chooser/Source/ChooserEditor/Private/ObjectChooserWidgetFactories.cpp
--- a/Engine/Plugins/Experimental/Chooser/Source/ChooserEditor/Private/ObjectChooserWidgetFactories.cpp
+++ b/Engine/Plugins/Experimental/Chooser/Source/ChooserEditor/Private/ObjectChooserWidgetFactories.cpp
@@ -67,6 +67,7 @@ TSharedPtr<SWidget> FObjectChooserWidgetFactories::CreateWidget(UClass* Interfac
 	{
 		FClassViewerInitializationOptions Options;
 		Options.ClassFilters.Add(MakeShared<FInterfaceClassFilter>(InterfaceType));
+		Options.NameTypeToDisplay = EClassViewerNameTypeToDisplay::DisplayName;
 		
 		// Add class filter for columns here
 		TSharedRef<SWidget>  ClassMenu = FModuleManager::LoadModuleChecked<FClassViewerModule>("ClassViewer").
diff --git a/Engine/Plugins/Experimental/Chooser/Source/ProxyTableEditor/Private/ProxyTableEditor.cpp b/Engine/Plugins/Experimental/Chooser/Source/ProxyTableEditor/Private/ProxyTableEditor.cpp
--- a/Engine/Plugins/Experimental/Chooser/Source/ProxyTableEditor/Private/ProxyTableEditor.cpp
+++ b/Engine/Plugins/Experimental/Chooser/Source/ProxyTableEditor/Private/ProxyTableEditor.cpp
@@ -487,6 +487,7 @@ TSharedRef<SDockTab> FProxyTableEditor::SpawnTableTab( const FSpawnTabArgs& Args
 		{
 			FClassViewerInitializationOptions Options;
 			Options.ClassFilters.Add(MakeShared<UE::ChooserEditor::FInterfaceClassFilter>(UObjectChooser::StaticClass()) );
+			Options.NameTypeToDisplay = EClassViewerNameTypeToDisplay::DisplayName;
 			
 			// Add class filter for columns here
 			TSharedRef<SWidget> Widget = FModuleManager::LoadModuleChecked<FClassViewerModule>("ClassViewer").CreateClassViewer(Options,
