AssetDefinition - Converting Chooser Table.

diff --git a/Engine/Plugins/Experimental/Chooser/Source/ChooserEditor/ChooserEditor.Build.cs b/Engine/Plugins/Experimental/Chooser/Source/ChooserEditor/ChooserEditor.Build.cs
--- a/Engine/Plugins/Experimental/Chooser/Source/ChooserEditor/ChooserEditor.Build.cs
+++ b/Engine/Plugins/Experimental/Chooser/Source/ChooserEditor/ChooserEditor.Build.cs
@@ -13,6 +13,7 @@ namespace UnrealBuildTool.Rules
 					"CoreUObject",
 					"Engine",
 					"AssetTools",
+					"AssetDefinition",
 					"Chooser",
 					"UnrealEd",
 					"EditorWidgets",
diff --git a/Engine/Plugins/Experimental/Chooser/Source/ChooserEditor/Private/AssetDefinition_ChooserTable.cpp b/Engine/Plugins/Experimental/Chooser/Source/ChooserEditor/Private/AssetDefinition_ChooserTable.cpp
new file mode 100644
--- /dev/null
+++ b/Engine/Plugins/Experimental/Chooser/Source/ChooserEditor/Private/AssetDefinition_ChooserTable.cpp
@@ -0,0 +1,16 @@
+// Copyright Epic Games, Inc. All Rights Reserved.
+
+#include "AssetDefinition_ChooserTable.h"
+#include "ChooserTableEditor.h"
+
+#define LOCTEXT_NAMESPACE "UAssetDefinition_ChooserTable" 
+
+EAssetCommandResult UAssetDefinition_ChooserTable::OpenAssets(const FAssetOpenArgs& OpenArgs) const
+{
+	TArray<UObject*> ChooserTables = TArray<UObject*>(OpenArgs.LoadObjects<UChooserTable>());
+	UE::ChooserEditor::FChooserTableEditor::CreateEditor(OpenArgs.GetToolkitMode(), OpenArgs.ToolkitHost, ChooserTables);
+
+	return EAssetCommandResult::Handled;
+}
+
+#undef LOCTEXT_NAMESPACE
\ No newline at end of file
diff --git a/Engine/Plugins/Experimental/Chooser/Source/ChooserEditor/Private/AssetDefinition_ChooserTable.h b/Engine/Plugins/Experimental/Chooser/Source/ChooserEditor/Private/AssetDefinition_ChooserTable.h
new file mode 100644
--- /dev/null
+++ b/Engine/Plugins/Experimental/Chooser/Source/ChooserEditor/Private/AssetDefinition_ChooserTable.h
@@ -0,0 +1,28 @@
+// Copyright Epic Games, Inc. All Rights Reserved.
+
+#pragma once
+
+#include "Chooser.h"
+#include "AssetDefinitionDefault.h"
+
+#include "AssetDefinition_ChooserTable.generated.h"
+
+UCLASS()
+class UAssetDefinition_ChooserTable : public UAssetDefinitionDefault
+{
+	GENERATED_BODY()
+
+public:
+	// UAssetDefinition Begin
+	virtual FText GetAssetDisplayName() const override { return NSLOCTEXT("AssetTypeActions", "AssetTypeActions_ChooserTable", "Chooser Table"); }
+	virtual FLinearColor GetAssetColor() const override { return FLinearColor(FColor(128,128,64)); }
+	virtual TSoftClassPtr<UObject> GetAssetClass() const override { return UChooserTable::StaticClass(); }
+	virtual TConstArrayView<FAssetCategoryPath> GetAssetCategories() const override
+	{
+		static const auto Categories = { FAssetCategoryPath(EAssetCategoryPaths::Animation, NSLOCTEXT("AssetTypeActions", "DataInterfacesSubMenu", "Data Interfaces")) };
+		
+		return Categories;
+	}
+	virtual EAssetCommandResult OpenAssets(const FAssetOpenArgs& OpenArgs) const override;
+	// UAssetDefinition End
+};
diff --git a/Engine/Plugins/Experimental/Chooser/Source/ChooserEditor/Private/AssetTypeActions_Chooser.cpp b/Engine/Plugins/Experimental/Chooser/Source/ChooserEditor/Private/AssetTypeActions_Chooser.cpp
deleted file mode 100644
diff --git a/Engine/Plugins/Experimental/Chooser/Source/ChooserEditor/Private/AssetTypeActions_Chooser.h b/Engine/Plugins/Experimental/Chooser/Source/ChooserEditor/Private/AssetTypeActions_Chooser.h
deleted file mode 100644
diff --git a/Engine/Plugins/Experimental/Chooser/Source/ChooserEditor/Private/ChooserEditorModule.cpp b/Engine/Plugins/Experimental/Chooser/Source/ChooserEditor/Private/ChooserEditorModule.cpp
--- a/Engine/Plugins/Experimental/Chooser/Source/ChooserEditor/Private/ChooserEditorModule.cpp
+++ b/Engine/Plugins/Experimental/Chooser/Source/ChooserEditor/Private/ChooserEditorModule.cpp
@@ -2,7 +2,6 @@
 
 #include "ChooserEditorModule.h"
 #include "IAssetTools.h"
-#include "AssetTypeActions_Chooser.h"
 #include "ChooserTableEditor.h"
 #include "BoolColumnEditor.h"
 #include "EnumColumnEditor.h"
@@ -16,9 +15,6 @@ namespace UE::ChooserEditor
 
 void FModule::StartupModule()
 {
-	IAssetTools& AssetTools = FModuleManager::LoadModuleChecked<FAssetToolsModule>("AssetTools").Get();
-	AssetTypeActions_ChooserTable = MakeShared<FAssetTypeActions_ChooserTable>();
-	AssetTools.RegisterAssetTypeActions(AssetTypeActions_ChooserTable.ToSharedRef());
 	FChooserTableEditor::RegisterWidgets();
 	RegisterGameplayTagWidgets();
 	RegisterFloatRangeWidgets();
@@ -34,15 +30,9 @@ void FModule::StartupModule()
 		FOnGetPropertyTypeCustomizationInstance::CreateLambda([] { return MakeShared<UE::ChooserEditor::FInterfacePropertyTypeCustomization>(); }),
 		InterfacePropertyTypeIdentifier);
 }
-		 
 
 void FModule::ShutdownModule()
 {
-	if(FModuleManager::Get().IsModuleLoaded("AssetTools"))
-	{
-		IAssetTools& AssetTools = FModuleManager::GetModuleChecked<FAssetToolsModule>("AssetTools").Get();
-		AssetTools.UnregisterAssetTypeActions(AssetTypeActions_ChooserTable.ToSharedRef());
-	}
 }
 
 }
diff --git a/Engine/Plugins/Experimental/Chooser/Source/ChooserEditor/Private/ChooserEditorModule.h b/Engine/Plugins/Experimental/Chooser/Source/ChooserEditor/Private/ChooserEditorModule.h
--- a/Engine/Plugins/Experimental/Chooser/Source/ChooserEditor/Private/ChooserEditorModule.h
+++ b/Engine/Plugins/Experimental/Chooser/Source/ChooserEditor/Private/ChooserEditorModule.h
@@ -16,7 +16,6 @@ private:
 	virtual void StartupModule() override;
 	virtual void ShutdownModule() override;
 
-	TSharedPtr<FAssetTypeActions_ChooserTable> AssetTypeActions_ChooserTable;
 	TSharedPtr<UE::ChooserEditor::FPropertyTypeIdentifier> InterfacePropertyTypeIdentifier;
 };
 
diff --git a/Engine/Source/Editor/AssetDefinition/Private/AssetDefinition.cpp b/Engine/Source/Editor/AssetDefinition/Private/AssetDefinition.cpp
--- a/Engine/Source/Editor/AssetDefinition/Private/AssetDefinition.cpp
+++ b/Engine/Source/Editor/AssetDefinition/Private/AssetDefinition.cpp
@@ -19,7 +19,7 @@ FAssetCategoryPath EAssetCategoryPaths::Texture(LOCTEXT("Texture", "Texture"));
 FAssetCategoryPath EAssetCategoryPaths::Foliage(LOCTEXT("Foliage", "Foliage"));
 FAssetCategoryPath EAssetCategoryPaths::Input(LOCTEXT("Input", "Input"));
 
-FAssetCategoryPath::FAssetCategoryPath(const FText InCategory)
+FAssetCategoryPath::FAssetCategoryPath(const FText& InCategory)
 {
 	CategoryPath = { TPair<FName, FText>(FName(*FTextInspector::GetSourceString(InCategory)), InCategory) };
 }
@@ -29,6 +29,12 @@ FAssetCategoryPath::FAssetCategoryPath(FText InCategory, FText InSubCategory)
 	CategoryPath = { TPair<FName, FText>(FName(*FTextInspector::GetSourceString(InCategory)), InCategory), TPair<FName, FText>(FName(*FTextInspector::GetSourceString(InSubCategory)), InSubCategory) };
 }
 
+FAssetCategoryPath::FAssetCategoryPath(const FAssetCategoryPath& InCategory, const FText& SubCategory)
+{
+    CategoryPath.Append(InCategory.CategoryPath);
+    CategoryPath.Add(TPair<FName, FText>(FName(*FTextInspector::GetSourceString(SubCategory)), SubCategory));
+}
+
 FAssetCategoryPath::FAssetCategoryPath(TConstArrayView<FText> InCategoryPath)
 {
 	check(InCategoryPath.Num() > 0);
diff --git a/Engine/Source/Editor/AssetDefinition/Public/AssetDefinition.h b/Engine/Source/Editor/AssetDefinition/Public/AssetDefinition.h
--- a/Engine/Source/Editor/AssetDefinition/Public/AssetDefinition.h
+++ b/Engine/Source/Editor/AssetDefinition/Public/AssetDefinition.h
@@ -196,8 +196,9 @@ struct FAssetOpenSupportArgs
  */
 struct ASSETDEFINITION_API FAssetCategoryPath
 {
-	FAssetCategoryPath(FText InCategory);
+	FAssetCategoryPath(const FText& InCategory);
 	FAssetCategoryPath(FText InCategory, FText InSubCategory);
+	FAssetCategoryPath(const FAssetCategoryPath& InCategory, const FText& SubCategory);
 	FAssetCategoryPath(TConstArrayView<FText> InCategoryPath);
 
 	FName GetCategory() const { return CategoryPath[0].Key; }
