Convert ChooserTable and ProxyTable data structures from UObject to FInstancedStruct

diff --git a/Engine/Plugins/Experimental/Chooser/Chooser.uplugin b/Engine/Plugins/Experimental/Chooser/Chooser.uplugin
--- a/Engine/Plugins/Experimental/Chooser/Chooser.uplugin
+++ b/Engine/Plugins/Experimental/Chooser/Chooser.uplugin
@@ -19,6 +19,10 @@
 		{
 			"Name": "GameplayTagsEditor",
 			"Enabled": true
+		},
+		{
+			"Name": "StructUtils",
+			"Enabled": true
 		}
 	],
 	"Modules" :
diff --git a/Engine/Plugins/Experimental/Chooser/Source/Chooser/Chooser.Build.cs b/Engine/Plugins/Experimental/Chooser/Source/Chooser/Chooser.Build.cs
--- a/Engine/Plugins/Experimental/Chooser/Source/Chooser/Chooser.Build.cs
+++ b/Engine/Plugins/Experimental/Chooser/Source/Chooser/Chooser.Build.cs
@@ -19,7 +19,8 @@ namespace UnrealBuildTool.Rules
 					"Core",
 					"CoreUObject",
 					"Engine",
-					"GameplayTags"
+					"GameplayTags",
+					"StructUtils"
 					// ... add private dependencies that you statically link with here ...
 				}
 			);
diff --git a/Engine/Plugins/Experimental/Chooser/Source/Chooser/Internal/BoolColumn.h b/Engine/Plugins/Experimental/Chooser/Source/Chooser/Internal/BoolColumn.h
--- a/Engine/Plugins/Experimental/Chooser/Source/Chooser/Internal/BoolColumn.h
+++ b/Engine/Plugins/Experimental/Chooser/Source/Chooser/Internal/BoolColumn.h
@@ -5,10 +5,11 @@
 #include "IChooserColumn.h"
 #include "IChooserParameterBool.h"
 #include "ChooserPropertyAccess.h"
+#include "InstancedStruct.h"
 #include "BoolColumn.generated.h"
 
-UCLASS(DisplayName = "Bool Property Binding")
-class CHOOSER_API UChooserParameterBool_ContextProperty :  public UObject, public IChooserParameterBool
+USTRUCT(DisplayName = "Bool Property Binding")
+struct CHOOSER_API FBoolContextProperty :  public FChooserParameterBoolBase
 {
 	GENERATED_BODY()
 public:
@@ -29,37 +30,77 @@ public:
 	{
 		UE::Chooser::CopyPropertyChain(InBindingChain, PropertyBindingChain);
 	}
+
+	virtual void GetDisplayName(FText& OutName) const override
+	{
+		if (!PropertyBindingChain.IsEmpty())
+		{
+			OutName = FText::FromName(PropertyBindingChain.Last());
+		}
+	}
 #endif
 };
 
-UCLASS()
-class CHOOSER_API UChooserColumnBool : public UObject, public IChooserColumn
+USTRUCT()
+struct CHOOSER_API FBoolColumn : public FChooserColumnBase
 {
 	GENERATED_BODY()
 	public:
-	UChooserColumnBool() {};
-	UChooserColumnBool(const FObjectInitializer& ObjectInitializer);
+	FBoolColumn();
+	
+	UPROPERTY(EditAnywhere, Meta = (ExcludeBaseStruct, BaseStruct = "/Script/Chooser.ChooserParameterBoolBase"), Category = "Hidden")
+	FInstancedStruct InputValue;
+	
+	UPROPERTY(EditAnywhere, Category=Runtime);
+	TArray<bool> RowValues; 
+	
+	virtual void Filter(const UObject* ContextObject, const TArray<uint32>& IndexListIn, TArray<uint32>& IndexListOut) const override;
+
+	CHOOSER_COLUMN_BOILERPLATE(FChooserParameterBoolBase);
+};
+
+// deprecated class versions for converting old data
+
+UCLASS(ClassGroup = "LiveLink", deprecated)
+class CHOOSER_API UDEPRECATED_ChooserParameterBool_ContextProperty :  public UObject, public IChooserParameterBool
+{
+	GENERATED_BODY()
+public:
+	UPROPERTY()
+	TArray<FName> PropertyBindingChain;
 	
-	UPROPERTY(EditAnywhere, Meta = (EditInlineInterface = "true"), Category = "Input")
+	virtual void ConvertToInstancedStruct(FInstancedStruct& OutInstancedStruct) const override
+	{
+		OutInstancedStruct.InitializeAs(FBoolContextProperty::StaticStruct());
+		FBoolContextProperty& Property = OutInstancedStruct.GetMutable<FBoolContextProperty>();
+		Property.PropertyBindingChain = PropertyBindingChain;
+	}
+};
+
+UCLASS(ClassGroup = "LiveLink", deprecated)
+class CHOOSER_API UDEPRECATED_ChooserColumnBool : public UObject, public IChooserColumn
+{
+	GENERATED_BODY()
+public:
+	UDEPRECATED_ChooserColumnBool() {};
+	UDEPRECATED_ChooserColumnBool(const FObjectInitializer& ObjectInitializer){};
+	
+	UPROPERTY(EditAnywhere, Category = "Input")
 	TScriptInterface<IChooserParameterBool> InputValue;
 
 	UPROPERTY(EditAnywhere, Category=Runtime)
 	// array of results (cells for this column for each row in the table)
 	// should match the length of the Results array 
 	TArray<bool> RowValues;
-	
-	virtual void Filter(const UObject* ContextObject, const TArray<uint32>& IndexListIn, TArray<uint32>& IndexListOut) override;
 
-	// todo: macro boilerplate
-	virtual void SetNumRows(uint32 NumRows) override { RowValues.SetNum(NumRows); }
-	virtual void DeleteRows(const TArray<uint32> & RowIndices )
+	virtual void ConvertToInstancedStruct(FInstancedStruct& OutInstancedStruct) const override
 	{
-		for(uint32 Index : RowIndices)
+		OutInstancedStruct.InitializeAs(FBoolColumn::StaticStruct());
+		FBoolColumn& Column = OutInstancedStruct.GetMutable<FBoolColumn>();
+		if (IChooserParameterBool* InputValueInterface = InputValue.GetInterface())
 		{
-			RowValues.RemoveAt(Index);
+			InputValueInterface->ConvertToInstancedStruct(Column.InputValue);
 		}
+		Column.RowValues = RowValues;
 	}
-	virtual UClass* GetInputValueInterface() override { return UChooserParameterBool::StaticClass(); };
-	virtual UObject* GetInputValue() override { return InputValue.GetObject(); };
-	virtual void SetInputValue(UObject* Value) override { InputValue = Value; };
 };
\ No newline at end of file
diff --git a/Engine/Plugins/Experimental/Chooser/Source/Chooser/Internal/Chooser.h b/Engine/Plugins/Experimental/Chooser/Source/Chooser/Internal/Chooser.h
--- a/Engine/Plugins/Experimental/Chooser/Source/Chooser/Internal/Chooser.h
+++ b/Engine/Plugins/Experimental/Chooser/Source/Chooser/Internal/Chooser.h
@@ -4,35 +4,49 @@
 #include "CoreMinimal.h"
 #include "UObject/Object.h"
 #include "IObjectChooser.h"
+#include "InstancedStruct.h"
 #include "IChooserColumn.h"
 
 #include "Chooser.generated.h"
 
 
-UCLASS(MinimalAPI)
+UCLASS(MinimalAPI, BlueprintType)
 class UChooserTable : public UObject
 {
 	GENERATED_UCLASS_BODY()
 public:
 	UChooserTable() {}
-
-	// each possible result
-	UPROPERTY(EditAnywhere, Meta = (EditInlineInterface = "true"), Category = "Hidden")
-	TArray<TScriptInterface<IObjectChooser>> Results;
 	
-	// columns which filter results
-	UPROPERTY(EditAnywhere, Meta = (EditInlineInterface = "true"), Category="Hidden")
-	TArray<TScriptInterface<IChooserColumn>> Columns;
+	// deprecated UObject Results
+	UPROPERTY()
+	TArray<TScriptInterface<IObjectChooser>> Results_DEPRECATED;
+	
+	// deprecated UObject Columns
+	UPROPERTY()
+	TArray<TScriptInterface<IChooserColumn>> Columns_DEPRECATED;
+
+	// Each possible Result (Rows of chooser table)
+	UPROPERTY(EditAnywhere, DisplayName = "Results", Meta = (ExcludeBaseStruct, BaseStruct = "/Script/Chooser.ObjectChooserBase"), Category = "Hidden")
+	TArray<FInstancedStruct> ResultsStructs;
+
+	// Columns which filter Results
+	UPROPERTY(EditAnywhere, DisplayName = "Columns", Category = Hidden, meta = (ExcludeBaseStruct, BaseStruct = "/Script/Chooser.ChooserColumnBase"))
+	TArray<FInstancedStruct> ColumnsStructs;
 
 	UPROPERTY(EditAnywhere, Category="Input", Meta = (AllowAbstract=true))
 	TObjectPtr<UClass> ContextObjectType;
 	
 	UPROPERTY(EditAnywhere, Category="Output", Meta = (AllowAbstract=true))
 	TObjectPtr<UClass> OutputObjectType;
+
+#if WITH_EDITOR
+	virtual void PostLoad() override;
+#endif
 };
 
-UCLASS(DisplayName = "Evaluate Chooser")
-class CHOOSER_API UObjectChooser_EvaluateChooser : public UObject, public IObjectChooser
+
+USTRUCT(DisplayName = "Evaluate Chooser")
+struct CHOOSER_API FEvaluateChooser : public FObjectChooserBase
 {
 	GENERATED_BODY()
 
@@ -44,6 +58,22 @@ class CHOOSER_API UObjectChooser_EvaluateChooser : public UObject, public IObjec
 	TObjectPtr<UChooserTable> Chooser;
 };
 
+// Deprecated class for converting old data
+UCLASS(ClassGroup = "LiveLink", deprecated)
+class CHOOSER_API UDEPRECATED_ObjectChooser_EvaluateChooser : public UObject, public IObjectChooser
+{
+	GENERATED_BODY()
+	UPROPERTY(EditAnywhere, Category="Parameters")
+	TObjectPtr<UChooserTable> Chooser;
+
+	virtual void ConvertToInstancedStruct(FInstancedStruct& OutInstancedStruct) const
+	{
+		OutInstancedStruct.InitializeAs(FEvaluateChooser::StaticStruct());
+		FEvaluateChooser& AssetChooser = OutInstancedStruct.GetMutable<FEvaluateChooser>();
+		AssetChooser.Chooser = Chooser;
+	}
+};
+
 
 UCLASS()
 class CHOOSER_API UChooserColumnMenuContext : public UObject
diff --git a/Engine/Plugins/Experimental/Chooser/Source/Chooser/Internal/EnumColumn.h b/Engine/Plugins/Experimental/Chooser/Source/Chooser/Internal/EnumColumn.h
--- a/Engine/Plugins/Experimental/Chooser/Source/Chooser/Internal/EnumColumn.h
+++ b/Engine/Plugins/Experimental/Chooser/Source/Chooser/Internal/EnumColumn.h
@@ -4,15 +4,16 @@
 #include "CoreMinimal.h"
 #include "IChooserColumn.h"
 #include "IChooserParameterEnum.h"
+#include "InstancedStruct.h"
 #include "EnumColumn.generated.h"
 
 struct FBindingChainElement;
 
-UCLASS(DisplayName = "Enum Property Binding")
-class CHOOSER_API UChooserParameterEnum_ContextProperty : public UObject, public IChooserParameterEnum
+USTRUCT(DisplayName = "Enum Property Binding")
+struct CHOOSER_API FEnumContextProperty : public FChooserParameterEnumBase
 {
 	GENERATED_BODY()
-public:
+	
 	UPROPERTY()
 	TArray<FName> PropertyBindingChain;
 
@@ -36,21 +37,26 @@ public:
 
 	void SetBinding(const TArray<FBindingChainElement>& InBindingChain);
 
-	virtual const UEnum* GetEnum() const override { return Enum; }
+	virtual void GetDisplayName(FText& OutName) const override
+	{
+		if (!PropertyBindingChain.IsEmpty())
+		{
+			OutName = FText::FromName(PropertyBindingChain.Last());
+		}
+	}
 
-	virtual FSimpleMulticastDelegate& OnEnumChanged() { return EnumChanged; }
+	virtual const UEnum* GetEnum() const override { return Enum; }
 #endif
 
 private:
 #if WITH_EDITORONLY_DATA
 	UPROPERTY()
 	TObjectPtr<const UEnum> Enum = nullptr;
-
-	FSimpleMulticastDelegate EnumChanged;
 #endif
 };
 
 
+
 UENUM()
 enum class EChooserEnumComparison : uint8
 {
@@ -76,46 +82,41 @@ struct FChooserEnumRowData
 	bool Evaluate(const uint8 LeftHandSide) const;
 };
 
-UCLASS()
-class CHOOSER_API UChooserColumnEnum : public UObject, public IChooserColumn
+
+USTRUCT()
+struct CHOOSER_API FEnumColumn : public FChooserColumnBase
 {
 	GENERATED_BODY()
 public:
-	UChooserColumnEnum() = default;
-	UChooserColumnEnum(const FObjectInitializer& ObjectInitializer);
+	FEnumColumn();
 
-	UPROPERTY(EditAnywhere, Meta = (EditInlineInterface = "true"), Category = "Input")
-	TScriptInterface<IChooserParameterEnum> InputValue;
+	UPROPERTY(EditAnywhere, Meta = (ExcludeBaseStruct, BaseStruct = "/Script/Chooser.ChooserParameterEnumBase"), Category = "Hidden")
+	FInstancedStruct InputValue;
 
 	UPROPERTY(EditAnywhere, Category = Runtime)
 	// array of results (cells for this column for each row in the table)
 	// should match the length of the Results array
 	TArray<FChooserEnumRowData> RowValues;
 
-	virtual void Filter(const UObject* ContextObject, const TArray<uint32>& IndexListIn, TArray<uint32>& IndexListOut) override;
-	virtual void SetNumRows(uint32 NumRows) { RowValues.SetNum(NumRows); }
-	virtual void DeleteRows(const TArray<uint32>& RowIndices)
-	{
-		for (uint32 Index : RowIndices)
-		{
-			RowValues.RemoveAt(Index);
-		}
-	}
-	virtual UClass* GetInputValueInterface() { return UChooserParameterEnum::StaticClass(); };
-	virtual UObject* GetInputValue() override { return InputValue.GetObject(); };
-	virtual void SetInputValue(UObject* Value) override
-	{
-		InputValue = Value;
+	virtual void Filter(const UObject* ContextObject, const TArray<uint32>& IndexListIn, TArray<uint32>& IndexListOut) const override;
+	
+	CHOOSER_COLUMN_BOILERPLATE_NoSetInputType(FChooserParameterEnumBase);
+	
 #if WITH_EDITOR
+	virtual void SetInputType(const UScriptStruct* Type) override
+	{
+		InputValue.InitializeAs(Type);
 		InputChanged();
-#endif
 	};
 
-#if WITH_EDITOR
 	FSimpleMulticastDelegate OnEnumChanged;
+	
 	void InputChanged()
 	{
-		InputValue->OnEnumChanged().AddLambda([this](){ OnEnumChanged.Broadcast(); });
+		if (InputValue.IsValid())
+		{
+			InputValue.GetMutable<FChooserParameterEnumBase>().OnEnumChanged.AddLambda([this](){ OnEnumChanged.Broadcast(); });
+		}
 		OnEnumChanged.Broadcast();
 	}
 
@@ -125,4 +126,47 @@ public:
 		InputChanged();
 	}
 #endif
+};
+
+// deprecated class version for converting old data
+UCLASS(ClassGroup = "LiveLink", deprecated)
+class CHOOSER_API UDEPRECATED_ChooserParameterEnum_ContextProperty : public UObject, public IChooserParameterEnum
+{
+	GENERATED_BODY()
+public:
+	UPROPERTY()
+	TArray<FName> PropertyBindingChain;
+	
+	virtual void ConvertToInstancedStruct(FInstancedStruct& OutInstancedStruct) const override
+	{
+		OutInstancedStruct.InitializeAs(FEnumContextProperty::StaticStruct());
+		FEnumContextProperty& Property = OutInstancedStruct.GetMutable<FEnumContextProperty>();
+		Property.PropertyBindingChain = PropertyBindingChain;
+	}
+};
+
+UCLASS(ClassGroup = "LiveLink", deprecated)
+class CHOOSER_API UDEPRECATED_ChooserColumnEnum : public UObject, public IChooserColumn
+{
+	GENERATED_BODY()
+public:
+	UDEPRECATED_ChooserColumnEnum() = default;
+	UDEPRECATED_ChooserColumnEnum(const FObjectInitializer& ObjectInitializer){};
+
+	UPROPERTY(EditAnywhere, Category = "Input")
+	TScriptInterface<IChooserParameterEnum> InputValue;
+
+	UPROPERTY(EditAnywhere, Category = Runtime)
+	TArray<FChooserEnumRowData> RowValues;
+	
+	virtual void ConvertToInstancedStruct(FInstancedStruct& OutInstancedStruct) const override
+	{
+		OutInstancedStruct.InitializeAs(FEnumColumn::StaticStruct());
+		FEnumColumn& Column = OutInstancedStruct.GetMutable<FEnumColumn>();
+		if (IChooserParameterEnum* InputValueInterface = InputValue.GetInterface())
+		{
+			InputValueInterface->ConvertToInstancedStruct(Column.InputValue);
+		}
+		Column.RowValues = RowValues;
+	}
 };
\ No newline at end of file
diff --git a/Engine/Plugins/Experimental/Chooser/Source/Chooser/Internal/FloatRangeColumn.h b/Engine/Plugins/Experimental/Chooser/Source/Chooser/Internal/FloatRangeColumn.h
--- a/Engine/Plugins/Experimental/Chooser/Source/Chooser/Internal/FloatRangeColumn.h
+++ b/Engine/Plugins/Experimental/Chooser/Source/Chooser/Internal/FloatRangeColumn.h
@@ -5,18 +5,19 @@
 #include "IChooserColumn.h"
 #include "IChooserParameterFloat.h"
 #include "ChooserPropertyAccess.h"
+#include "InstancedStruct.h"
 #include "FloatRangeColumn.generated.h"
 
-UCLASS(DisplayName = "Float Property Binding")
-class CHOOSER_API UChooserParameterFloat_ContextProperty :  public UObject, public IChooserParameterFloat
+USTRUCT(DisplayName = "Float Property Binding")
+struct CHOOSER_API FFloatContextProperty :  public FChooserParameterFloatBase
 {
 	GENERATED_BODY()
-public:
-	UPROPERTY()
-	TArray<FName> PropertyBindingChain;
 	
 	virtual bool GetValue(const UObject* ContextObject, float& OutResult) const override;
 
+	UPROPERTY()
+	TArray<FName> PropertyBindingChain;
+
 #if WITH_EDITOR
 	static bool CanBind(const FProperty& Property)
 	{
@@ -30,6 +31,14 @@ public:
 	{
 		UE::Chooser::CopyPropertyChain(InBindingChain, PropertyBindingChain);
 	}
+
+	virtual void GetDisplayName(FText& OutName) const override
+	{
+		if (!PropertyBindingChain.IsEmpty())
+		{
+			OutName = FText::FromName(PropertyBindingChain.Last());
+		}
+	}
 #endif
 };
 
@@ -45,15 +54,55 @@ struct FChooserFloatRangeRowData
 	float Max=0;
 };
 
-UCLASS()
-class CHOOSER_API UChooserColumnFloatRange : public UObject, public IChooserColumn
+
+USTRUCT()
+struct CHOOSER_API FFloatRangeColumn : public FChooserColumnBase
+{
+	GENERATED_BODY()
+	public:
+	FFloatRangeColumn();
+		
+	UPROPERTY(EditAnywhere, Meta = (ExcludeBaseStruct, BaseStruct = "/Script/Chooser.ChooserParameterFloatBase"), Category = "Hidden")
+	FInstancedStruct InputValue;
+	
+	UPROPERTY(EditAnywhere, Category=Runtime)
+	// array of results (cells for this column for each row in the table)
+	// should match the length of the Results array 
+	TArray<FChooserFloatRangeRowData> RowValues;
+	
+	virtual void Filter(const UObject* ContextObject, const TArray<uint32>& IndexListIn, TArray<uint32>& IndexListOut) const override;
+	
+	CHOOSER_COLUMN_BOILERPLATE(FChooserParameterFloatBase);
+
+};
+
+// deprecated class version to support upgrading old data
+
+UCLASS(ClassGroup = "LiveLink", deprecated)
+class CHOOSER_API UDEPRECATED_ChooserParameterFloat_ContextProperty :  public UObject, public IChooserParameterFloat
+{
+	GENERATED_BODY()
+public:
+	UPROPERTY()
+	TArray<FName> PropertyBindingChain;
+	
+	virtual void ConvertToInstancedStruct(FInstancedStruct& OutInstancedStruct) const override
+	{
+		OutInstancedStruct.InitializeAs(FFloatContextProperty::StaticStruct());
+		FFloatContextProperty& Property = OutInstancedStruct.GetMutable<FFloatContextProperty>();
+		Property.PropertyBindingChain = PropertyBindingChain;
+	}
+};
+
+UCLASS(ClassGroup = "LiveLink", deprecated)
+class CHOOSER_API UDEPRECATED_ChooserColumnFloatRange : public UObject, public IChooserColumn
 {
 	GENERATED_BODY()
 	public:
-	UChooserColumnFloatRange() {}
-	UChooserColumnFloatRange(const FObjectInitializer& ObjectInitializer);
+	UDEPRECATED_ChooserColumnFloatRange() {}
+	UDEPRECATED_ChooserColumnFloatRange(const FObjectInitializer& ObjectInitializer) {};
 		
-	UPROPERTY(EditAnywhere, Meta = (EditInlineInterface = "true"), Category = "Input")
+	UPROPERTY(EditAnywhere, Category = "Input")
 	TScriptInterface<IChooserParameterFloat> InputValue;
 	
 	UPROPERTY(EditAnywhere, Category=Runtime)
@@ -61,16 +110,14 @@ class CHOOSER_API UChooserColumnFloatRange : public UObject, public IChooserColu
 	// should match the length of the Results array 
 	TArray<FChooserFloatRangeRowData> RowValues;
 	
-	virtual void Filter(const UObject* ContextObject, const TArray<uint32>& IndexListIn, TArray<uint32>& IndexListOut) override;
-	virtual void SetNumRows(uint32 NumRows) { RowValues.SetNum(NumRows); }
-	virtual void DeleteRows(const TArray<uint32> & RowIndices )
+	virtual void ConvertToInstancedStruct(FInstancedStruct& OutInstancedStruct) const override
 	{
-		for(uint32 Index : RowIndices)
+		OutInstancedStruct.InitializeAs(FFloatRangeColumn::StaticStruct());
+		FFloatRangeColumn& Column = OutInstancedStruct.GetMutable<FFloatRangeColumn>();
+		if (IChooserParameterFloat* InputValueInterface = InputValue.GetInterface())
 		{
-			RowValues.RemoveAt(Index);
+			InputValueInterface->ConvertToInstancedStruct(Column.InputValue);
 		}
+		Column.RowValues = RowValues;
 	}
-	virtual UClass* GetInputValueInterface() { return UChooserParameterFloat::StaticClass(); };
-	virtual UObject* GetInputValue() override { return InputValue.GetObject(); };
-	virtual void SetInputValue(UObject* Value) override { InputValue = Value; };
 };
\ No newline at end of file
diff --git a/Engine/Plugins/Experimental/Chooser/Source/Chooser/Internal/GameplayTagColumn.h b/Engine/Plugins/Experimental/Chooser/Source/Chooser/Internal/GameplayTagColumn.h
--- a/Engine/Plugins/Experimental/Chooser/Source/Chooser/Internal/GameplayTagColumn.h
+++ b/Engine/Plugins/Experimental/Chooser/Source/Chooser/Internal/GameplayTagColumn.h
@@ -6,12 +6,14 @@
 #include "IChooserParameterGameplayTag.h"
 #include "ChooserPropertyAccess.h"
 #include "GameplayTagContainer.h"
+#include "InstancedStruct.h"
 #include "GameplayTagColumn.generated.h"
 
 struct FBindingChainElement;
 
-UCLASS(DisplayName = "Gameplay Tags Property Binding")
-class CHOOSER_API UChooserParameterGameplayTag_ContextProperty :  public UObject, public IChooserParameterGameplayTag
+
+USTRUCT(DisplayName = "Gameplay Tags Property Binding")
+struct CHOOSER_API FGameplayTagContextProperty :  public FChooserParameterGameplayTagBase
 {
 	GENERATED_BODY()
 public:
@@ -31,19 +33,26 @@ public:
 	{
 		UE::Chooser::CopyPropertyChain(InBindingChain, PropertyBindingChain);
 	}
+
+	virtual void GetDisplayName(FText& OutName) const override
+	{
+		if (!PropertyBindingChain.IsEmpty())
+		{
+			OutName = FText::FromName(PropertyBindingChain.Last());
+		}
+	}
 #endif
 };
 
-UCLASS()
-class CHOOSER_API UChooserColumnGameplayTag : public UObject, public IChooserColumn
+USTRUCT()
+struct CHOOSER_API FGameplayTagColumn : public FChooserColumnBase
 {
 	GENERATED_BODY()
 	public:
-	UChooserColumnGameplayTag() {};
-	UChooserColumnGameplayTag(const FObjectInitializer& ObjectInitializer);
+	FGameplayTagColumn();
 	
-	UPROPERTY(EditAnywhere, Meta = (EditInlineInterface = "true"), Category = "Input")
-	TScriptInterface<IChooserParameterGameplayTag> InputValue;
+	UPROPERTY(EditAnywhere, Meta = (ExcludeBaseStruct, BaseStruct = "/Script/Chooser.ChooserParameterGameplayTagBase"), Category = "Hidden")
+	FInstancedStruct InputValue;
 
 	UPROPERTY(EditAnywhere, Category=Runtime)
 	EGameplayContainerMatchType	TagMatchType = EGameplayContainerMatchType::Any;
@@ -53,18 +62,56 @@ class CHOOSER_API UChooserColumnGameplayTag : public UObject, public IChooserCol
 	// should match the length of the Results array 
 	TArray<FGameplayTagContainer> RowValues;
 	
-	virtual void Filter(const UObject* ContextObject, const TArray<uint32>& IndexListIn, TArray<uint32>& IndexListOut) override;
+	virtual void Filter(const UObject* ContextObject, const TArray<uint32>& IndexListIn, TArray<uint32>& IndexListOut) const override;
+
+	CHOOSER_COLUMN_BOILERPLATE(FChooserParameterGameplayTagBase);
+};
 
-	// todo: macro boilerplate
-	virtual void SetNumRows(uint32 NumRows) override { RowValues.SetNum(NumRows); }
-	virtual void DeleteRows(const TArray<uint32> & RowIndices )
+// deprecated class versions for converting old data
+
+UCLASS(ClassGroup = "LiveLink", deprecated)
+class CHOOSER_API UDEPRECATED_ChooserParameterGameplayTag_ContextProperty :  public UObject, public IChooserParameterGameplayTag
+{
+	GENERATED_BODY()
+public:
+	UPROPERTY()
+	TArray<FName> PropertyBindingChain;
+	
+	virtual void ConvertToInstancedStruct(FInstancedStruct& OutInstancedStruct) const override
 	{
-		for(uint32 Index : RowIndices)
+		OutInstancedStruct.InitializeAs(FGameplayTagContextProperty::StaticStruct());
+		FGameplayTagContextProperty& Property = OutInstancedStruct.GetMutable<FGameplayTagContextProperty>();
+		Property.PropertyBindingChain = PropertyBindingChain;
+	}
+};
+
+UCLASS(ClassGroup = "LiveLink", deprecated)
+class CHOOSER_API UDEPRECATED_ChooserColumnGameplayTag : public UObject, public IChooserColumn
+{
+	GENERATED_BODY()
+	public:
+	UDEPRECATED_ChooserColumnGameplayTag() {};
+	UDEPRECATED_ChooserColumnGameplayTag(const FObjectInitializer& ObjectInitializer){};
+	
+	UPROPERTY(EditAnywhere, Category = "Input")
+	TScriptInterface<IChooserParameterGameplayTag> InputValue;
+
+	UPROPERTY(EditAnywhere, Category=Runtime)
+	EGameplayContainerMatchType	TagMatchType = EGameplayContainerMatchType::Any;
+
+	UPROPERTY(EditAnywhere, Category=Runtime)
+	// array of results (cells for this column for each row in the table)
+	// should match the length of the Results array 
+	TArray<FGameplayTagContainer> RowValues;
+	
+	virtual void ConvertToInstancedStruct(FInstancedStruct& OutInstancedStruct) const override
+	{
+		OutInstancedStruct.InitializeAs(FGameplayTagColumn::StaticStruct());
+		FGameplayTagColumn& Column = OutInstancedStruct.GetMutable<FGameplayTagColumn>();
+		if (IChooserParameterGameplayTag* InputValueInterface = InputValue.GetInterface())
 		{
-			RowValues.RemoveAt(Index);
+			InputValueInterface->ConvertToInstancedStruct(Column.InputValue);
 		}
+		Column.RowValues = RowValues;
 	}
-	virtual UClass* GetInputValueInterface() override { return UChooserParameterGameplayTag::StaticClass(); };
-	virtual UObject* GetInputValue() override { return InputValue.GetObject(); };
-	virtual void SetInputValue(UObject* Value) override { InputValue = Value; };
-};
\ No newline at end of file
+};
diff --git a/Engine/Plugins/Experimental/Chooser/Source/Chooser/Internal/ObjectChooser_Asset.h b/Engine/Plugins/Experimental/Chooser/Source/Chooser/Internal/ObjectChooser_Asset.h
--- a/Engine/Plugins/Experimental/Chooser/Source/Chooser/Internal/ObjectChooser_Asset.h
+++ b/Engine/Plugins/Experimental/Chooser/Source/Chooser/Internal/ObjectChooser_Asset.h
@@ -6,15 +6,32 @@
 #include "IObjectChooser.h"
 #include "ObjectChooser_Asset.generated.h"
 
-UCLASS(DisplayName = "Choose Asset")
-class CHOOSER_API UObjectChooser_Asset : public UObject, public IObjectChooser
+USTRUCT(DisplayName = "Asset")
+struct CHOOSER_API FAssetChooser : public FObjectChooserBase
 {
 	GENERATED_BODY()
 	
-	// IObjectChooser interface
+	// FObjectChooserBase interface
 	virtual UObject* ChooseObject(const UObject* ContextObject) const final override;
-	
 public: 
 	UPROPERTY(EditAnywhere, Category = "Parameters")
 	TObjectPtr<UObject> Asset;
 };
+
+// deprecated class for upgrading old data
+UCLASS(ClassGroup = "LiveLink", deprecated)
+class CHOOSER_API UDEPRECATED_ObjectChooser_Asset : public UObject, public IObjectChooser
+{
+	GENERATED_BODY()
+	
+public:
+	virtual void ConvertToInstancedStruct(FInstancedStruct& OutInstancedStruct) const override
+	{
+		OutInstancedStruct.InitializeAs(FAssetChooser::StaticStruct());
+		FAssetChooser& AssetChooser = OutInstancedStruct.GetMutable<FAssetChooser>();
+		AssetChooser.Asset = Asset;
+	}
+	
+	UPROPERTY(EditAnywhere, Category = "Parameters")
+	TObjectPtr<UObject> Asset;
+};
\ No newline at end of file
diff --git a/Engine/Plugins/Experimental/Chooser/Source/Chooser/Private/BoolColumn.cpp b/Engine/Plugins/Experimental/Chooser/Source/Chooser/Private/BoolColumn.cpp
--- a/Engine/Plugins/Experimental/Chooser/Source/Chooser/Private/BoolColumn.cpp
+++ b/Engine/Plugins/Experimental/Chooser/Source/Chooser/Private/BoolColumn.cpp
@@ -2,7 +2,7 @@
 #include "BoolColumn.h"
 #include "ChooserPropertyAccess.h"
 
-bool UChooserParameterBool_ContextProperty::GetValue(const UObject* ContextObject, bool& OutResult) const
+bool FBoolContextProperty::GetValue(const UObject* ContextObject, bool& OutResult) const
 {
 	UStruct* StructType = ContextObject->GetClass();
 	const void* Container = ContextObject;
@@ -19,18 +19,17 @@ bool UChooserParameterBool_ContextProperty::GetValue(const UObject* ContextObjec
 	return false;
 }
 
-UChooserColumnBool::UChooserColumnBool(const FObjectInitializer& ObjectInitializer)
+FBoolColumn::FBoolColumn()
 {
-	InputValue = ObjectInitializer.CreateDefaultSubobject<UChooserParameterBool_ContextProperty>(this, "InputValue");
-	InputValue.GetObject()->SetFlags(RF_Transactional);
+	InputValue.InitializeAs(FBoolContextProperty::StaticStruct());
 }
 
-void UChooserColumnBool::Filter(const UObject* ContextObject, const TArray<uint32>& IndexListIn, TArray<uint32>& IndexListOut)
+void FBoolColumn::Filter(const UObject* ContextObject, const TArray<uint32>& IndexListIn, TArray<uint32>& IndexListOut) const
 {
-	if (ContextObject && InputValue)
+	if (ContextObject && InputValue.IsValid())
 	{
 		bool Result = false;
-		InputValue->GetValue(ContextObject,Result);
+		InputValue.Get<FChooserParameterBoolBase>().GetValue(ContextObject,Result);
 		
 		for (uint32 Index : IndexListIn)
 		{
diff --git a/Engine/Plugins/Experimental/Chooser/Source/Chooser/Private/Chooser.cpp b/Engine/Plugins/Experimental/Chooser/Source/Chooser/Private/Chooser.cpp
--- a/Engine/Plugins/Experimental/Chooser/Source/Chooser/Private/Chooser.cpp
+++ b/Engine/Plugins/Experimental/Chooser/Source/Chooser/Private/Chooser.cpp
@@ -11,17 +11,66 @@ UChooserTable::UChooserTable(const FObjectInitializer& Initializer)
 
 }
 
-static IObjectChooser::EIteratorStatus StaticEvaluateChooser(const UObject* ContextObject, const UChooserTable* Chooser, IObjectChooser::FObjectChooserIteratorCallback Callback)
+#if WITH_EDITOR
+void UChooserTable::PostLoad()
+{
+	Super::PostLoad();
+
+	// convert old data if it exists
+
+	if (Results_DEPRECATED.Num() > 0 || Columns_DEPRECATED.Num() > 0)
+	{
+		ResultsStructs.Reserve(Results_DEPRECATED.Num());
+		ColumnsStructs.Reserve(Columns_DEPRECATED.Num());
+		
+		for(TScriptInterface<IObjectChooser>& Result : Results_DEPRECATED)
+		{
+			ResultsStructs.SetNum(ResultsStructs.Num() + 1);
+			IObjectChooser* ResultInterface = Result.GetInterface();
+			if (ResultInterface)
+			{
+				ResultInterface->ConvertToInstancedStruct(ResultsStructs.Last());
+			}
+		}
+		
+		for(TScriptInterface<IChooserColumn>& Column : Columns_DEPRECATED)
+		{
+			ColumnsStructs.SetNum(ColumnsStructs.Num() + 1);
+			IChooserColumn* ColumnInterface = Column.GetInterface();
+			if (ColumnInterface)
+			{
+				ColumnInterface->ConvertToInstancedStruct(ColumnsStructs.Last());
+			}
+		}
+
+		Results_DEPRECATED.SetNum(0);
+		Columns_DEPRECATED.SetNum(0);
+	}
+
+	// call PostLoad on Columns
+	for (FInstancedStruct& ColumnData : ColumnsStructs)
+	{
+		if (ColumnData.IsValid())
+		{
+			FChooserColumnBase& Column = ColumnData.GetMutable<FChooserColumnBase>();
+			Column.PostLoad();
+		}
+	}
+		
+}
+#endif
+
+static FObjectChooserBase::EIteratorStatus StaticEvaluateChooser(const UObject* ContextObject, const UChooserTable* Chooser, FObjectChooserBase::FObjectChooserIteratorCallback Callback)
 {
 	if (Chooser == nullptr)
 	{
-		return IObjectChooser::EIteratorStatus::Continue;
+		return FObjectChooserBase::EIteratorStatus::Continue;
 	}
 
 	TArray<uint32> Indices1;
 	TArray<uint32> Indices2;
 
-	int RowCount = Chooser->Results.Num();
+	int RowCount = Chooser->ResultsStructs.Num();
 	Indices1.SetNum(RowCount);
 	for(int i=0;i<RowCount;i++)
 	{
@@ -30,47 +79,43 @@ static IObjectChooser::EIteratorStatus StaticEvaluateChooser(const UObject* Cont
 	TArray<uint32>* IndicesOut = &Indices1;
 	TArray<uint32>* IndicesIn = &Indices2;
 
-	for (auto Column: Chooser->Columns)
+	for (const FInstancedStruct& ColumnData : Chooser->ColumnsStructs)
 	{
-		if (Column)
-		{
-			Swap(IndicesIn, IndicesOut);
-			IndicesOut->SetNum(0, false);
-			Column->Filter(ContextObject, *IndicesIn, *IndicesOut);
-		}
+		const FChooserColumnBase& Column = ColumnData.Get<FChooserColumnBase>(); 
+		Swap(IndicesIn, IndicesOut);
+		IndicesOut->SetNum(0, false);
+		Column.Filter(ContextObject, *IndicesIn, *IndicesOut);
 	}
 
 	// of the rows that passed all column filters, return the first one for which the result row succeeds (could fail eg for a nexted chooser where no rows passed)
 	for (uint32 SelectedIndex : *IndicesOut)
 	{
-		if (Chooser->Results.Num() > (int32)SelectedIndex)
+		if (Chooser->ResultsStructs.Num() > (int32)SelectedIndex)
 		{
-			if (const IObjectChooser* SelectedResult = Chooser->Results[SelectedIndex].GetInterface())
+			const FObjectChooserBase& SelectedResult = Chooser->ResultsStructs[SelectedIndex].Get<FObjectChooserBase>();
+			if (SelectedResult.ChooseMulti(ContextObject, Callback) == FObjectChooserBase::EIteratorStatus::Stop)
 			{
-				if (SelectedResult->ChooseMulti(ContextObject, Callback) == IObjectChooser::EIteratorStatus::Stop)
-				{
-					return IObjectChooser::EIteratorStatus::Stop;
-				}
+				return FObjectChooserBase::EIteratorStatus::Stop;
 			}
 		}
 	}
 	
-	return IObjectChooser::EIteratorStatus::Continue;
+	return FObjectChooserBase::EIteratorStatus::Continue;
 }
 
-UObject* UObjectChooser_EvaluateChooser::ChooseObject(const UObject* ContextObject) const
+UObject* FEvaluateChooser::ChooseObject(const UObject* ContextObject) const
 {
 	UObject* Result = nullptr;
 	StaticEvaluateChooser(ContextObject, Chooser, FObjectChooserIteratorCallback::CreateLambda([&Result](UObject* InResult)
 	{
 		Result = InResult;
-		return IObjectChooser::EIteratorStatus::Stop;
+		return FObjectChooserBase::EIteratorStatus::Stop;
 	}));
 
 	return Result;
 }
 
-IObjectChooser::EIteratorStatus UObjectChooser_EvaluateChooser::ChooseMulti(const UObject* ContextObject, FObjectChooserIteratorCallback Callback) const
+FObjectChooserBase::EIteratorStatus FEvaluateChooser::ChooseMulti(const UObject* ContextObject, FObjectChooserIteratorCallback Callback) const
 {
 	return StaticEvaluateChooser(ContextObject, Chooser, Callback);
 }
@@ -83,13 +128,23 @@ UChooserFunctionLibrary::UChooserFunctionLibrary(const FObjectInitializer& Objec
 UObject* UChooserFunctionLibrary::EvaluateChooser(const UObject* ContextObject, const UChooserTable* Chooser)
 {
 	UObject* Result = nullptr;
-	StaticEvaluateChooser(ContextObject, Chooser, IObjectChooser::FObjectChooserIteratorCallback::CreateLambda([&Result](UObject* InResult)
+	StaticEvaluateChooser(ContextObject, Chooser, FObjectChooserBase::FObjectChooserIteratorCallback::CreateLambda([&Result](UObject* InResult)
 	{
 		Result = InResult;
-		return IObjectChooser::EIteratorStatus::Stop;
+		return FObjectChooserBase::EIteratorStatus::Stop;
 	}));
 
 	return Result;
+}
+
+TArray<UObject*> UChooserFunctionLibrary::EvaluateChooserMulti(const UObject* ContextObject, const UChooserTable* Chooser)
+{
+	TArray<UObject*> Result;
+	StaticEvaluateChooser(ContextObject, Chooser, FObjectChooserBase::FObjectChooserIteratorCallback::CreateLambda([&Result](UObject* InResult)
+	{
+		Result.Add(InResult);
+		return FObjectChooserBase::EIteratorStatus::Continue;
+	}));
 
 	return Result;
-}
+}
\ No newline at end of file
diff --git a/Engine/Plugins/Experimental/Chooser/Source/Chooser/Private/EnumColumn.cpp b/Engine/Plugins/Experimental/Chooser/Source/Chooser/Private/EnumColumn.cpp
--- a/Engine/Plugins/Experimental/Chooser/Source/Chooser/Private/EnumColumn.cpp
+++ b/Engine/Plugins/Experimental/Chooser/Source/Chooser/Private/EnumColumn.cpp
@@ -6,7 +6,7 @@
 #include "IPropertyAccessEditor.h"
 #endif
 
-bool UChooserParameterEnum_ContextProperty::GetValue(const UObject* ContextObject, uint8& OutResult) const
+bool FEnumContextProperty::GetValue(const UObject* ContextObject, uint8& OutResult) const
 {
 	UStruct* StructType = ContextObject->GetClass();
 	const void* Container = ContextObject;
@@ -34,7 +34,7 @@ bool UChooserParameterEnum_ContextProperty::GetValue(const UObject* ContextObjec
 
 #if WITH_EDITOR
 
-void UChooserParameterEnum_ContextProperty::SetBinding(const TArray<FBindingChainElement>& InBindingChain)
+void FEnumContextProperty::SetBinding(const TArray<FBindingChainElement>& InBindingChain)
 {
 	const UEnum* PreviousEnum = Enum;
 	Enum = nullptr;
@@ -54,17 +54,16 @@ void UChooserParameterEnum_ContextProperty::SetBinding(const TArray<FBindingChai
 	if (Enum != PreviousEnum)
 	{
 		// Our enum type has changed! Need to refresh the UI to update enum value pickers.
-		EnumChanged.Broadcast();
+		OnEnumChanged.Broadcast();
 	}
 }
 
 #endif // WITH_EDITOR
 
-UChooserColumnEnum::UChooserColumnEnum(const FObjectInitializer& ObjectInitializer)
+FEnumColumn::FEnumColumn()
 {
 #if WITH_EDITOR
-	InputValue = ObjectInitializer.CreateDefaultSubobject<UChooserParameterEnum_ContextProperty>(this, "InputValue");
-	InputValue.GetObject()->SetFlags(RF_Transactional);
+	InputValue.InitializeAs(FEnumContextProperty::StaticStruct());
 	InputChanged();
 #endif
 }
@@ -97,12 +96,12 @@ bool FChooserEnumRowData::Evaluate(const uint8 LeftHandSide) const
 	}
 }
 
-void UChooserColumnEnum::Filter(const UObject* ContextObject, const TArray<uint32>& IndexListIn, TArray<uint32>& IndexListOut)
+void FEnumColumn::Filter(const UObject* ContextObject, const TArray<uint32>& IndexListIn, TArray<uint32>& IndexListOut) const
 {
 	uint8 Result = 0;
 	if (ContextObject != nullptr &&
-		InputValue != nullptr &&
-		InputValue->GetValue(ContextObject, Result))
+		InputValue.IsValid() &&
+		InputValue.Get<FChooserParameterEnumBase>().GetValue(ContextObject, Result))
 	{
 		for (const uint32 Index : IndexListIn)
 		{
diff --git a/Engine/Plugins/Experimental/Chooser/Source/Chooser/Private/FloatRangeColumn.cpp b/Engine/Plugins/Experimental/Chooser/Source/Chooser/Private/FloatRangeColumn.cpp
--- a/Engine/Plugins/Experimental/Chooser/Source/Chooser/Private/FloatRangeColumn.cpp
+++ b/Engine/Plugins/Experimental/Chooser/Source/Chooser/Private/FloatRangeColumn.cpp
@@ -2,7 +2,7 @@
 #include "FloatRangeColumn.h"
 #include "ChooserPropertyAccess.h"
 
-bool UChooserParameterFloat_ContextProperty::GetValue(const UObject* ContextObject, float& OutResult) const
+bool FFloatContextProperty::GetValue(const UObject* ContextObject, float& OutResult) const
 {
 	UStruct* StructType = ContextObject->GetClass();
 	const void* Container = ContextObject;
@@ -25,18 +25,17 @@ bool UChooserParameterFloat_ContextProperty::GetValue(const UObject* ContextObje
 	return false;
 }
 
-UChooserColumnFloatRange::UChooserColumnFloatRange(const FObjectInitializer& ObjectInitializer)
+FFloatRangeColumn::FFloatRangeColumn()
 {
-	InputValue = ObjectInitializer.CreateDefaultSubobject<UChooserParameterFloat_ContextProperty>(this, "InputValue");
-	InputValue.GetObject()->SetFlags(RF_Transactional);
+	InputValue.InitializeAs(FFloatContextProperty::StaticStruct());
 }
 
-void UChooserColumnFloatRange::Filter(const UObject* ContextObject, const TArray<uint32>& IndexListIn, TArray<uint32>& IndexListOut)
+void FFloatRangeColumn::Filter(const UObject* ContextObject, const TArray<uint32>& IndexListIn, TArray<uint32>& IndexListOut) const
 {
-	if (ContextObject && InputValue)
+	if (ContextObject && InputValue.IsValid())
 	{
 		float Result = 0.0f;
-		InputValue->GetValue(ContextObject, Result);
+		InputValue.Get<FChooserParameterFloatBase>().GetValue(ContextObject, Result);
 
 		for(uint32 Index : IndexListIn)
 		{
diff --git a/Engine/Plugins/Experimental/Chooser/Source/Chooser/Private/GameplayTagColumn.cpp b/Engine/Plugins/Experimental/Chooser/Source/Chooser/Private/GameplayTagColumn.cpp
--- a/Engine/Plugins/Experimental/Chooser/Source/Chooser/Private/GameplayTagColumn.cpp
+++ b/Engine/Plugins/Experimental/Chooser/Source/Chooser/Private/GameplayTagColumn.cpp
@@ -2,7 +2,7 @@
 #include "GameplayTagColumn.h"
 #include "ChooserPropertyAccess.h"
 
-bool UChooserParameterGameplayTag_ContextProperty::GetValue(const UObject* ContextObject, const FGameplayTagContainer*& OutResult) const
+bool FGameplayTagContextProperty::GetValue(const UObject* ContextObject, const FGameplayTagContainer*& OutResult) const
 {
 	UStruct* StructType = ContextObject->GetClass();
 	const void* Container = ContextObject;
@@ -19,22 +19,19 @@ bool UChooserParameterGameplayTag_ContextProperty::GetValue(const UObject* Conte
 	return false;
 }
 
-UChooserColumnGameplayTag::UChooserColumnGameplayTag(const FObjectInitializer& ObjectInitializer)
+FGameplayTagColumn::FGameplayTagColumn()
 {
-	InputValue = ObjectInitializer.CreateDefaultSubobject<UChooserParameterGameplayTag_ContextProperty>(this, "InputValue");
-	InputValue.GetObject()->SetFlags(RF_Transactional);
+	InputValue.InitializeAs(FGameplayTagContextProperty::StaticStruct());
 }
 
-void UChooserColumnGameplayTag::Filter(const UObject* ContextObject, const TArray<uint32>& IndexListIn, TArray<uint32>& IndexListOut)
+void FGameplayTagColumn::Filter(const UObject* ContextObject, const TArray<uint32>& IndexListIn, TArray<uint32>& IndexListOut) const
 {
-	if (ContextObject && InputValue)
+	const FGameplayTagContainer* Result = nullptr;
+	if (ContextObject && InputValue.IsValid() && InputValue.Get<FChooserParameterGameplayTagBase>().GetValue(ContextObject,Result))
 	{
-		const FGameplayTagContainer* Result = nullptr;
-		InputValue->GetValue(ContextObject,Result);
-
 		for (uint32 Index : IndexListIn)
 		{
-			 
+		 
 			if (RowValues.Num() > (int)Index)
 			{
 				if (RowValues[Index].IsEmpty())
diff --git a/Engine/Plugins/Experimental/Chooser/Source/Chooser/Private/ObjectChooser_Asset.cpp b/Engine/Plugins/Experimental/Chooser/Source/Chooser/Private/ObjectChooser_Asset.cpp
--- a/Engine/Plugins/Experimental/Chooser/Source/Chooser/Private/ObjectChooser_Asset.cpp
+++ b/Engine/Plugins/Experimental/Chooser/Source/Chooser/Private/ObjectChooser_Asset.cpp
@@ -2,7 +2,7 @@
 
 #include "ObjectChooser_Asset.h"
 
-UObject* UObjectChooser_Asset::ChooseObject(const UObject* ContextObject) const
+UObject* FAssetChooser::ChooseObject(const UObject* ContextObject) const
 {
 	return Asset;
 }
\ No newline at end of file
diff --git a/Engine/Plugins/Experimental/Chooser/Source/Chooser/Public/ChooserFunctionLibrary.h b/Engine/Plugins/Experimental/Chooser/Source/Chooser/Public/ChooserFunctionLibrary.h
--- a/Engine/Plugins/Experimental/Chooser/Source/Chooser/Public/ChooserFunctionLibrary.h
+++ b/Engine/Plugins/Experimental/Chooser/Source/Chooser/Public/ChooserFunctionLibrary.h
@@ -25,4 +25,12 @@ public:
 	UFUNCTION(BlueprintPure, meta = (BlueprintThreadSafe), Category = "Animation")
 	static UObject* EvaluateChooser(const UObject* ContextObject, const UChooserTable* ChooserTable);
 
+	/**
+    * Evaluate a chooser table and return the list of all selected UObjects
+    *
+    * @param ContextObject			(in) An Object from which the parameters to the Chooser Table will be read
+    * @param ChooserTable			(in) The ChooserTable asset
+    */
+    UFUNCTION(BlueprintPure, meta = (BlueprintThreadSafe), Category = "Animation")
+    static TArray<UObject*> EvaluateChooserMulti(const UObject* ContextObject, const UChooserTable* ChooserTable);
 };
diff --git a/Engine/Plugins/Experimental/Chooser/Source/Chooser/Public/IChooserColumn.h b/Engine/Plugins/Experimental/Chooser/Source/Chooser/Public/IChooserColumn.h
--- a/Engine/Plugins/Experimental/Chooser/Source/Chooser/Public/IChooserColumn.h
+++ b/Engine/Plugins/Experimental/Chooser/Source/Chooser/Public/IChooserColumn.h
@@ -5,6 +5,8 @@
 #include "CoreMinimal.h"
 #include "UObject/Object.h"
 #include "UObject/Interface.h"
+#include "IChooserParameterBase.h"
+#include "InstancedStruct.h"
 #include "IChooserColumn.generated.h"
 
 UINTERFACE(NotBlueprintType, meta = (CannotImplementInterfaceInBlueprint))
@@ -14,14 +16,55 @@ class CHOOSER_API UChooserColumn : public UInterface
 };
 
 class CHOOSER_API IChooserColumn 
+{
+	GENERATED_BODY()
+public:
+	virtual void ConvertToInstancedStruct(FInstancedStruct& OutInstancedStruct) const {}
+};
+
+
+USTRUCT()
+struct CHOOSER_API FChooserColumnBase
 {
 	GENERATED_BODY()
 
 public:
-	virtual void Filter(const UObject* ContextObject, const TArray<uint32>& IndexListIn, TArray<uint32>& IndexListOut) {};
+	virtual ~FChooserColumnBase() {}
+	virtual void PostLoad() {};
+	virtual void Filter(const UObject* ContextObject, const TArray<uint32>& IndexListIn, TArray<uint32>& IndexListOut) const {}
+
+#if WITH_EDITOR
 	virtual void SetNumRows(uint32 NumRows) {}
 	virtual void DeleteRows(const TArray<uint32> & RowIndices) {}
-	virtual UClass* GetInputValueInterface() {return nullptr;};
-	virtual UObject* GetInputValue() {return nullptr;}
-	virtual void SetInputValue(UObject* InputValue) {}
+	
+	virtual UScriptStruct* GetInputBaseType() const { return nullptr; };
+	virtual const UScriptStruct* GetInputType() const { return nullptr; };
+	virtual void SetInputType(const UScriptStruct* Type) { };
+	virtual FChooserParameterBase* GetInputValue() { return nullptr; };
+#endif
 };
+
+#if WITH_EDITOR
+#define CHOOSER_COLUMN_BOILERPLATE_NoSetInputType(ParameterType) \
+	virtual void SetNumRows(uint32 NumRows) override { RowValues.SetNum(NumRows); }\
+	virtual void DeleteRows(const TArray<uint32> & RowIndices )\
+	{\
+		for(uint32 Index : RowIndices)\
+		{\
+			RowValues.RemoveAt(Index);\
+		}\
+	}\
+	virtual UScriptStruct* GetInputBaseType() const override { return ParameterType::StaticStruct(); };\
+	virtual const UScriptStruct* GetInputType() const override { return InputValue.IsValid() ? InputValue.GetScriptStruct() : nullptr; };\
+	virtual FChooserParameterBase* GetInputValue() override { return InputValue.IsValid() ? &InputValue.GetMutable<FChooserParameterBase>() : nullptr; };
+	
+#define CHOOSER_COLUMN_BOILERPLATE(ParameterType) \
+	CHOOSER_COLUMN_BOILERPLATE_NoSetInputType(ParameterType); \
+	virtual void SetInputType(const UScriptStruct* Type) override { InputValue.InitializeAs(Type); };
+
+#else
+
+#define CHOOSER_COLUMN_BOILERPLATE(ParameterType)
+#define CHOOSER_COLUMN_BOILERPLATE_NoSetInputType(ParameterType)
+
+#endif
diff --git a/Engine/Plugins/Experimental/Chooser/Source/Chooser/Public/IChooserParameterBase.h b/Engine/Plugins/Experimental/Chooser/Source/Chooser/Public/IChooserParameterBase.h
new file mode 100644
--- /dev/null
+++ b/Engine/Plugins/Experimental/Chooser/Source/Chooser/Public/IChooserParameterBase.h
@@ -0,0 +1,18 @@
+// Copyright Epic Games, Inc. All Rights Reserved.
+
+#pragma once
+
+#include "CoreMinimal.h"
+#include "UObject/Object.h"
+#include "UObject/Interface.h"
+#include "IChooserParameterBase.generated.h"
+
+USTRUCT()
+struct FChooserParameterBase
+{
+	GENERATED_BODY()
+
+	virtual void GetDisplayName(FText& OutName) const { }
+
+	virtual ~FChooserParameterBase() {}
+};
\ No newline at end of file
diff --git a/Engine/Plugins/Experimental/Chooser/Source/Chooser/Public/IChooserParameterBool.h b/Engine/Plugins/Experimental/Chooser/Source/Chooser/Public/IChooserParameterBool.h
--- a/Engine/Plugins/Experimental/Chooser/Source/Chooser/Public/IChooserParameterBool.h
+++ b/Engine/Plugins/Experimental/Chooser/Source/Chooser/Public/IChooserParameterBool.h
@@ -5,6 +5,8 @@
 #include "CoreMinimal.h"
 #include "UObject/Object.h"
 #include "UObject/Interface.h"
+#include "IChooserParameterBase.h"
+#include "InstancedStruct.h"
 #include "IChooserParameterBool.generated.h"
 
 UINTERFACE(NotBlueprintType, meta = (CannotImplementInterfaceInBlueprint))
@@ -16,7 +18,15 @@ class CHOOSER_API UChooserParameterBool : public UInterface
 class CHOOSER_API IChooserParameterBool
 {
 	GENERATED_BODY()
-
 public:
-	virtual bool GetValue(const UObject* ContextObject, bool& OutResult) const { return false; }
+	virtual void ConvertToInstancedStruct(FInstancedStruct& OutInstancedStruct) const {}
 };
+
+USTRUCT()
+struct FChooserParameterBoolBase : public FChooserParameterBase
+{
+	GENERATED_BODY()
+    
+    public:
+    	virtual bool GetValue(const UObject* ContextObject, bool& OutResult) const { return false; }
+};
\ No newline at end of file
diff --git a/Engine/Plugins/Experimental/Chooser/Source/Chooser/Public/IChooserParameterEnum.h b/Engine/Plugins/Experimental/Chooser/Source/Chooser/Public/IChooserParameterEnum.h
--- a/Engine/Plugins/Experimental/Chooser/Source/Chooser/Public/IChooserParameterEnum.h
+++ b/Engine/Plugins/Experimental/Chooser/Source/Chooser/Public/IChooserParameterEnum.h
@@ -5,6 +5,7 @@
 #include "CoreMinimal.h"
 #include "UObject/Object.h"
 #include "UObject/Interface.h"
+#include "InstancedStruct.h"
 #include "IChooserParameterEnum.generated.h"
 
 UINTERFACE(NotBlueprintType, meta = (CannotImplementInterfaceInBlueprint))
@@ -16,12 +17,19 @@ class CHOOSER_API UChooserParameterEnum : public UInterface
 class CHOOSER_API IChooserParameterEnum
 {
 	GENERATED_BODY()
-
 public:
-	virtual bool GetValue(const UObject* ContextObject, uint8& OutResult) const { return false; }
+	virtual void ConvertToInstancedStruct(FInstancedStruct& OutInstancedStruct) const {}
+};
+
+USTRUCT()
+struct FChooserParameterEnumBase : public FChooserParameterBase
+{
+	GENERATED_BODY()
+	
+		virtual bool GetValue(const UObject* ContextObject, uint8& OutResult) const { return false; }
 
-#if WITH_EDITOR
-	virtual const UEnum* GetEnum() const { return nullptr; }
-	virtual FSimpleMulticastDelegate& OnEnumChanged() = 0;
-#endif
+	#if WITH_EDITOR
+    	virtual const UEnum* GetEnum() const { return nullptr; }
+    	FSimpleMulticastDelegate OnEnumChanged;
+    #endif
 };
\ No newline at end of file
diff --git a/Engine/Plugins/Experimental/Chooser/Source/Chooser/Public/IChooserParameterFloat.h b/Engine/Plugins/Experimental/Chooser/Source/Chooser/Public/IChooserParameterFloat.h
--- a/Engine/Plugins/Experimental/Chooser/Source/Chooser/Public/IChooserParameterFloat.h
+++ b/Engine/Plugins/Experimental/Chooser/Source/Chooser/Public/IChooserParameterFloat.h
@@ -16,7 +16,15 @@ class CHOOSER_API UChooserParameterFloat : public UInterface
 class CHOOSER_API IChooserParameterFloat
 {
 	GENERATED_BODY()
+public:
+	virtual void ConvertToInstancedStruct(FInstancedStruct& OutInstancedStruct) const {}
+};
 
+USTRUCT()
+struct FChooserParameterFloatBase : public FChooserParameterBase
+{
+	GENERATED_BODY()
+    
 public:
 	virtual bool GetValue(const UObject* ContextObject, float& OutResult) const { return false; }
 };
\ No newline at end of file
diff --git a/Engine/Plugins/Experimental/Chooser/Source/Chooser/Public/IChooserParameterGameplayTag.h b/Engine/Plugins/Experimental/Chooser/Source/Chooser/Public/IChooserParameterGameplayTag.h
--- a/Engine/Plugins/Experimental/Chooser/Source/Chooser/Public/IChooserParameterGameplayTag.h
+++ b/Engine/Plugins/Experimental/Chooser/Source/Chooser/Public/IChooserParameterGameplayTag.h
@@ -5,6 +5,7 @@
 #include "CoreMinimal.h"
 #include "UObject/Object.h"
 #include "UObject/Interface.h"
+#include "IChooserParameterBase.h"
 #include "IChooserParameterGameplayTag.generated.h"
 
 struct FGameplayTagContainer;
@@ -18,7 +19,13 @@ class CHOOSER_API UChooserParameterGameplayTag : public UInterface
 class CHOOSER_API IChooserParameterGameplayTag
 {
 	GENERATED_BODY()
-
 public:
+	virtual void ConvertToInstancedStruct(FInstancedStruct& OutInstancedStruct) const {}
+};
+
+USTRUCT()
+struct FChooserParameterGameplayTagBase : public FChooserParameterBase
+{
+	GENERATED_BODY()
 	virtual bool GetValue(const UObject* ContextObject, const FGameplayTagContainer*& OutResult) const { return false; }
-};
\ No newline at end of file
+};
diff --git a/Engine/Plugins/Experimental/Chooser/Source/Chooser/Public/IObjectChooser.h b/Engine/Plugins/Experimental/Chooser/Source/Chooser/Public/IObjectChooser.h
--- a/Engine/Plugins/Experimental/Chooser/Source/Chooser/Public/IObjectChooser.h
+++ b/Engine/Plugins/Experimental/Chooser/Source/Chooser/Public/IObjectChooser.h
@@ -5,6 +5,7 @@
 #include "CoreMinimal.h"
 #include "UObject/Object.h"
 #include "UObject/Interface.h"
+#include "InstancedStruct.h"
 #include "IObjectChooser.generated.h"
 
 UINTERFACE(NotBlueprintType, meta = (CannotImplementInterfaceInBlueprint))
@@ -14,10 +15,20 @@ class CHOOSER_API UObjectChooser : public UInterface
 };
 
 class CHOOSER_API IObjectChooser
+{
+	GENERATED_BODY()
+public:
+	virtual void ConvertToInstancedStruct(FInstancedStruct& OutInstancedStruct) const { }
+};
+
+
+USTRUCT()
+struct CHOOSER_API FObjectChooserBase
 {
 	GENERATED_BODY()
 
 public:
+	virtual ~FObjectChooserBase() {}
 
 	virtual UObject* ChooseObject(const UObject* ContextObject) const { return nullptr; };
 
@@ -33,4 +44,4 @@ public:
 		}
 		return EIteratorStatus::Continue;
 	};
-};
+};
\ No newline at end of file
diff --git a/Engine/Plugins/Experimental/Chooser/Source/ChooserEditor/ChooserEditor.Build.cs b/Engine/Plugins/Experimental/Chooser/Source/ChooserEditor/ChooserEditor.Build.cs
--- a/Engine/Plugins/Experimental/Chooser/Source/ChooserEditor/ChooserEditor.Build.cs
+++ b/Engine/Plugins/Experimental/Chooser/Source/ChooserEditor/ChooserEditor.Build.cs
@@ -27,7 +27,8 @@ namespace UnrealBuildTool.Rules
 					"BlueprintGraph",
 					"GraphEditor",
 					"GameplayTags",
-					"GameplayTagsEditor"
+					"GameplayTagsEditor",
+					"StructUtils"
 					// ... add private dependencies that you statically link with here ...
 				}
 			);
diff --git a/Engine/Plugins/Experimental/Chooser/Source/ChooserEditor/Private/BoolColumnEditor.cpp b/Engine/Plugins/Experimental/Chooser/Source/ChooserEditor/Private/BoolColumnEditor.cpp
--- a/Engine/Plugins/Experimental/Chooser/Source/ChooserEditor/Private/BoolColumnEditor.cpp
+++ b/Engine/Plugins/Experimental/Chooser/Source/ChooserEditor/Private/BoolColumnEditor.cpp
@@ -3,8 +3,8 @@
 #include "BoolColumnEditor.h"
 #include "BoolColumn.h"
 #include "ContextPropertyWidget.h"
-#include "ObjectChooserWidgetFactories.h"
 #include "ChooserTableEditor.h"
+#include "ObjectChooserWidgetFactories.h"
 #include "Widgets/Input/SCheckBox.h"
 #include "GraphEditorSettings.h"
 
@@ -13,17 +13,17 @@
 namespace UE::ChooserEditor
 {
 
-TSharedRef<SWidget> CreateBoolColumnWidget(UObject* Column, int Row)
+TSharedRef<SWidget> CreateBoolColumnWidget(UChooserTable* Chooser, FChooserColumnBase* Column, int Row)
 {
-	UChooserColumnBool* BoolColumn = Cast<UChooserColumnBool>(Column);
+	FBoolColumn* BoolColumn = static_cast<FBoolColumn*>(Column);
 
 	return SNew (SCheckBox)
-	.OnCheckStateChanged_Lambda([BoolColumn,Row](ECheckBoxState State)
+	.OnCheckStateChanged_Lambda([Chooser, BoolColumn,Row](ECheckBoxState State)
 	{
 		if (Row < BoolColumn->RowValues.Num())
 		{
 			const FScopedTransaction Transaction(LOCTEXT("Change Bool Value", "Change Bool Value"));
-			BoolColumn->Modify(true);
+			Chooser->Modify(true);
 			BoolColumn->RowValues[Row] = (State == ECheckBoxState::Checked);
 		}
 	})
@@ -34,17 +34,15 @@ TSharedRef<SWidget> CreateBoolColumnWidget(UObject* Column, int Row)
 	});
 }
 	
-TSharedRef<SWidget> CreateBoolPropertyWidget(UObject* Object, UClass* ContextClass)
+TSharedRef<SWidget> CreateBoolPropertyWidget(UObject* TransactionObject, void* Value, UClass* ContextClass)
 {
-	return CreatePropertyWidget<UChooserParameterBool_ContextProperty>(Object, ContextClass, GetDefault<UGraphEditorSettings>()->BooleanPinTypeColor);
+	return CreatePropertyWidget<FBoolContextProperty>(TransactionObject, Value, ContextClass, GetDefault<UGraphEditorSettings>()->BooleanPinTypeColor);
 }
 	
 void RegisterBoolWidgets()
 {
-	FObjectChooserWidgetFactories::ChooserWidgetCreators.Add(UChooserParameterBool_ContextProperty::StaticClass(), CreateBoolPropertyWidget);
-	FObjectChooserWidgetFactories::ChooserTextConverter.Add(UChooserParameterBool_ContextProperty::StaticClass(), ConvertToText_ContextProperty<UChooserParameterBool_ContextProperty>);
-
-	FChooserTableEditor::ColumnWidgetCreators.Add(UChooserColumnBool::StaticClass(), CreateBoolColumnWidget);
+	FObjectChooserWidgetFactories::ChooserWidgetCreators.Add(FBoolContextProperty::StaticStruct(), CreateBoolPropertyWidget);
+	FChooserTableEditor::ColumnWidgetCreators.Add(FBoolColumn::StaticStruct(), CreateBoolColumnWidget);
 }
 	
 }
diff --git a/Engine/Plugins/Experimental/Chooser/Source/ChooserEditor/Private/ChooserEditorModule.cpp b/Engine/Plugins/Experimental/Chooser/Source/ChooserEditor/Private/ChooserEditorModule.cpp
--- a/Engine/Plugins/Experimental/Chooser/Source/ChooserEditor/Private/ChooserEditorModule.cpp
+++ b/Engine/Plugins/Experimental/Chooser/Source/ChooserEditor/Private/ChooserEditorModule.cpp
@@ -20,15 +20,6 @@ void FModule::StartupModule()
 	RegisterFloatRangeWidgets();
 	RegisterBoolWidgets();
 	RegisterEnumWidgets();
-	
-	InterfacePropertyTypeIdentifier = MakeShared<UE::ChooserEditor::FPropertyTypeIdentifier>();
-	
-	FPropertyEditorModule& PropertyModule = FModuleManager::LoadModuleChecked<FPropertyEditorModule>("PropertyEditor");
-	
-	PropertyModule.RegisterCustomPropertyTypeLayout(
-		FInterfaceProperty::StaticClass()->GetFName(),
-		FOnGetPropertyTypeCustomizationInstance::CreateLambda([] { return MakeShared<UE::ChooserEditor::FInterfacePropertyTypeCustomization>(); }),
-		InterfacePropertyTypeIdentifier);
 }
 
 void FModule::ShutdownModule()
diff --git a/Engine/Plugins/Experimental/Chooser/Source/ChooserEditor/Private/ChooserEditorModule.h b/Engine/Plugins/Experimental/Chooser/Source/ChooserEditor/Private/ChooserEditorModule.h
--- a/Engine/Plugins/Experimental/Chooser/Source/ChooserEditor/Private/ChooserEditorModule.h
+++ b/Engine/Plugins/Experimental/Chooser/Source/ChooserEditor/Private/ChooserEditorModule.h
@@ -3,7 +3,7 @@
 #pragma once
 
 #include "CoreMinimal.h"
-#include "InterfacePropertyTypeCustomization.h"
+#include "Modules/ModuleInterface.h"
 
 namespace UE::ChooserEditor
 {
@@ -16,7 +16,6 @@ private:
 	virtual void StartupModule() override;
 	virtual void ShutdownModule() override;
 
-	TSharedPtr<UE::ChooserEditor::FPropertyTypeIdentifier> InterfacePropertyTypeIdentifier;
 };
 
 }
\ No newline at end of file
diff --git a/Engine/Plugins/Experimental/Chooser/Source/ChooserEditor/Private/ChooserTableEditor.cpp b/Engine/Plugins/Experimental/Chooser/Source/ChooserEditor/Private/ChooserTableEditor.cpp
--- a/Engine/Plugins/Experimental/Chooser/Source/ChooserEditor/Private/ChooserTableEditor.cpp
+++ b/Engine/Plugins/Experimental/Chooser/Source/ChooserEditor/Private/ChooserTableEditor.cpp
@@ -15,6 +15,7 @@
 #include "Widgets/Input/SComboButton.h"
 #include "SAssetDropTarget.h"
 #include "SClassViewer.h"
+#include "StructViewerModule.h"
 #include "SourceCodeNavigation.h"
 #include "Chooser.h"
 #include "DetailCategoryBuilder.h"
@@ -307,7 +308,6 @@ public:
 		Operation->ChooserEditor = InEditor;
 		Operation->RowIndex = InRowIndex;
 		Operation->DefaultHoverText = LOCTEXT("Chooser Row", "Chooser Row");
-		FObjectChooserWidgetFactories::ConvertToText(InEditor->GetChooser()->Results[InRowIndex].GetObject(), Operation->DefaultHoverText);
 		Operation->CurrentHoverText = Operation->DefaultHoverText;
 			
 		Operation->Construct();
@@ -387,9 +387,9 @@ public:
 		static FName Result = "Result";
 		static FName Handles = "Handles";
 		
-		if (RowIndex->RowIndex < Chooser->Results.Num())
+		if (RowIndex->RowIndex < Chooser->ResultsStructs.Num())
 		{
-			if (ColumnName == Handles && RowIndex->RowIndex < Chooser->Results.Num())
+			if (ColumnName == Handles && RowIndex->RowIndex < Chooser->ResultsStructs.Num())
 			{
 				// row drag handle
 				
@@ -397,15 +397,13 @@ public:
 			}
 			else if (ColumnName == Result) 
 			{
-				UObject* RowValue = Chooser->Results[RowIndex->RowIndex].GetObject();
-				TSharedPtr<SWidget> ResultWidget = FObjectChooserWidgetFactories::CreateWidget(UObjectChooser::StaticClass(), RowValue,Chooser->ContextObjectType,
-				FOnClassPicked::CreateLambda([this, RowIndex=RowIndex->RowIndex](UClass* ChosenClass)
+				TSharedPtr<SWidget> ResultWidget = FObjectChooserWidgetFactories::CreateWidget(Chooser, FObjectChooserBase::StaticStruct(), Chooser->ResultsStructs[RowIndex->RowIndex].GetMutableMemory(), Chooser->ResultsStructs[RowIndex->RowIndex].GetScriptStruct(),Chooser->ContextObjectType,
+				FOnStructPicked::CreateLambda([this, RowIndex=RowIndex->RowIndex](const UScriptStruct* ChosenStruct)
 				{
 					const FScopedTransaction Transaction(LOCTEXT("Change Row Result Type", "Change Row Result Type"));
-					UObject* RowValue = NewObject<UObject>(Chooser, ChosenClass, NAME_None, RF_Transactional);
+					Chooser->ResultsStructs[RowIndex].InitializeAs(ChosenStruct);
 					Chooser->Modify(true);
-					Chooser->Results[RowIndex] = RowValue;
-					FObjectChooserWidgetFactories::CreateWidget(UObjectChooser::StaticClass(), RowValue, Chooser->ContextObjectType, FOnClassPicked(), &CacheBorder);
+					FObjectChooserWidgetFactories::CreateWidget(Chooser, FObjectChooserBase::StaticStruct(), Chooser->ResultsStructs[RowIndex].GetMutableMemory(), ChosenStruct, Chooser->ContextObjectType, FOnStructPicked(), &CacheBorder);
 				}),
 				&CacheBorder
 				);
@@ -414,19 +412,19 @@ public:
 			else
 			{
 				const int ColumnIndex = ColumnName.GetNumber() - 1;
-				if (ColumnIndex < Chooser->Columns.Num() && ColumnIndex >=0)
+				if (ColumnIndex < Chooser->ColumnsStructs.Num() && ColumnIndex >=0)
 				{
-					TScriptInterface<IChooserColumn>& Column = Chooser->Columns[ColumnIndex];
+					FChooserColumnBase* Column = &Chooser->ColumnsStructs[ColumnIndex].GetMutable<FChooserColumnBase>();
+					const UStruct * ColumnStruct = Chooser->ColumnsStructs[ColumnIndex].GetScriptStruct();
 					TSharedPtr<SWidget> ColumnWidget;
-					UClass* ColumnClass = Column.GetObject()->GetClass();
-					while (ColumnClass && !ColumnWidget.IsValid())
+					while (ColumnStruct && !ColumnWidget.IsValid())
 					{
-						if (auto Creator = FChooserTableEditor::ColumnWidgetCreators.Find(ColumnClass))
+						if (auto Creator = FChooserTableEditor::ColumnWidgetCreators.Find(ColumnStruct))
 						{
-							ColumnWidget = (*Creator)(Column.GetObject(), RowIndex->RowIndex);
+							ColumnWidget = (*Creator)(Chooser, Column, RowIndex->RowIndex);
 							break;
 						}
-						ColumnClass = ColumnClass->GetSuperClass();
+						ColumnStruct = ColumnStruct->GetSuperStruct();
 					}
 					
 					if (ColumnWidget.IsValid())
@@ -436,7 +434,7 @@ public:
 				}
 			}
 		}
-		else if (RowIndex->RowIndex == Chooser->Results.Num())
+		else if (RowIndex->RowIndex == Chooser->ResultsStructs.Num())
         {
 			// on the row past the end, show an Add button in the result column
 			if (ColumnName == Result)
@@ -493,10 +491,15 @@ void FChooserTableEditor::UpdateTableColumns()
 					.ManualWidth(300));
 
 	FName ColumnId("ChooserColumn", 1);
-	int NumColumns = Chooser->Columns.Num();	
+	int NumColumns = Chooser->ColumnsStructs.Num();	
 	for(int ColumnIndex = 0; ColumnIndex < NumColumns; ColumnIndex++)
 	{
-		TScriptInterface<IChooserColumn>& Column = Chooser->Columns[ColumnIndex];
+		FChooserColumnBase& Column = Chooser->ColumnsStructs[ColumnIndex].GetMutable<FChooserColumnBase>();
+		TSharedPtr<SWidget> HeaderWidget = nullptr;
+		if (FChooserParameterBase* InputValue = Column.GetInputValue())
+		{
+			HeaderWidget = FObjectChooserWidgetFactories::CreateWidget(Chooser, InputValue, Column.GetInputType(), Chooser->ContextObjectType);
+		}
 
 		HeaderRow->AddColumn(SHeaderRow::FColumn::FArguments()
 			.ColumnId(ColumnId)
@@ -514,8 +517,7 @@ void FChooserTableEditor::UpdateTableColumns()
 					FUIAction(
 						FExecuteAction::CreateLambda([this,Chooser,ColumnIndex, ColumnId, &Column]()
 						{
-							DetailsView->SetObject(Chooser->Columns[ColumnIndex].GetObject());
-							SelectedColumn = ColumnId;
+							SelectColumn(ColumnId.GetNumber() - 1);
 						})
 						)
 					);
@@ -526,30 +528,30 @@ void FChooserTableEditor::UpdateTableColumns()
 						{
 							const FScopedTransaction Transaction(LOCTEXT("Delete Column Transaction", "Delete Column"));
 							Chooser->Modify(true);
-							Chooser->Columns.RemoveAt(ColumnIndex);
+							Chooser->ColumnsStructs.RemoveAt(ColumnIndex);
 							UpdateTableColumns();
 						})
 						));
 				
 				MenuBuilder.AddSubMenu(LOCTEXT("Input Type", "Input Type"),
 					LOCTEXT("InputTypeToolTip", "Change input parameter type"),
-					FNewMenuDelegate::CreateLambda([this, &Column](FMenuBuilder& Builder)
+					FNewMenuDelegate::CreateLambda([this, Chooser, &ColumnIndex](FMenuBuilder& Builder)
 					{
-						FClassViewerInitializationOptions Options;
-						Options.ClassFilters.Add(MakeShared<FInterfaceClassFilter>(Column->GetInputValueInterface()));
-						Options.NameTypeToDisplay = EClassViewerNameTypeToDisplay::DisplayName;
+						FStructViewerInitializationOptions Options;
+						Options.StructFilter = MakeShared<FStructFilter>(Chooser->ColumnsStructs[ColumnIndex].Get<FChooserColumnBase>().GetInputBaseType());
+						Options.NameTypeToDisplay = EStructViewerNameTypeToDisplay::DisplayName;
 						
 						// Add class filter for columns here
-						TSharedRef<SWidget> Widget = FModuleManager::LoadModuleChecked<FClassViewerModule>("ClassViewer").CreateClassViewer(Options, FOnClassPicked::CreateLambda([this, &Column](UClass* ChosenClass)
+						TSharedRef<SWidget> Widget = FModuleManager::LoadModuleChecked<FStructViewerModule>("StructViewer").CreateStructViewer(Options, FOnStructPicked::CreateLambda([this, ColumnIndex](const UScriptStruct* ChosenStruct)
 						{
 							const FScopedTransaction Transaction(LOCTEXT("SetColumnInputType", "Set Column Input Type"));
-							UObject* Value = NewObject<UObject>(Column.GetObject(), ChosenClass, NAME_None, RF_Transactional);
-							Column.GetObject()->Modify(true);
-							Column->SetInputValue(Value);
+							UChooserTable* Chooser = Cast<UChooserTable>(EditingObjects[0]);
+							Chooser->ColumnsStructs[ColumnIndex].GetMutable<FChooserColumnBase>().SetInputType(ChosenStruct);
+							Chooser->Modify(true);
 							UpdateTableColumns();
 							UpdateTableRows();
 						}));
-						
+					
 						Builder.AddWidget(Widget, FText());
 					}));
 				
@@ -564,28 +566,18 @@ void FChooserTableEditor::UpdateTableColumns()
 				.BorderBackgroundColor_Lambda([this, ColumnId] ()
 				{
 					// unclear why this color is coming out much darker
-					return SelectedColumn == ColumnId ? FSlateColor(FColor(0x00, 0x70, 0xe0, 0xFF)) : FSlateColor(FLinearColor(0.05f,0.05f,0.05f));
+					return (SelectedColumn && SelectedColumn->Column == ColumnId.GetNumber() - 1) ? FSlateColor(FColor(0x00, 0x70, 0xe0, 0xFF)) : FSlateColor(FLinearColor(0.05f,0.05f,0.05f));
 				})
 				.OnMouseButtonDown_Lambda([this, ColumnIndex, ColumnId](	const FGeometry&, const FPointerEvent& PointerEvent)
 				{
 					UChooserTable* Chooser = Cast<UChooserTable>(EditingObjects[0]);
 					TableView->ClearSelection();
 					
-					if (PointerEvent.GetEffectingButton() == EKeys::LeftMouseButton )
-					{
-						DetailsView->SetObject(Chooser->Columns[ColumnIndex].GetObject());
-						SelectedColumn = ColumnId;
-						return FReply::Handled();
-					}
-					else 
-					{
-						SelectedColumn = ColumnId;
-						
-						return FReply::Handled();
-					}
+					SelectColumn(ColumnId.GetNumber() - 1);
+					return FReply::Handled();
 				})
 				[
-					FObjectChooserWidgetFactories::CreateWidget(Chooser->Columns[ColumnIndex]->GetInputValue(), Chooser->ContextObjectType).ToSharedRef()
+					HeaderWidget ? HeaderWidget.ToSharedRef() : SNullWidget::NullWidget
 				]
 				
 			]);
@@ -618,24 +610,25 @@ TSharedRef<SDockTab> FChooserTableEditor::SpawnTableTab( const FSpawnTabArgs& Ar
 	
 	CreateColumnComboButton = SNew(SComboButton).OnGetMenuContent_Lambda([this]()
 	{
-		FClassViewerInitializationOptions Options;
-		Options.ClassFilters.Add(MakeShared<FInterfaceClassFilter>(UChooserColumn::StaticClass()));
-		Options.NameTypeToDisplay = EClassViewerNameTypeToDisplay::DisplayName;
+		FStructViewerInitializationOptions Options;
+		Options.StructFilter = MakeShared<FStructFilter>(FChooserColumnBase::StaticStruct());
+		Options.NameTypeToDisplay = EStructViewerNameTypeToDisplay::DisplayName;
 		
 		// Add class filter for columns here
-		TSharedRef<SWidget> Widget = FModuleManager::LoadModuleChecked<FClassViewerModule>("ClassViewer").CreateClassViewer(Options, FOnClassPicked::CreateLambda([this](UClass* ChosenClass)
+		TSharedRef<SWidget> Widget = FModuleManager::LoadModuleChecked<FStructViewerModule>("StructViewer").CreateStructViewer(Options, FOnStructPicked::CreateLambda([this](const UScriptStruct* ChosenStruct)
 		{
 			CreateColumnComboButton->SetIsOpen(false);
 			UChooserTable* Chooser = Cast<UChooserTable>(EditingObjects[0]);
 			const FScopedTransaction Transaction(LOCTEXT("Add Column Transaction", "Add Column"));
 			Chooser->Modify(true);
-			UObject* ColumnObject = NewObject<UObject>(Chooser, ChosenClass, NAME_None, RF_Transactional);
-			Chooser->Columns.Add(TScriptInterface<IChooserColumn>(ColumnObject));
+
+			Chooser->ColumnsStructs.SetNum(Chooser->ColumnsStructs.Num()+1);
+			Chooser->ColumnsStructs.Last().InitializeAs(ChosenStruct);
+			
 			UpdateTableColumns();
 			UpdateTableRows();
-			SelectedColumn = "ChooserColumn";
-			SelectedColumn.SetNumber(Chooser->Columns.Num());
-			DetailsView->SetObject(ColumnObject);
+
+			SelectColumn(Chooser->ColumnsStructs.Num() - 1);
 		}));
 		return Widget;
 	})
@@ -654,22 +647,22 @@ TSharedRef<SDockTab> FChooserTableEditor::SpawnTableTab( const FSpawnTabArgs& Ar
 		]
 		.OnGetMenuContent_Lambda([this]()
 		{
-			FClassViewerInitializationOptions Options;
-			Options.ClassFilters.Add(MakeShared<FInterfaceClassFilter>(UObjectChooser::StaticClass()) );
-			Options.NameTypeToDisplay = EClassViewerNameTypeToDisplay::DisplayName;
+			FStructViewerInitializationOptions Options;
+			Options.StructFilter = MakeShared<FStructFilter>(FObjectChooserBase::StaticStruct());
+			Options.NameTypeToDisplay = EStructViewerNameTypeToDisplay::DisplayName;
+			
+			TSharedRef<SWidget> Widget = FModuleManager::LoadModuleChecked<FStructViewerModule>("StructViewer").CreateStructViewer(Options, FOnStructPicked::CreateLambda([this](const UScriptStruct* ChosenStruct)
+			{
+				CreateRowComboButton->SetIsOpen(false);
+				UChooserTable* Chooser = Cast<UChooserTable>(EditingObjects[0]);
+				const FScopedTransaction Transaction(LOCTEXT("Add Row Transaction", "Add Row"));
+				Chooser->Modify(true);
+
+				Chooser->ResultsStructs.SetNum(Chooser->ResultsStructs.Num()+1);
+				Chooser->ResultsStructs.Last().InitializeAs(ChosenStruct);
+				UpdateTableRows();
+			}));
 			
-			// Add class filter for columns here
-			TSharedRef<SWidget> Widget = FModuleManager::LoadModuleChecked<FClassViewerModule>("ClassViewer").CreateClassViewer(Options,
-				FOnClassPicked::CreateLambda([this](UClass* ChosenClass)
-				{
-					CreateRowComboButton->SetIsOpen(false);
-					UChooserTable* Chooser = Cast<UChooserTable>(EditingObjects[0]);
-					const FScopedTransaction Transaction(LOCTEXT("Add Row Transaction", "Add Row"));
-					Chooser->Modify(true);
-					UObject* ResultObject = NewObject<UObject>(Chooser, ChosenClass, NAME_None, RF_Transactional);
-					Chooser->Results.Add(TScriptInterface<IObjectChooser>(ResultObject));
-					UpdateTableRows();
-				}));
 			return Widget;
 		});
 
@@ -697,12 +690,13 @@ TSharedRef<SDockTab> FChooserTableEditor::SpawnTableTab( const FSpawnTabArgs& Ar
 						RowsToDelete.Sort([](int32 A, int32 B){ return A>B; });
 						for(uint32 RowIndex : RowsToDelete)
 						{
-							Chooser->Results.RemoveAt(RowIndex);
+							Chooser->ResultsStructs.RemoveAt(RowIndex);
 						}
 
-						for(auto& Column:Chooser->Columns)
+						for(FInstancedStruct& ColumnData : Chooser->ColumnsStructs)
 						{
-							Column.GetInterface()->DeleteRows(RowsToDelete);
+							FChooserColumnBase& Column = ColumnData.GetMutable<FChooserColumnBase>();
+							Column.DeleteRows(RowsToDelete);
 						}
 
 						UpdateTableRows();
@@ -714,7 +708,7 @@ TSharedRef<SDockTab> FChooserTableEditor::SpawnTableTab( const FSpawnTabArgs& Ar
 				)
 				.OnSelectionChanged_Lambda([this](TSharedPtr<FChooserTableRow> SelectedItem,  ESelectInfo::Type SelectInfo)
 				{
-					SelectedColumn = "None";
+					SelectedColumn = nullptr;
 					if (SelectedItem)
 					{
 						SelectedRows.SetNum(0);
@@ -723,7 +717,7 @@ TSharedRef<SDockTab> FChooserTableEditor::SpawnTableTab( const FSpawnTabArgs& Ar
 						TObjectPtr<UChooserRowDetails> Selection = NewObject<UChooserRowDetails>();
 						Selection->Chooser = Chooser;
 						Selection->Row = SelectedItem->RowIndex;
-						Selection->SetFlags(RF_Transactional); // for undo?
+						Selection->SetFlags(RF_Transactional);
 						SelectedRows.Add(Selection);
 											
 						TArray<UObject*> DetailsObjects;
@@ -755,7 +749,7 @@ void FChooserTableEditor::UpdateTableRows()
 {
 	UChooserTable* Chooser = Cast<UChooserTable>(EditingObjects[0]);
 	int32 OldNum = TableRows.Num();
-	int32 NewNum = Chooser->Results.Num();
+	int32 NewNum = Chooser->ResultsStructs.Num();
 
 	// Sync the TableRows array which drives the ui table to match the number of results.
 	// Add 1 at the end, for the "Add Row" control
@@ -767,12 +761,10 @@ void FChooserTableEditor::UpdateTableRows()
 	}
 
 	// Make sure each column has the same number of row datas as there are results
-	for(auto& Column : Chooser->Columns)
+	for(FInstancedStruct& ColumnData : Chooser->ColumnsStructs)
 	{
-		if (IChooserColumn* ColumnInterface = Column.GetInterface())
-		{
-			ColumnInterface->SetNumRows(NewNum);
-		}
+		FChooserColumnBase& Column = ColumnData.GetMutable<FChooserColumnBase>();
+		Column.SetNumRows(NewNum);
 	}
 
 	if (TableView.IsValid())
@@ -826,96 +818,70 @@ TSharedRef<FChooserTableEditor> FChooserTableEditor::CreateEditor( const EToolki
 	NewEditor->InitEditor( Mode, InitToolkitHost, ObjectsToEdit, GetDetailsViewObjects );
 	return NewEditor;
 }
-
+	
 void FChooserTableEditor::SelectColumn(int Index)
 {
 	UChooserTable* Chooser = GetChooser();
-	if (Index < Chooser->Columns.Num())
+	if (Index < Chooser->ColumnsStructs.Num())
 	{
-		SelectedColumn = "Column";
-		SelectedColumn.SetNumber(Index);
-		DetailsView->SetObject(Chooser->Columns[Index].GetObject());
+		if (SelectedColumn == nullptr)
+		{
+			SelectedColumn = NewObject<UChooserColumnDetails>();
+			SelectedColumn->Chooser = Chooser;
+		}
+		SelectedColumn->Column = Index;
+		DetailsView->SetObject(SelectedColumn, true);
 	}
 }
 	
 void FChooserTableEditor::DeleteColumn(int Index)
 {
 	UChooserTable* Chooser = GetChooser();
-	if (Index < Chooser->Columns.Num())
+	if (Index < Chooser->ColumnsStructs.Num())
 	{
-		Chooser->Columns.RemoveAt(Index);
+		Chooser->ColumnsStructs.RemoveAt(Index);
 		UpdateTableColumns();
 	}
 }
 
 /// Result widgets
 
-TMap<const UClass*, TFunction<TSharedRef<SWidget> (UObject* Column, int Row)>> FChooserTableEditor::ColumnWidgetCreators;
+TMap<const UStruct*, TFunction<TSharedRef<SWidget> (UChooserTable* Chooser, FChooserColumnBase* Column, int Row)>> FChooserTableEditor::ColumnWidgetCreators;
 
-TSharedRef<SWidget> CreateAssetWidget(UObject* Object, UClass* ContextClass)
+TSharedRef<SWidget> CreateAssetWidget(UObject* TransactionObject, void* Value, UClass* ContextClass)
 {
-	UObjectChooser_Asset* DIAsset = Cast<UObjectChooser_Asset>(Object);
+	FAssetChooser* DIAsset = static_cast<FAssetChooser*>(Value);
 
 	UObject* Asset = DIAsset->Asset;
 
-	UChooserTable* Chooser = Object->GetTypedOuter<UChooserTable>();
+	UChooserTable* Chooser = Cast<UChooserTable>(TransactionObject);
 	
 	return SNew(SObjectPropertyEntryBox)
 		.AllowedClass((Chooser!=nullptr && Chooser->OutputObjectType!=nullptr) ? Chooser->OutputObjectType.Get() : UObject::StaticClass())
 		.ObjectPath_Lambda([DIAsset](){ return DIAsset->Asset ? DIAsset->Asset.GetPath() : "";})
-		.OnObjectChanged_Lambda([DIAsset](const FAssetData& AssetData)
+		.OnObjectChanged_Lambda([TransactionObject, DIAsset](const FAssetData& AssetData)
 		{
 			const FScopedTransaction Transaction(LOCTEXT("Edit Asset", "Edit Asset"));
-			DIAsset->Modify(true);
+			TransactionObject->Modify(true);
 			DIAsset->Asset = AssetData.GetAsset();
 		});
 }
 
-TSharedRef<SWidget> CreateEvaluateChooserWidget(UObject* Object, UClass* ContextObject)
+TSharedRef<SWidget> CreateEvaluateChooserWidget(UObject* TransactionObject, void* Value, UClass* ContextObject)
 {
-	UObjectChooser_EvaluateChooser* EvaluateChooser = Cast<UObjectChooser_EvaluateChooser>(Object);
+	FEvaluateChooser* EvaluateChooser = static_cast<FEvaluateChooser*>(Value);
 	
 	return SNew(SObjectPropertyEntryBox)
 		.AllowedClass(UChooserTable::StaticClass())
 		.ObjectPath_Lambda([EvaluateChooser](){ return EvaluateChooser->Chooser ? EvaluateChooser->Chooser.GetPath() : "";})
-		.OnObjectChanged_Lambda([EvaluateChooser](const FAssetData& AssetData){ EvaluateChooser->Chooser = Cast<UChooserTable>(AssetData.GetAsset()); });
-}
-
-void ConvertToText_Asset(const UObject* Object, FText& OutText)
-{
-	const UObjectChooser_Asset* AssetInterface = Cast<UObjectChooser_Asset>(Object);
-	if (AssetInterface->Asset == nullptr)
-	{
-		OutText = LOCTEXT("NoChooser", "[No Chooser]");
-	}
-	else
-	{
-		OutText = FText::FromString(AssetInterface->Asset->GetName());
-	}
-}
-
-void ConvertToText_EvaluateChooser(const UObject* Object, FText& OutText)
-{
-	const UObjectChooser_EvaluateChooser* EvaluateChooser = Cast<UObjectChooser_EvaluateChooser>(Object);
-	if (EvaluateChooser->Chooser == nullptr)
-	{
-		OutText = LOCTEXT("NoChooser", "[No Chooser]");
-	}
-	else
-	{
-		OutText = FText::FromString(EvaluateChooser->Chooser->GetName());
-	}
-}
-
-TSharedRef<SWidget> CreateObjectWidget(UObject* Object, UClass* ContextClass)
-{
-	FText ObjectText = FText::FromString(Object->GetName());
-	FObjectChooserWidgetFactories::ConvertToText(Object, ObjectText);
-	
-	return SNew (STextBlock).Text(ObjectText); // could make this use Text_Lambda, to have it update correctly?
+		.OnObjectChanged_Lambda([TransactionObject, EvaluateChooser](const FAssetData& AssetData)
+		{
+			const FScopedTransaction Transaction(LOCTEXT("Edit Chooser", "Edit Chooser"));
+			TransactionObject->Modify(true);
+			EvaluateChooser->Chooser = Cast<UChooserTable>(AssetData.GetAsset());
+		});
 }
 
-
 class FChooserRowDetails : public IDetailCustomization
 {
 public:
@@ -940,26 +906,25 @@ void FChooserRowDetails::CustomizeDetails(IDetailLayoutBuilder& DetailBuilder)
 	UChooserRowDetails* Row = Cast<UChooserRowDetails>(Objects[0]);
 	UChooserTable* Chooser = Row->Chooser;
 	
-	if (Chooser->Results.Num() > Row->Row)
+	if (Chooser->ResultsStructs.Num() > Row->Row)
 	{
 		IDetailCategoryBuilder& PropertiesCategory = DetailBuilder.EditCategory("Row Properties");
 
 		TSharedPtr<IPropertyHandle> ChooserProperty = DetailBuilder.GetProperty("Chooser", Row->StaticClass());
 		DetailBuilder.HideProperty(ChooserProperty);
 	
-		TSharedPtr<IPropertyHandle> ResultsArrayProperty = ChooserProperty->GetChildHandle("Results");
+		TSharedPtr<IPropertyHandle> ResultsArrayProperty = ChooserProperty->GetChildHandle("ResultsStructs");
 		TSharedPtr<IPropertyHandle> CurrentResultProperty = ResultsArrayProperty->AsArray()->GetElement(Row->Row);
 		IDetailPropertyRow& NewResultProperty = PropertiesCategory.AddProperty(CurrentResultProperty);
 		NewResultProperty.DisplayName(LOCTEXT("ResultColumnName","Result"));
 		NewResultProperty.ShowPropertyButtons(false); // hide array add button
 		NewResultProperty.ShouldAutoExpand(true);
 	
-		for(int ColumnIndex=0; ColumnIndex<Chooser->Columns.Num(); ColumnIndex++)
+		for(int ColumnIndex=0; ColumnIndex<Chooser->ColumnsStructs.Num(); ColumnIndex++)
 		{
-			TScriptInterface<IChooserColumn>& Column = Chooser->Columns[ColumnIndex];
-			UObject* ColumnObject = Chooser->Columns[ColumnIndex].GetObject();
-			
-			TSharedPtr<IPropertyHandle> ColumnDataProperty =  DetailBuilder.AddObjectPropertyData({ColumnObject}, "RowValues");
+			FChooserColumnBase& Column = Chooser->ColumnsStructs[ColumnIndex].GetMutable<FChooserColumnBase>();
+			TSharedRef<FStructOnScope> StructOnScope = MakeShared<FStructOnScope>(Chooser->ColumnsStructs[ColumnIndex].GetScriptStruct(), reinterpret_cast<uint8*>(&Column));
+			TSharedPtr<IPropertyHandle> ColumnDataProperty = DetailBuilder.AddStructurePropertyData({StructOnScope}, "RowValues");
 			uint32 NumElements = 0;
 			ColumnDataProperty->AsArray()->GetNumElements(NumElements);
 			if (Row->Row < (int)NumElements)
@@ -968,9 +933,9 @@ void FChooserRowDetails::CustomizeDetails(IDetailLayoutBuilder& DetailBuilder)
 	
 				IDetailPropertyRow& NewColumnProperty = PropertiesCategory.AddProperty(CellData);
 				FText DisplayName = LOCTEXT("No Input Value", "No Input Value");
-				if (UObject* InputValue = Column.GetInterface()->GetInputValue())
+				if (FChooserParameterBase* InputValue = Column.GetInputValue())
 				{
-					FObjectChooserWidgetFactories::ConvertToText(InputValue, DisplayName);
+					InputValue->GetDisplayName(DisplayName); // todo
 				}
 				NewColumnProperty.DisplayName(DisplayName);
 				NewColumnProperty.ShowPropertyButtons(false); // hide array add button
@@ -980,18 +945,56 @@ void FChooserRowDetails::CustomizeDetails(IDetailLayoutBuilder& DetailBuilder)
 	}
 }
 
-void FChooserTableEditor::RegisterWidgets()
+class FChooserColumnDetails : public IDetailCustomization
 {
-	FObjectChooserWidgetFactories::ChooserTextConverter.Add(UObjectChooser_Asset::StaticClass(), ConvertToText_Asset);
-	FObjectChooserWidgetFactories::ChooserTextConverter.Add(UObjectChooser_EvaluateChooser::StaticClass(), ConvertToText_EvaluateChooser);
+public:
+	FChooserColumnDetails() {};
+	virtual ~FChooserColumnDetails() override {};
+
+	static TSharedRef<IDetailCustomization> MakeInstance()
+	{
+		return MakeShareable( new FChooserColumnDetails() );
+	}
+
+	// IDetailCustomization interface
+	virtual void CustomizeDetails(class IDetailLayoutBuilder& DetailBuilder) override;
+};
+
+// Make the details panel show the values for the selected row, showing each column value
+void FChooserColumnDetails::CustomizeDetails(IDetailLayoutBuilder& DetailBuilder)
+{
+	TArray<TWeakObjectPtr<UObject>> Objects;
+	DetailBuilder.GetObjectsBeingCustomized(Objects);
+
+	UChooserColumnDetails* Column = Cast<UChooserColumnDetails>(Objects[0]);
+	UChooserTable* Chooser = Column->Chooser;
 	
-	FObjectChooserWidgetFactories::ChooserWidgetCreators.Add(UObject::StaticClass(), CreateObjectWidget);
-	FObjectChooserWidgetFactories::ChooserWidgetCreators.Add(UObjectChooser_Asset::StaticClass(), CreateAssetWidget);
-	FObjectChooserWidgetFactories::ChooserWidgetCreators.Add(UObjectChooser_EvaluateChooser::StaticClass(), CreateEvaluateChooserWidget);
+	if (Chooser->ColumnsStructs.Num() > Column->Column)
+	{
+		IDetailCategoryBuilder& PropertiesCategory = DetailBuilder.EditCategory("Row Properties");
+
+		TSharedPtr<IPropertyHandle> ChooserProperty = DetailBuilder.GetProperty("Chooser", Column->StaticClass());
+		DetailBuilder.HideProperty(ChooserProperty);
+	
+		TSharedPtr<IPropertyHandle> ColumnsArrayProperty = ChooserProperty->GetChildHandle("ColumnsStructs");
+		TSharedPtr<IPropertyHandle> CurrentColumnProperty = ColumnsArrayProperty->AsArray()->GetElement(Column->Column);
+		IDetailPropertyRow& NewResultProperty = PropertiesCategory.AddProperty(CurrentColumnProperty);
+		NewResultProperty.DisplayName(LOCTEXT("Selected Column","Selected Column"));
+		NewResultProperty.ShowPropertyButtons(false); // hide array add button
+		NewResultProperty.ShouldAutoExpand(true);
+	}
+}
+
+void FChooserTableEditor::RegisterWidgets()
+{
+	// todo: fallback widget
+	FObjectChooserWidgetFactories::ChooserWidgetCreators.Add(FAssetChooser::StaticStruct(), CreateAssetWidget);
+	FObjectChooserWidgetFactories::ChooserWidgetCreators.Add(FEvaluateChooser::StaticStruct(), CreateEvaluateChooserWidget);
 	
 
 	FPropertyEditorModule& PropertyModule = FModuleManager::LoadModuleChecked<FPropertyEditorModule>("PropertyEditor");
 	PropertyModule.RegisterCustomClassLayout("ChooserRowDetails", FOnGetDetailCustomizationInstance::CreateStatic(&FChooserRowDetails::MakeInstance));	
+	PropertyModule.RegisterCustomClassLayout("ChooserColumnDetails", FOnGetDetailCustomizationInstance::CreateStatic(&FChooserColumnDetails::MakeInstance));	
 }
 }
 
diff --git a/Engine/Plugins/Experimental/Chooser/Source/ChooserEditor/Private/ChooserTableEditor.h b/Engine/Plugins/Experimental/Chooser/Source/ChooserEditor/Private/ChooserTableEditor.h
--- a/Engine/Plugins/Experimental/Chooser/Source/ChooserEditor/Private/ChooserTableEditor.h
+++ b/Engine/Plugins/Experimental/Chooser/Source/ChooserEditor/Private/ChooserTableEditor.h
@@ -28,6 +28,17 @@ public:
 	int Row;
 };
 
+// Class used for chooser editor details customization
+UCLASS()
+class UChooserColumnDetails : public UObject
+{
+	GENERATED_BODY()
+public:
+	UPROPERTY(EditAnywhere, Instanced, Category="Hidden")
+	TObjectPtr<UChooserTable> Chooser;
+	int Column;
+};
+
 namespace UE::ChooserEditor
 {
 	class FChooserTableEditor : public FAssetEditorToolkit, public FSelfRegisteringEditorUndoClient, public FNotifyHook
@@ -114,6 +125,7 @@ namespace UE::ChooserEditor
 		/** The objects open within this editor */
 		TArray<UObject*> EditingObjects;
 
+		TObjectPtr<UChooserColumnDetails> SelectedColumn;
 		TArray<TObjectPtr<UChooserRowDetails>> SelectedRows;
 
 		void UpdateTableColumns();
@@ -124,12 +136,11 @@ namespace UE::ChooserEditor
 		
 		TSharedPtr<SHeaderRow> HeaderRow;
 		TSharedPtr<SListView<TSharedPtr<FChooserTableRow>>> TableView;
-
-		FName SelectedColumn;
+		
 	public:
 		TSharedPtr<SComboButton>& GetCreateRowComboButton() { return CreateRowComboButton; };
 
-		static TMap<const UClass*, TFunction<TSharedRef<SWidget> (UObject* Column, int Row)>> ColumnWidgetCreators;
+		static TMap<const UStruct*, TFunction<TSharedRef<SWidget> (UChooserTable* Chooser, FChooserColumnBase* Column, int Row)>> ColumnWidgetCreators;
 	
 		/** The name given to all instances of this type of editor */
 		static const FName ToolkitFName;
diff --git a/Engine/Plugins/Experimental/Chooser/Source/ChooserEditor/Private/EnumColumnEditor.cpp b/Engine/Plugins/Experimental/Chooser/Source/ChooserEditor/Private/EnumColumnEditor.cpp
--- a/Engine/Plugins/Experimental/Chooser/Source/ChooserEditor/Private/EnumColumnEditor.cpp
+++ b/Engine/Plugins/Experimental/Chooser/Source/ChooserEditor/Private/EnumColumnEditor.cpp
@@ -20,24 +20,24 @@ namespace UE::ChooserEditor
 			: _RowIndex(-1)
 		{}
 
-		SLATE_ARGUMENT(UChooserColumnEnum*, EnumColumn)
+		SLATE_ARGUMENT(FEnumColumn*, EnumColumn)
 		SLATE_ATTRIBUTE(int, RowIndex);
                 
 		SLATE_END_ARGS()
 
 		TSharedRef<SWidget> CreateEnumComboBox()
 		{
-			if (const UChooserColumnEnum* EnumColumnPointer = EnumColumn.Get())
+			if (const FEnumColumn* EnumColumnPointer = EnumColumn)
 			{
-				if (EnumColumnPointer->InputValue)
+				if (EnumColumnPointer->InputValue.IsValid())
 				{
-					if (const UEnum* Enum = EnumColumnPointer->InputValue->GetEnum())
+					if (const UEnum* Enum = EnumColumnPointer->InputValue.Get<FChooserParameterEnumBase>().GetEnum())
 					{
 						return SNew(SEnumComboBox, Enum)
 							.CurrentValue_Lambda([this]()
 							{
 								int Row = RowIndex.Get();
-								return EnumColumn.Get()->RowValues.IsValidIndex(Row) ? static_cast<int32>(EnumColumn->RowValues[Row].Value) : 0;
+								return EnumColumn->RowValues.IsValidIndex(Row) ? static_cast<int32>(EnumColumn->RowValues[Row].Value) : 0;
 							})
 							.OnEnumSelectionChanged_Lambda([this](int32 EnumValue, ESelectInfo::Type)
 							{
@@ -45,7 +45,8 @@ namespace UE::ChooserEditor
 								if (EnumColumn->RowValues.IsValidIndex(Row))
 								{
 									const FScopedTransaction Transaction(LOCTEXT("Edit RHS", "Edit Enum Value"));
-									EnumColumn->Modify(true);
+									// todo: need a reference to the UChooserTable to modify
+									// EnumColumn->Modify(true);
 									EnumColumn->RowValues[Row].Value = static_cast<uint8>(EnumValue);
 								}
 							});
@@ -66,7 +67,7 @@ namespace UE::ChooserEditor
 			RowIndex = InArgs._RowIndex;
 			EnumColumn = InArgs._EnumColumn;
 
-			if (UChooserColumnEnum* EnumColumnPointer = EnumColumn.Get())
+			if (FEnumColumn* EnumColumnPointer = EnumColumn)
 			{
 				EnumChangedHandle = EnumColumnPointer->OnEnumChanged.AddSP(this, &SEnumCell::UpdateEnumComboBox);
 			}
@@ -76,21 +77,21 @@ namespace UE::ChooserEditor
 
 		~SEnumCell()
 		{
-			if (UChooserColumnEnum* EnumColumnPointer = EnumColumn.Get())
+			if (FEnumColumn* EnumColumnPointer = EnumColumn)
 			{
 				EnumColumnPointer->OnEnumChanged.Remove(EnumChangedHandle);
 			}
 		}
 
 	private:
-		TWeakObjectPtr<UChooserColumnEnum> EnumColumn;
+		FEnumColumn* EnumColumn = nullptr;
 		TAttribute<int> RowIndex;
 		FDelegateHandle EnumChangedHandle;
 	};
 
-	TSharedRef<SWidget> CreateEnumColumnWidget(UObject* Column, int Row)
+	TSharedRef<SWidget> CreateEnumColumnWidget(UChooserTable* Chooser, FChooserColumnBase* Column, int Row)
 	{
-		UChooserColumnEnum* EnumColumn = CastChecked<UChooserColumnEnum>(Column);
+		FEnumColumn* EnumColumn = static_cast<FEnumColumn*>(Column);
 		
 		return SNew(SHorizontalBox)
 		+ SHorizontalBox::Slot().MaxWidth(95.0f)
@@ -100,12 +101,12 @@ namespace UE::ChooserEditor
 			{
 				return EnumColumn->RowValues.IsValidIndex(Row) ? static_cast<int32>(EnumColumn->RowValues[Row].Comparison) : 0;
 			})
-			.OnEnumSelectionChanged_Lambda([EnumColumn, Row](int32 EnumValue, ESelectInfo::Type)
+			.OnEnumSelectionChanged_Lambda([Chooser, EnumColumn, Row](int32 EnumValue, ESelectInfo::Type)
 			{
 				if (EnumColumn->RowValues.IsValidIndex(Row))
 				{
 					const FScopedTransaction Transaction(LOCTEXT("Edit Comparison", "Edit Comparison Operation"));
-					EnumColumn->Modify(true);
+					Chooser->Modify(true);
 					EnumColumn->RowValues[Row].Comparison = static_cast<EChooserEnumComparison>(EnumValue);
 				}
 			})
@@ -118,17 +119,15 @@ namespace UE::ChooserEditor
 
 	
 
-TSharedRef<SWidget> CreateEnumPropertyWidget(UObject* Object, UClass* ContextClass)
+TSharedRef<SWidget> CreateEnumPropertyWidget(UObject* TransactionObject, void* Value, UClass* ContextClass)
 {
-	return CreatePropertyWidget<UChooserParameterEnum_ContextProperty>(Object, ContextClass, GetDefault<UGraphEditorSettings>()->BytePinTypeColor);
+	return CreatePropertyWidget<FEnumContextProperty>(TransactionObject, Value, ContextClass, GetDefault<UGraphEditorSettings>()->BytePinTypeColor);
 }
 	
 void RegisterEnumWidgets()
 {
-	FObjectChooserWidgetFactories::ChooserWidgetCreators.Add(UChooserParameterEnum_ContextProperty::StaticClass(), CreateEnumPropertyWidget);
-	FObjectChooserWidgetFactories::ChooserTextConverter.Add(UChooserParameterEnum_ContextProperty::StaticClass(), ConvertToText_ContextProperty<UChooserParameterEnum_ContextProperty>);
-
-	FChooserTableEditor::ColumnWidgetCreators.Add(UChooserColumnEnum::StaticClass(), CreateEnumColumnWidget);
+	FObjectChooserWidgetFactories::ChooserWidgetCreators.Add(FEnumContextProperty::StaticStruct(), CreateEnumPropertyWidget);
+	FChooserTableEditor::ColumnWidgetCreators.Add(FEnumColumn::StaticStruct(), CreateEnumColumnWidget);
 }
 	
 }
diff --git a/Engine/Plugins/Experimental/Chooser/Source/ChooserEditor/Private/FloatRangeColumnEditor.cpp b/Engine/Plugins/Experimental/Chooser/Source/ChooserEditor/Private/FloatRangeColumnEditor.cpp
--- a/Engine/Plugins/Experimental/Chooser/Source/ChooserEditor/Private/FloatRangeColumnEditor.cpp
+++ b/Engine/Plugins/Experimental/Chooser/Source/ChooserEditor/Private/FloatRangeColumnEditor.cpp
@@ -12,14 +12,14 @@
 
 namespace UE::ChooserEditor
 {
-TSharedRef<SWidget> CreateFloatPropertyWidget(UObject* Object, UClass* ContextClass)
+TSharedRef<SWidget> CreateFloatPropertyWidget(UObject* TransactionObject, void* Value, UClass* ContextClass)
 {
-	return CreatePropertyWidget<UChooserParameterFloat_ContextProperty>(Object, ContextClass, GetDefault<UGraphEditorSettings>()->FloatPinTypeColor);
+	return CreatePropertyWidget<FFloatContextProperty>(TransactionObject, Value, ContextClass, GetDefault<UGraphEditorSettings>()->FloatPinTypeColor);
 }
 	
-TSharedRef<SWidget> CreateFloatRangeColumnWidget(UObject* Column, int Row)
+TSharedRef<SWidget> CreateFloatRangeColumnWidget(UChooserTable* Chooser, FChooserColumnBase* Column, int Row)
 {
-	UChooserColumnFloatRange* FloatRangeColumn = Cast<UChooserColumnFloatRange>(Column);
+	FFloatRangeColumn* FloatRangeColumn = static_cast<FFloatRangeColumn*>(Column);
 
 	return SNew(SHorizontalBox)
 	+ SHorizontalBox::Slot().MaxWidth(10)
@@ -37,12 +37,12 @@ TSharedRef<SWidget> CreateFloatRangeColumnWidget(UObject* Column, int Row)
 		{
 			return (Row < FloatRangeColumn->RowValues.Num()) ? FloatRangeColumn->RowValues[Row].Min : 0;
 		})
-		.OnValueCommitted_Lambda([FloatRangeColumn, Row](float NewValue, ETextCommit::Type CommitType)
+		.OnValueCommitted_Lambda([Chooser, FloatRangeColumn, Row](float NewValue, ETextCommit::Type CommitType)
 		{
 			if (Row < FloatRangeColumn->RowValues.Num())
 			{
 				const FScopedTransaction Transaction(LOCTEXT("Edit Min Value", "Edit Min Value"));
-				FloatRangeColumn->Modify(true);
+				Chooser->Modify(true);
 				FloatRangeColumn->RowValues[Row].Min = NewValue;
 			}
 		})
@@ -62,12 +62,12 @@ TSharedRef<SWidget> CreateFloatRangeColumnWidget(UObject* Column, int Row)
 		{
 			return (Row < FloatRangeColumn->RowValues.Num()) ? FloatRangeColumn->RowValues[Row].Max : 0;
 		})
-		.OnValueCommitted_Lambda([FloatRangeColumn, Row](float NewValue, ETextCommit::Type CommitType)
+		.OnValueCommitted_Lambda([Chooser, FloatRangeColumn, Row](float NewValue, ETextCommit::Type CommitType)
 		{
 			if (Row < FloatRangeColumn->RowValues.Num())
 			{
 				const FScopedTransaction Transaction(LOCTEXT("Edit Max", "Edit Max Value"));
-				FloatRangeColumn->Modify(true);
+				Chooser->Modify(true);
 				FloatRangeColumn->RowValues[Row].Max = NewValue;
 			}
 		})
@@ -81,10 +81,8 @@ TSharedRef<SWidget> CreateFloatRangeColumnWidget(UObject* Column, int Row)
 	
 void RegisterFloatRangeWidgets()
 {
-	FObjectChooserWidgetFactories::ChooserWidgetCreators.Add(UChooserParameterFloat_ContextProperty::StaticClass(), CreateFloatPropertyWidget);
-	FObjectChooserWidgetFactories::ChooserTextConverter.Add(UChooserParameterFloat_ContextProperty::StaticClass(), ConvertToText_ContextProperty<UChooserParameterFloat_ContextProperty>);
-
-	FChooserTableEditor::ColumnWidgetCreators.Add(UChooserColumnFloatRange::StaticClass(), CreateFloatRangeColumnWidget);
+	FObjectChooserWidgetFactories::ChooserWidgetCreators.Add(FFloatContextProperty::StaticStruct(), CreateFloatPropertyWidget);
+	FChooserTableEditor::ColumnWidgetCreators.Add(FFloatRangeColumn::StaticStruct(), CreateFloatRangeColumnWidget);
 }
 	
 }
diff --git a/Engine/Plugins/Experimental/Chooser/Source/ChooserEditor/Private/GameplayTagColumnEditor.cpp b/Engine/Plugins/Experimental/Chooser/Source/ChooserEditor/Private/GameplayTagColumnEditor.cpp
--- a/Engine/Plugins/Experimental/Chooser/Source/ChooserEditor/Private/GameplayTagColumnEditor.cpp
+++ b/Engine/Plugins/Experimental/Chooser/Source/ChooserEditor/Private/GameplayTagColumnEditor.cpp
@@ -2,21 +2,21 @@
 
 #include "GameplayTagColumnEditor.h"
 #include "ContextPropertyWidget.h"
+#include "ChooserTableEditor.h"
 #include "GameplayTagColumn.h"
 #include "ObjectChooserWidgetFactories.h"
-#include "ChooserTableEditor.h"
 #include "SGameplayTagWidget.h"
 #include "SSimpleComboButton.h"
 #include "GraphEditorSettings.h"
 
-#define LOCTEXT_NAMESPACE "GameplayTagColumnEditor"
+#define LOCTEXT_NAMESPACE "FGameplayTagColumnEditor"
 
 namespace UE::ChooserEditor
 {
 
-TSharedRef<SWidget> CreateGameplayTagColumnWidget(UObject* Column, int Row)
+TSharedRef<SWidget> CreateGameplayTagColumnWidget(UChooserTable* Chooser, FChooserColumnBase* Column, int Row)
 {
-	UChooserColumnGameplayTag* GameplayTagColumn = Cast<UChooserColumnGameplayTag>(Column);
+	FGameplayTagColumn* GameplayTagColumn = static_cast<struct FGameplayTagColumn*>(Column);
 
 	return SNew(SSimpleComboButton)
 		.Text_Lambda([GameplayTagColumn, Row]()
@@ -35,12 +35,12 @@ TSharedRef<SWidget> CreateGameplayTagColumnWidget(UObject* Column, int Row)
 				return FText();
 			}
 		})	
-		.OnGetMenuContent_Lambda([GameplayTagColumn, Row]()
+		.OnGetMenuContent_Lambda([Chooser, GameplayTagColumn, Row]()
 		{
 			if (Row < GameplayTagColumn->RowValues.Num())
 			{
 				TArray<SGameplayTagWidget::FEditableGameplayTagContainerDatum> EditableContainers;
-				EditableContainers.Emplace(GameplayTagColumn, &(GameplayTagColumn->RowValues[Row]));
+				EditableContainers.Emplace(Chooser, &(GameplayTagColumn->RowValues[Row]));
 				return TSharedRef<SWidget>(SNew(SGameplayTagWidget, EditableContainers));
 			}
 
@@ -49,17 +49,15 @@ TSharedRef<SWidget> CreateGameplayTagColumnWidget(UObject* Column, int Row)
 	);
 }
 
-TSharedRef<SWidget> CreateGameplayTagPropertyWidget(UObject* Object, UClass* ContextClass)
+TSharedRef<SWidget> CreateGameplayTagPropertyWidget(UObject* TransactionObject, void* Value, UClass* ContextClass)
 {
-	return CreatePropertyWidget<UChooserParameterGameplayTag_ContextProperty>(Object, ContextClass, GetDefault<UGraphEditorSettings>()->StructPinTypeColor);
+	return CreatePropertyWidget<FGameplayTagContextProperty>(TransactionObject, Value, ContextClass, GetDefault<UGraphEditorSettings>()->StructPinTypeColor);
 }
 	
 void RegisterGameplayTagWidgets()
 {
-	FObjectChooserWidgetFactories::ChooserWidgetCreators.Add(UChooserParameterGameplayTag_ContextProperty::StaticClass(), CreateGameplayTagPropertyWidget);
-	FObjectChooserWidgetFactories::ChooserTextConverter.Add(UChooserParameterGameplayTag_ContextProperty::StaticClass(), ConvertToText_ContextProperty<UChooserParameterGameplayTag_ContextProperty>);
-
-	FChooserTableEditor::ColumnWidgetCreators.Add(UChooserColumnGameplayTag::StaticClass(), CreateGameplayTagColumnWidget);
+	FObjectChooserWidgetFactories::ChooserWidgetCreators.Add(FGameplayTagContextProperty::StaticStruct(), CreateGameplayTagPropertyWidget);
+	FChooserTableEditor::ColumnWidgetCreators.Add(FGameplayTagColumn::StaticStruct(), CreateGameplayTagColumnWidget);
 }
 	
 }
diff --git a/Engine/Plugins/Experimental/Chooser/Source/ChooserEditor/Private/InterfacePropertyTypeCustomization.cpp b/Engine/Plugins/Experimental/Chooser/Source/ChooserEditor/Private/InterfacePropertyTypeCustomization.cpp
deleted file mode 100644
diff --git a/Engine/Plugins/Experimental/Chooser/Source/ChooserEditor/Private/InterfacePropertyTypeCustomization.h b/Engine/Plugins/Experimental/Chooser/Source/ChooserEditor/Private/InterfacePropertyTypeCustomization.h
deleted file mode 100644
diff --git a/Engine/Plugins/Experimental/Chooser/Source/ChooserEditor/Private/ObjectChooserWidgetFactories.cpp b/Engine/Plugins/Experimental/Chooser/Source/ChooserEditor/Private/ObjectChooserWidgetFactories.cpp
--- a/Engine/Plugins/Experimental/Chooser/Source/ChooserEditor/Private/ObjectChooserWidgetFactories.cpp
+++ b/Engine/Plugins/Experimental/Chooser/Source/ChooserEditor/Private/ObjectChooserWidgetFactories.cpp
@@ -1,82 +1,66 @@
 // Copyright Epic Games, Inc. All Rights Reserved.
 
 #include "ObjectChooserWidgetFactories.h"
+
 #include "Modules/ModuleManager.h"
 #include "Widgets/Input/SHyperlink.h"
 #include "Widgets/Input/SNumericEntryBox.h"
 #include "Widgets/Input/SComboButton.h"
 #include "SClassViewer.h"
 #include "DetailCategoryBuilder.h"
+#include "IObjectChooser.h"
 #include "Widgets/Layout/SWidgetSwitcher.h"
 #include "ObjectChooserClassFilter.h"
+#include "StructViewerModule.h"
 
 #define LOCTEXT_NAMESPACE "DataInterfaceEditor"
 
 namespace UE::ChooserEditor
 {
 
-void FObjectChooserWidgetFactories::ConvertToText(UObject* Object, FText& OutText)
-{
-	UClass* Class = Object->GetClass();
-	while (Class)
-	{
-		if (FChooserTextConverter* TextConverter = FObjectChooserWidgetFactories::ChooserTextConverter.Find(Class))
-		{
-			(*TextConverter)(Object, OutText);
-			break;
-		}
-		Class = Class->GetSuperClass();
-	}
-}
-	
-TSharedPtr<SWidget> FObjectChooserWidgetFactories::CreateWidget(UObject* Value, UClass* ContextClass)
+TSharedPtr<SWidget> FObjectChooserWidgetFactories::CreateWidget(UObject* TransactionObject, void* Value, const UStruct* ValueType, UClass* ContextClass)
 {
 	if (Value)
 	{
-		UClass* Class = Value->GetClass();
-		while (Class)
+		while (ValueType)
 		{
-			if (FChooserWidgetCreator* Creator = FObjectChooserWidgetFactories::ChooserWidgetCreators.Find(Class))
+			if (FChooserWidgetCreator* Creator = ChooserWidgetCreators.Find(ValueType))
 			{
-				return (*Creator)(Value, ContextClass);
-				break;
+				return (*Creator)(TransactionObject, Value, ContextClass);
 			}
-			Class = Class->GetSuperClass();
+			ValueType = ValueType->GetSuperStruct();
 		}
 	}
 
 	return nullptr;
 }
 
-TSharedPtr<SWidget> FObjectChooserWidgetFactories::CreateWidget(UClass* InterfaceType, UObject* Value, UClass* ContextClass, const FOnClassPicked& CreateClassCallback, TSharedPtr<SBorder>* InnerWidget)
+TSharedPtr<SWidget> FObjectChooserWidgetFactories::CreateWidget(UObject* TransactionObject, const UScriptStruct* BaseType, void* Value, const UStruct* ValueType, UClass* ContextClass, const FOnStructPicked& CreateClassCallback, TSharedPtr<SBorder>* InnerWidget)
 {
-	TSharedPtr<SWidget> LeftWidget = CreateWidget(Value, ContextClass);
-
+	TSharedPtr<SWidget> LeftWidget = CreateWidget(TransactionObject, Value, ValueType, ContextClass);
 	
 	if (!LeftWidget.IsValid())
 	{
 		LeftWidget = SNew(STextBlock).Text(LOCTEXT("SelectDataType", "Select Data Type..."));
 	}
 
-
 	// button for replacing data with a different Data Interface class
 	TSharedPtr<SComboButton> Button = SNew(SComboButton)
 			.ComboButtonStyle(FAppStyle::Get(), "SimpleComboButton");
 	
-	Button->SetOnGetMenuContent(FOnGetContent::CreateLambda([InterfaceType, Button, CreateClassCallback]()
+	Button->SetOnGetMenuContent(FOnGetContent::CreateLambda([BaseType, Button, CreateClassCallback]()
 	{
-		FClassViewerInitializationOptions Options;
-		Options.ClassFilters.Add(MakeShared<FInterfaceClassFilter>(InterfaceType));
-		Options.NameTypeToDisplay = EClassViewerNameTypeToDisplay::DisplayName;
+		FStructViewerInitializationOptions Options;
+		Options.StructFilter = MakeShared<FStructFilter>(BaseType);
+		Options.NameTypeToDisplay = EStructViewerNameTypeToDisplay::DisplayName;
 		
-		// Add class filter for columns here
-		TSharedRef<SWidget>  ClassMenu = FModuleManager::LoadModuleChecked<FClassViewerModule>("ClassViewer").
-		CreateClassViewer(Options, FOnClassPicked::CreateLambda([Button, CreateClassCallback](UClass* Class)
+		TSharedRef<SWidget> Widget = FModuleManager::LoadModuleChecked<FStructViewerModule>("StructViewer").CreateStructViewer(Options, FOnStructPicked::CreateLambda(
+			[Button, CreateClassCallback](const UScriptStruct* ChosenStruct)
 		{
 			Button->SetIsOpen(false);
-			CreateClassCallback.Execute(Class);
+			CreateClassCallback.Execute(ChosenStruct);
 		}));
-		return ClassMenu;
+		return Widget;
 	}));
 
 	TSharedPtr <SBorder> Border;
@@ -111,17 +95,10 @@ TSharedPtr<SWidget> FObjectChooserWidgetFactories::CreateWidget(UClass* Interfac
 	return Widget;
 }
 
-TMap<const UClass*, FChooserTextConverter> FObjectChooserWidgetFactories::ChooserTextConverter;
-TMap<const UClass*, FChooserWidgetCreator> FObjectChooserWidgetFactories::ChooserWidgetCreators;
-
-void ConvertToText_Base(const UObject* Object, FText& OutText)
-{
-	OutText = FText::FromString(Object->GetName());
-}
+TMap<const UStruct*, FChooserWidgetCreator> FObjectChooserWidgetFactories::ChooserWidgetCreators;
 
 void FObjectChooserWidgetFactories::RegisterWidgets()
 {
-	ChooserTextConverter.Add(UObject::StaticClass(), ConvertToText_Base);
 }
 	
 }
diff --git a/Engine/Plugins/Experimental/Chooser/Source/ChooserEditor/Public/ContextPropertyWidget.h b/Engine/Plugins/Experimental/Chooser/Source/ChooserEditor/Public/ContextPropertyWidget.h
--- a/Engine/Plugins/Experimental/Chooser/Source/ChooserEditor/Public/ContextPropertyWidget.h
+++ b/Engine/Plugins/Experimental/Chooser/Source/ChooserEditor/Public/ContextPropertyWidget.h
@@ -8,121 +8,108 @@
 
 namespace UE::ChooserEditor
 {
-	// Filter class for things which implement a particular interface
 	template <typename PropertyType>
-    TSharedRef<SWidget> CreatePropertyWidget(UObject* Object, UClass* ContextClass, const FLinearColor& BindingColor)
-    {
-    	TWeakObjectPtr<PropertyType> ContextProperty = Cast<PropertyType>(Object);
-    	
-    	FPropertyBindingWidgetArgs Args;
-    	Args.bAllowPropertyBindings = true;
-    	
-    	Args.OnCanBindProperty = FOnCanBindProperty::CreateLambda([](FProperty* Property)
-    	{
-    		return Property == nullptr || PropertyType::CanBind(*Property);
-    	});
-    	Args.OnCanBindToClass = FOnCanBindToClass::CreateLambda([ContextClass ](UClass* InClass)
-    	{
-    		return true;
-    	});
-    	
-    	Args.CurrentBindingColor = MakeAttributeLambda([ContextProperty, BindingColor]() {
-    		 if (ContextProperty.IsValid() && ContextProperty->PropertyBindingChain.Num() > 0)
-    		 {
-    		 	return BindingColor;
-    		 }
-    		return FLinearColor::Gray;
-    	});
-    
-    	Args.OnCanAcceptPropertyOrChildren = FOnCanBindProperty::CreateLambda([](FProperty* InProperty)
-    		{
-    			// Make only editor visible properties visible for binding.
-    			return InProperty->HasAnyPropertyFlags(CPF_Edit);
-    		});	
-    
-    	Args.OnAddBinding = FOnAddBinding::CreateLambda([ContextProperty](FName InPropertyName, const TArray<FBindingChainElement>& InBindingChain)
-    		{
-    			if (ContextProperty.IsValid())
-    			{
-    				const FScopedTransaction Transaction(NSLOCTEXT("ContextPropertyWidget", "Change Property Binding", "Change Property Binding"));
-    				ContextProperty->Modify(true);
-    				ContextProperty->SetBinding(InBindingChain);
-    			}
-    		});
-    
-    	Args.CurrentBindingToolTipText = MakeAttributeLambda([ContextProperty]()
-    			{
-    				const FText Bind = NSLOCTEXT("ContextPropertyWidget", "Bind", "Bind");
-    				FText CurrentValue = Bind;
-    
-    				if (ContextProperty.IsValid() && ContextProperty->PropertyBindingChain.Num()>0)
-    				{
-    					TArray<FText> BindingChainText;
-                     	BindingChainText.Reserve(ContextProperty->PropertyBindingChain.Num());
-                     
-                     	for (const FName& Name : ContextProperty->PropertyBindingChain)
-                     	{
-                     		BindingChainText.Add(FText::FromName(Name));
-                     	}
-    					
-    					CurrentValue = FText::Join(NSLOCTEXT("ContextPropertyWidget", "PropertyPathSeparator","."), BindingChainText);
-    				}
-    
-    				return CurrentValue;	
-    			});
-    	
-    	Args.CurrentBindingText = MakeAttributeLambda([ContextProperty]()
-    			{
-    				const FText Bind = NSLOCTEXT("ContextPropertyWidget", "Bind", "Bind");
-    				FText CurrentValue = Bind;
-    
-    				if (ContextProperty.IsValid())
-    				{
-    					int BindingChainLength = ContextProperty->PropertyBindingChain.Num();
-    					if (BindingChainLength > 0)
-    					{
-    						if (BindingChainLength == 1)
-    						{
-    							// single property, just use the property name
-    							CurrentValue = FText::FromName(ContextProperty->PropertyBindingChain.Last());
-    						}
-    						else
-    						{
-    							// for longer chains always show the last struct/object name, and the final property name (full path in tooltip)
-    							CurrentValue = FText::Join(NSLOCTEXT("ContextPropertyWidget", "PropertyPathSeparator","."),
-    								TArray<FText>({
-    									FText::FromName(ContextProperty->PropertyBindingChain[BindingChainLength-2]),
-    									FText::FromName(ContextProperty->PropertyBindingChain[BindingChainLength-1])
-    								}));
-    						}
-    					}
-    				}
-    
-    				return CurrentValue;
-    			});
-    
-    	Args.CurrentBindingImage = MakeAttributeLambda([]() -> const FSlateBrush*
-    		{
-    			static FName PropertyIcon(TEXT("Kismet.Tabs.Variables"));
-    			return FAppStyle::GetBrush(PropertyIcon);
-    		});
-    
-    	IPropertyAccessEditor& PropertyAccessEditor = IModularFeatures::Get().GetModularFeature<IPropertyAccessEditor>("PropertyAccessEditor");
-    	FBindingContextStruct StructInfo;
-    	StructInfo.Struct = ContextClass;
-    	return PropertyAccessEditor.MakePropertyBindingWidget({StructInfo}, Args);
-    }
-				
-	template <typename PropertyType>
-	void ConvertToText_ContextProperty(const UObject* Object, FText& OutText)
+	TSharedRef<SWidget> CreatePropertyWidget(UObject* TransactionObject, void* Value, UClass* ContextClass, const FLinearColor& BindingColor)
 	{
-		if (const PropertyType* ContextProperty = Cast<PropertyType>(Object))
+		PropertyType* ContextProperty = static_cast<PropertyType*>(Value);
+		
+		FPropertyBindingWidgetArgs Args;
+		Args.bAllowPropertyBindings = true;
+		
+		Args.OnCanBindProperty = FOnCanBindProperty::CreateLambda([](FProperty* Property)
+		{
+			return Property == nullptr || PropertyType::CanBind(*Property);
+		});
+		Args.OnCanBindToClass = FOnCanBindToClass::CreateLambda([ContextClass ](UClass* InClass)
 		{
-			if (ContextProperty->PropertyBindingChain.Num()>0)
+			return true;
+		});
+		
+		Args.CurrentBindingColor = MakeAttributeLambda([ContextProperty, BindingColor]() {
+			 if (ContextProperty!= nullptr && ContextProperty->PropertyBindingChain.Num() > 0)
+			 {
+				return BindingColor;
+			 }
+			return FLinearColor::Gray;
+		});
+	
+		Args.OnCanAcceptPropertyOrChildren = FOnCanBindProperty::CreateLambda([](FProperty* InProperty)
+			{
+				// Make only editor visible properties visible for binding.
+				return InProperty->HasAnyPropertyFlags(CPF_Edit);
+			});	
+	
+		Args.OnAddBinding = FOnAddBinding::CreateLambda([TransactionObject, ContextProperty](FName InPropertyName, const TArray<FBindingChainElement>& InBindingChain)
+			{
+				if (ContextProperty != nullptr)
+				{
+					const FScopedTransaction Transaction(NSLOCTEXT("ContextPropertyWidget", "Change Property Binding", "Change Property Binding"));
+					TransactionObject->Modify(true); // todo
+					ContextProperty->SetBinding(InBindingChain);
+				}
+			});
+	
+		Args.CurrentBindingToolTipText = MakeAttributeLambda([ContextProperty]()
+				{
+					const FText Bind = NSLOCTEXT("ContextPropertyWidget", "Bind", "Bind");
+					FText CurrentValue = Bind;
+	
+					if (ContextProperty != nullptr && ContextProperty->PropertyBindingChain.Num()>0)
+					{
+						TArray<FText> BindingChainText;
+						BindingChainText.Reserve(ContextProperty->PropertyBindingChain.Num());
+					 
+						for (const FName& Name : ContextProperty->PropertyBindingChain)
+						{
+							BindingChainText.Add(FText::FromName(Name));
+						}
+						
+						CurrentValue = FText::Join(NSLOCTEXT("ContextPropertyWidget", "PropertyPathSeparator","."), BindingChainText);
+					}
+	
+					return CurrentValue;	
+				});
+		
+		Args.CurrentBindingText = MakeAttributeLambda([ContextProperty]()
+				{
+					const FText Bind = NSLOCTEXT("ContextPropertyWidget", "Bind", "Bind");
+					FText CurrentValue = Bind;
+	
+					if (ContextProperty != nullptr)
+					{
+						int BindingChainLength = ContextProperty->PropertyBindingChain.Num();
+						if (BindingChainLength > 0)
+						{
+							if (BindingChainLength == 1)
+							{
+								// single property, just use the property name
+								CurrentValue = FText::FromName(ContextProperty->PropertyBindingChain.Last());
+							}
+							else
+							{
+								// for longer chains always show the last struct/object name, and the final property name (full path in tooltip)
+								CurrentValue = FText::Join(NSLOCTEXT("ContextPropertyWidget", "PropertyPathSeparator","."),
+									TArray<FText>({
+										FText::FromName(ContextProperty->PropertyBindingChain[BindingChainLength-2]),
+										FText::FromName(ContextProperty->PropertyBindingChain[BindingChainLength-1])
+									}));
+							}
+						}
+					}
+	
+					return CurrentValue;
+				});
+	
+		Args.CurrentBindingImage = MakeAttributeLambda([]() -> const FSlateBrush*
 			{
-				OutText = FText::FromName(ContextProperty->PropertyBindingChain.Last());
-			}
-		}
+				static FName PropertyIcon(TEXT("Kismet.Tabs.Variables"));
+				return FAppStyle::GetBrush(PropertyIcon);
+			});
+	
+		IPropertyAccessEditor& PropertyAccessEditor = IModularFeatures::Get().GetModularFeature<IPropertyAccessEditor>("PropertyAccessEditor");
+		FBindingContextStruct StructInfo;
+		StructInfo.Struct = ContextClass;
+		return PropertyAccessEditor.MakePropertyBindingWidget({StructInfo}, Args);
 	}
 
 }
\ No newline at end of file
diff --git a/Engine/Plugins/Experimental/Chooser/Source/ChooserEditor/Public/ObjectChooserClassFilter.h b/Engine/Plugins/Experimental/Chooser/Source/ChooserEditor/Public/ObjectChooserClassFilter.h
--- a/Engine/Plugins/Experimental/Chooser/Source/ChooserEditor/Public/ObjectChooserClassFilter.h
+++ b/Engine/Plugins/Experimental/Chooser/Source/ChooserEditor/Public/ObjectChooserClassFilter.h
@@ -2,7 +2,7 @@
 #pragma once
 
 #include "ClassViewerFilter.h"
-#include "IObjectChooser.h"
+#include "StructViewerFilter.h"
 
 namespace UE::ChooserEditor
 {
@@ -25,5 +25,23 @@ namespace UE::ChooserEditor
 	private:
 		UClass* InterfaceType;
 	};
+	
+
+	class FStructFilter : public IStructViewerFilter
+	{
+	public:
+		FStructFilter(const UScriptStruct* InBaseType) : BaseType(InBaseType)  { };
+		
+		virtual bool IsStructAllowed(const FStructViewerInitializationOptions& InInitOptions, const UScriptStruct* InStruct, TSharedRef<FStructViewerFilterFuncs> InFilterFuncs) override
+		{
+			return InStruct->IsChildOf(BaseType) && InStruct !=BaseType;
+		}
+		virtual bool IsUnloadedStructAllowed(const FStructViewerInitializationOptions& InInitOptions, const FSoftObjectPath& InStructPath, TSharedRef<class FStructViewerFilterFuncs> InFilterFuncs)
+		{
+			return false;
+		};
+	private:
+		const UScriptStruct* BaseType;
+	};
 		
 }
\ No newline at end of file
diff --git a/Engine/Plugins/Experimental/Chooser/Source/ChooserEditor/Public/ObjectChooserWidgetFactories.h b/Engine/Plugins/Experimental/Chooser/Source/ChooserEditor/Public/ObjectChooserWidgetFactories.h
--- a/Engine/Plugins/Experimental/Chooser/Source/ChooserEditor/Public/ObjectChooserWidgetFactories.h
+++ b/Engine/Plugins/Experimental/Chooser/Source/ChooserEditor/Public/ObjectChooserWidgetFactories.h
@@ -5,22 +5,19 @@
 #include "CoreMinimal.h"
 #include "Widgets/SWidget.h"
 #include "Widgets/Layout/SBorder.h"
-#include "SClassViewer.h"
+#include "StructViewerModule.h"
 
 namespace UE::ChooserEditor
 {
-	typedef TFunction<void(const UObject* Object, FText& OutText)> FChooserTextConverter;
-	typedef TFunction<TSharedRef<SWidget>(UObject* Object, UClass* ContextClass)> FChooserWidgetCreator;
+	typedef TFunction<TSharedRef<SWidget>(UObject* TransactionObject, void* Value, UClass* ContextClass)> FChooserWidgetCreator;
 
 	class CHOOSEREDITOR_API FObjectChooserWidgetFactories
 	{
 	public:
-		static TMap<const UClass*, FChooserTextConverter> ChooserTextConverter;
-		static TMap<const UClass*, FChooserWidgetCreator> ChooserWidgetCreators;
+		static TMap<const UStruct*, FChooserWidgetCreator> ChooserWidgetCreators;
 
-		static TSharedPtr<SWidget> CreateWidget(UObject* Value, UClass* ContextClass);
-		static TSharedPtr<SWidget> CreateWidget(UClass* InterfaceType, UObject* Value, UClass* ContextClass, const FOnClassPicked& CreateClassCallback, TSharedPtr<SBorder>* InnerWidget = nullptr);
-		static void ConvertToText(UObject* Object, FText& OutText);
+		static TSharedPtr<SWidget> CreateWidget(UObject* TransactionObject, void* Value, const UStruct* ValueType, UClass* ContextClass);
+		static TSharedPtr<SWidget> CreateWidget(UObject* TransactionObject, const UScriptStruct* BaseType, void* Value, const UStruct* ValueType, UClass* ContextClass, const FOnStructPicked& CreateClassCallback, TSharedPtr<SBorder>* InnerWidget = nullptr);
 		
 		static void RegisterWidgets();
 	};
diff --git a/Engine/Plugins/Experimental/Chooser/Source/ProxyTable/Internal/ProxyTable.h b/Engine/Plugins/Experimental/Chooser/Source/ProxyTable/Internal/ProxyTable.h
--- a/Engine/Plugins/Experimental/Chooser/Source/ProxyTable/Internal/ProxyTable.h
+++ b/Engine/Plugins/Experimental/Chooser/Source/ProxyTable/Internal/ProxyTable.h
@@ -6,7 +6,7 @@
 #include "IObjectChooser.h"
 #include "ChooserPropertyAccess.h"
 #include "IChooserParameterProxyTable.h"
-
+#include "InstancedStruct.h"
 #include "ProxyTable.generated.h"
 
 
@@ -16,8 +16,13 @@ struct FProxyEntry
 	GENERATED_BODY()
 	UPROPERTY(EditAnywhere, Category = "Data")
 	FName Key;
-	UPROPERTY(EditAnywhere, Meta = (EditInlineInterface = "true"), Category = "Data")
+	
+	UPROPERTY()
 	TScriptInterface<IObjectChooser> Value;
+	
+	UPROPERTY(DisplayName="Value", EditAnywhere, Meta = (ExcludeBaseStruct, BaseStruct = "/Script/Chooser.ObjectChooserBase"), Category = "Data")
+	FInstancedStruct ValueStruct;
+	
 };
 
 UCLASS(MinimalAPI,BlueprintType)
@@ -37,8 +42,9 @@ public:
 
 struct FBindingChainElement;
 
-UCLASS()
-class PROXYTABLE_API UChooserParameterProxyTable_ContextProperty :  public UObject, public IChooserParameterProxyTable
+
+USTRUCT()
+struct PROXYTABLE_API FProxyTableContextProperty :  public FChooserParameterProxyTableBase
 {
 	GENERATED_BODY()
 public:
@@ -61,21 +67,61 @@ public:
 #endif
 };
 
-
-
-UCLASS()
-class PROXYTABLE_API UObjectChooser_LookupProxy : public UObject, public IObjectChooser
+USTRUCT()
+struct PROXYTABLE_API FLookupProxy : public FObjectChooserBase
 {
 	GENERATED_BODY()
 	virtual UObject* ChooseObject(const UObject* ContextObject) const final override;
 	
 	public:
-	UObjectChooser_LookupProxy() {};
-	UObjectChooser_LookupProxy(const FObjectInitializer& ObjectInitializer);
+	FLookupProxy();
 	
-	UPROPERTY(EditAnywhere, Meta = (EditInlineInterface = "true"), Category = "Input")
-	TScriptInterface<IChooserParameterProxyTable> ProxyTable;
+	UPROPERTY(EditAnywhere, Meta = (ExcludeBaseStruct, BaseStruct ="/Script/ProxyTable.ChooserParameterProxyTableBase"), Category = "Input")
+	FInstancedStruct ProxyTable;
 	
 	UPROPERTY(EditAnywhere, Category="Parameters")
 	FName Key;
 };
+
+// deprecated classes for upgrading old data
+
+UCLASS(ClassGroup = "LiveLink", deprecated)
+class PROXYTABLE_API UDEPRECATED_ChooserParameterProxyTable_ContextProperty :  public UObject, public IChooserParameterProxyTable
+{
+	GENERATED_BODY()
+public:
+
+	UPROPERTY()
+	TArray<FName> PropertyBindingChain;
+
+	virtual void ConvertToInstancedStruct(FInstancedStruct& OutInstancedStruct) const override
+	{
+		OutInstancedStruct.InitializeAs(FProxyTableContextProperty::StaticStruct());
+		FProxyTableContextProperty& Property = OutInstancedStruct.GetMutable<FProxyTableContextProperty>();
+		Property.PropertyBindingChain = PropertyBindingChain;
+	}
+};
+
+UCLASS(ClassGroup = "LiveLink", deprecated)
+class PROXYTABLE_API UDEPRECATED_ObjectChooser_LookupProxy : public UObject, public IObjectChooser
+{
+	GENERATED_BODY()
+	
+	UPROPERTY(EditAnywhere, Category = "Input")
+	TScriptInterface<IChooserParameterProxyTable> ProxyTable;
+	
+	UPROPERTY(EditAnywhere, Category="Parameters")
+	FName Key;
+	
+	virtual void ConvertToInstancedStruct(FInstancedStruct& OutInstancedStruct) const override
+   	{
+   		OutInstancedStruct.InitializeAs(FLookupProxy::StaticStruct());
+   		FLookupProxy& AssetChooser = OutInstancedStruct.GetMutable<FLookupProxy>();
+		
+		if (IChooserParameterProxyTable* ProxyTableInterface = ProxyTable.GetInterface())
+		{
+			ProxyTableInterface->ConvertToInstancedStruct(AssetChooser.ProxyTable);
+		}
+		AssetChooser.Key = Key;
+   	}
+};
\ No newline at end of file
diff --git a/Engine/Plugins/Experimental/Chooser/Source/ProxyTable/Private/ProxyTable.cpp b/Engine/Plugins/Experimental/Chooser/Source/ProxyTable/Private/ProxyTable.cpp
--- a/Engine/Plugins/Experimental/Chooser/Source/ProxyTable/Private/ProxyTable.cpp
+++ b/Engine/Plugins/Experimental/Chooser/Source/ProxyTable/Private/ProxyTable.cpp
@@ -1,22 +1,19 @@
 // Copyright Epic Games, Inc. All Rights Reserved.
 #include "ProxyTable.h"
 
-UObjectChooser_LookupProxy::UObjectChooser_LookupProxy(const FObjectInitializer& ObjectInitializer)
+FLookupProxy::FLookupProxy()
 {
-	ProxyTable = ObjectInitializer.CreateDefaultSubobject<UChooserParameterProxyTable_ContextProperty>(this, "InputValue");
-	ProxyTable.GetObject()->SetFlags(RF_Transactional);
+	ProxyTable.InitializeAs(FProxyTableContextProperty::StaticStruct());
 }
 
 static UObject* FindProxyObject(const UProxyTable* Table, FName Key, const UObject* ContextObject)
 {
 	for (const FProxyEntry& Entry : Table->Entries)
 	{
-		if (Entry.Key == Key)
+		if (Entry.Key == Key && Entry.ValueStruct.IsValid())
 		{
-			if (const IObjectChooser* EntryValue = Entry.Value.GetInterface())
-			{
-				return EntryValue->ChooseObject(ContextObject);
-			}
+			const FObjectChooserBase& EntryValue = Entry.ValueStruct.Get<FObjectChooserBase>();
+			return EntryValue.ChooseObject(ContextObject);
 		}
 	}
 
@@ -32,12 +29,12 @@ static UObject* FindProxyObject(const UProxyTable* Table, FName Key, const UObje
 	return nullptr;
 }
 
-UObject* UObjectChooser_LookupProxy::ChooseObject(const UObject* ContextObject) const
+UObject* FLookupProxy::ChooseObject(const UObject* ContextObject) const
 {
-	if (ProxyTable)
+	if (ProxyTable.IsValid())
 	{
 		const UProxyTable* Table;
-		if (ProxyTable->GetValue(ContextObject, Table))
+		if (ProxyTable.Get<FChooserParameterProxyTableBase>().GetValue(ContextObject, Table))
 		{
 			if (Table)
 			{
@@ -47,7 +44,6 @@ UObject* UObjectChooser_LookupProxy::ChooseObject(const UObject* ContextObject)
 				}
 			}
 		}
-				
 	}
 	
 	return nullptr;
@@ -59,7 +55,7 @@ UProxyTable::UProxyTable(const FObjectInitializer& Initializer)
 
 }
 
-bool UChooserParameterProxyTable_ContextProperty::GetValue(const UObject* ContextObject, const UProxyTable*& OutResult) const
+bool FProxyTableContextProperty::GetValue(const UObject* ContextObject, const UProxyTable*& OutResult) const
 {
 	UStruct* StructType = ContextObject->GetClass();
 	const void* Container = ContextObject;
diff --git a/Engine/Plugins/Experimental/Chooser/Source/ProxyTable/ProxyTable.Build.cs b/Engine/Plugins/Experimental/Chooser/Source/ProxyTable/ProxyTable.Build.cs
--- a/Engine/Plugins/Experimental/Chooser/Source/ProxyTable/ProxyTable.Build.cs
+++ b/Engine/Plugins/Experimental/Chooser/Source/ProxyTable/ProxyTable.Build.cs
@@ -19,7 +19,8 @@ namespace UnrealBuildTool.Rules
 					"Core",
 					"CoreUObject",
 					"Engine",
-					"Chooser"
+					"Chooser",
+					"StructUtils"
 					// ... add private dependencies that you statically link with here ...
 				}
 			);
diff --git a/Engine/Plugins/Experimental/Chooser/Source/ProxyTable/Public/IChooserParameterProxyTable.h b/Engine/Plugins/Experimental/Chooser/Source/ProxyTable/Public/IChooserParameterProxyTable.h
--- a/Engine/Plugins/Experimental/Chooser/Source/ProxyTable/Public/IChooserParameterProxyTable.h
+++ b/Engine/Plugins/Experimental/Chooser/Source/ProxyTable/Public/IChooserParameterProxyTable.h
@@ -5,6 +5,8 @@
 #include "CoreMinimal.h"
 #include "UObject/Object.h"
 #include "UObject/Interface.h"
+#include "IChooserParameterBase.h"
+#include "InstancedStruct.h"
 #include "IChooserParameterProxyTable.generated.h"
 
 class UProxyTable;
@@ -18,7 +20,16 @@ class PROXYTABLE_API UChooserParameterProxyTable : public UInterface
 class PROXYTABLE_API IChooserParameterProxyTable
 {
 	GENERATED_BODY()
+public:
+	virtual void ConvertToInstancedStruct(FInstancedStruct& OutInstancedStruct) const {}
+};
+
 
+USTRUCT()
+struct FChooserParameterProxyTableBase : public FChooserParameterBase
+{
+	GENERATED_BODY()
+    
 public:
 	virtual bool GetValue(const UObject* ContextObject, const UProxyTable*& OutResult) const { return false; }
 };
\ No newline at end of file
diff --git a/Engine/Plugins/Experimental/Chooser/Source/ProxyTableEditor/Private/ProxyTableEditor.cpp b/Engine/Plugins/Experimental/Chooser/Source/ProxyTableEditor/Private/ProxyTableEditor.cpp
--- a/Engine/Plugins/Experimental/Chooser/Source/ProxyTableEditor/Private/ProxyTableEditor.cpp
+++ b/Engine/Plugins/Experimental/Chooser/Source/ProxyTableEditor/Private/ProxyTableEditor.cpp
@@ -29,6 +29,8 @@
 #include "ScopedTransaction.h"
 #include "ObjectChooserWidgetFactories.h"
 #include "ContextPropertyWidget.h"
+#include "IObjectChooser.h"
+#include "Misc/TransactionObjectEvent.h"
 
 #define LOCTEXT_NAMESPACE "ProxyTableEditor"
 
@@ -294,7 +296,7 @@ public:
 		Operation->ChooserEditor = InEditor;
 		Operation->RowIndex = InRowIndex;
 		Operation->DefaultHoverText = LOCTEXT("Proxy Row", "Proxy Row");
-		UE::ChooserEditor::FObjectChooserWidgetFactories::ConvertToText(InEditor->GetProxyTable()->Entries[InRowIndex].Value.GetObject(), Operation->DefaultHoverText);
+		// UE::ChooserEditor::FObjectChooserWidgetFactories::ConvertToText(InEditor->GetProxyTable()->Entries[InRowIndex].Value.GetObject(), Operation->DefaultHoverText);
 		Operation->CurrentHoverText = Operation->DefaultHoverText;
 			
 		Operation->Construct();
@@ -380,23 +382,26 @@ public:
 			if (ColumnName == Handles)
 			{
 				// row drag handle
-				
 				return SNew(SProxyRowHandle).ProxyEditor(Editor).RowIndex(RowIndex->RowIndex);
 			}
 			else if (ColumnName == Value) 
 			{
-				UObject* RowValue = ProxyTable->Entries[RowIndex->RowIndex].Value.GetObject();
-				TSharedPtr<SWidget> ResultWidget = UE::ChooserEditor::FObjectChooserWidgetFactories::CreateWidget(UObjectChooser::StaticClass(), RowValue, nullptr/*ProxyTable->ContextObjectType*/,
-				FOnClassPicked::CreateLambda([this, RowIndex=RowIndex->RowIndex](UClass* ChosenClass)
-				{
-					const FScopedTransaction Transaction(LOCTEXT("Change Row Result Type", "Change Row Result Type"));
-					UObject* RowValue = NewObject<UObject>(ProxyTable, ChosenClass, NAME_None, RF_Transactional);
-					ProxyTable->Modify(true);
-					ProxyTable->Entries[RowIndex].Value = RowValue;
-					UE::ChooserEditor::FObjectChooserWidgetFactories::CreateWidget(UObjectChooser::StaticClass(), RowValue, nullptr/*Chooser->ContextObjectType*/, FOnClassPicked(), &CacheBorder);
-				}),
-				&CacheBorder
-				);
+				TSharedPtr<SWidget> ResultWidget = ChooserEditor::FObjectChooserWidgetFactories::CreateWidget(ProxyTable, FObjectChooserBase::StaticStruct(),
+					ProxyTable->Entries[RowIndex->RowIndex].ValueStruct.GetMutableMemory(),
+					ProxyTable->Entries[RowIndex->RowIndex].ValueStruct.GetScriptStruct(),
+					nullptr/*ProxyTable->ContextObjectType*/,
+					FOnStructPicked::CreateLambda([this, RowIndex=RowIndex->RowIndex](const UScriptStruct* ChosenStruct)
+					{
+						const FScopedTransaction Transaction(LOCTEXT("Change Value Type", "Change Value Type"));
+						ProxyTable->Entries[RowIndex].ValueStruct.InitializeAs(ChosenStruct);
+						ProxyTable->Modify(true);
+						ChooserEditor::FObjectChooserWidgetFactories::CreateWidget(ProxyTable, FObjectChooserBase::StaticStruct(),
+								ProxyTable->Entries[RowIndex].ValueStruct.GetMutableMemory(),
+								ProxyTable->Entries[RowIndex].ValueStruct.GetScriptStruct(),
+								nullptr/*ProxyTable->ContextObjectType*/, FOnStructPicked(), &CacheBorder);
+					}),
+					&CacheBorder
+					);
 				return ResultWidget.ToSharedRef();
 			}
 			else if (ColumnName == Key)
@@ -485,24 +490,22 @@ TSharedRef<SDockTab> FProxyTableEditor::SpawnTableTab( const FSpawnTabArgs& Args
 		]
 		.OnGetMenuContent_Lambda([this]()
 		{
-			FClassViewerInitializationOptions Options;
-			Options.ClassFilters.Add(MakeShared<UE::ChooserEditor::FInterfaceClassFilter>(UObjectChooser::StaticClass()) );
-			Options.NameTypeToDisplay = EClassViewerNameTypeToDisplay::DisplayName;
+			FStructViewerInitializationOptions Options;
+			Options.StructFilter = MakeShared<ChooserEditor::FStructFilter>(FObjectChooserBase::StaticStruct());
+			Options.NameTypeToDisplay = EStructViewerNameTypeToDisplay::DisplayName;
+			
+			TSharedRef<SWidget> Widget = FModuleManager::LoadModuleChecked<FStructViewerModule>("StructViewer").CreateStructViewer(Options, FOnStructPicked::CreateLambda([this](const UScriptStruct* ChosenStruct)
+			{
+				CreateRowComboButton->SetIsOpen(false);
+				UProxyTable* ProxyTable = Cast<UProxyTable>(EditingObjects[0]);
+				const FScopedTransaction Transaction(LOCTEXT("Add Row Transaction", "Add Row"));
+				ProxyTable->Modify(true);
+
+				ProxyTable->Entries.SetNum(ProxyTable->Entries.Num()+1);
+				ProxyTable->Entries.Last().ValueStruct.InitializeAs(ChosenStruct);
+				UpdateTableRows();
+			}));
 			
-			// Add class filter for columns here
-			TSharedRef<SWidget> Widget = FModuleManager::LoadModuleChecked<FClassViewerModule>("ClassViewer").CreateClassViewer(Options,
-				FOnClassPicked::CreateLambda([this](UClass* ChosenClass)
-				{
-					CreateRowComboButton->SetIsOpen(false);
-					UProxyTable* ProxyTable = Cast<UProxyTable>(EditingObjects[0]);
-					const FScopedTransaction Transaction(LOCTEXT("Add Row Transaction", "Add Row"));
-					ProxyTable->Modify(true);
-					UObject* ResultObject = NewObject<UObject>(ProxyTable, ChosenClass, NAME_None, RF_Transactional);
-					FProxyEntry Entry;
-					Entry.Value = TScriptInterface<IObjectChooser>(ResultObject);
-					ProxyTable->Entries.Add(Entry);
-					UpdateTableRows();
-				}));
 			return Widget;
 		});
 
@@ -600,14 +603,14 @@ void FProxyTableEditor::OnObjectsReplaced(const TMap<UObject*, UObject*>& Replac
 	// Refresh our details view if one of the objects replaced was in the map. This gets called before the reinstance GC fixup, so we might as well fixup EditingObjects now too
 	for (int32 i = 0; i < EditingObjects.Num(); i++)
 	{
-		// UObject* SourceObject = EditingObjects[i];
-		// UObject* ReplacedObject = ReplacementMap.FindRef(SourceObject);
-		//
-		// if (ReplacedObject && ReplacedObject != SourceObject)
-		// {
-		// 	EditingObjects[i] = ReplacedObject;
-		// 	bChangedAny = true;
-		// }
+		UObject* SourceObject = EditingObjects[i];
+		UObject* ReplacedObject = ReplacementMap.FindRef(SourceObject);
+
+		if (ReplacedObject && ReplacedObject != SourceObject)
+		{
+			EditingObjects[i] = ReplacedObject;
+			bChangedAny = true;
+		}
 	}
 
 	if (bChangedAny)
@@ -641,38 +644,31 @@ TSharedRef<FProxyTableEditor> FProxyTableEditor::CreateEditor( const EToolkitMod
 
 /// Result widgets
 
-TSharedRef<SWidget> CreateLookupProxyWidget(UObject* Object, UClass* ContextObject)
+TSharedRef<SWidget> CreateLookupProxyWidget(UObject* TransactionObject, void* Value, UClass* ContextObject)
 {
-	UObjectChooser_LookupProxy* LookupProxy = Cast<UObjectChooser_LookupProxy>(Object);
+	FLookupProxy* LookupProxy = static_cast<FLookupProxy*>(Value);
 	
 	return SNew(SHorizontalBox)
 		+ SHorizontalBox::Slot()
 		[
-			UE::ChooserEditor::CreatePropertyWidget<UChooserParameterProxyTable_ContextProperty>(LookupProxy->ProxyTable.GetObject(), ContextObject, GetDefault<UGraphEditorSettings>()->ObjectPinTypeColor)
+			UE::ChooserEditor::CreatePropertyWidget<FProxyTableContextProperty>(TransactionObject, LookupProxy->ProxyTable.GetMutablePtr<FProxyTableContextProperty>(), ContextObject, GetDefault<UGraphEditorSettings>()->ObjectPinTypeColor)
 		]
 		+ SHorizontalBox::Slot()
 		[
 			SNew(SEditableTextBox)
 			.Text_Lambda([LookupProxy]() { return FText::FromName(LookupProxy->Key);})
-			.OnTextChanged_Lambda([LookupProxy](const FText& NewText)
+			.OnTextChanged_Lambda([TransactionObject,LookupProxy](const FText& NewText)
 			{
 				FScopedTransaction ScopedTransaction(LOCTEXT("Change LookupProxy Key Name", "Change LookupProxy Key Name"));
-				LookupProxy->Modify(true);
+				TransactionObject->Modify(true);
 				LookupProxy->Key = FName(NewText.ToString());
 			})
 		];
 }
 
-void ConvertToText_LookupProxy(const UObject* Object, FText& OutText)
-{
-	const UObjectChooser_LookupProxy* LookupProxy = Cast<UObjectChooser_LookupProxy>(Object);
-	OutText = FText::FromString(LookupProxy->Key.ToString());
-}
-	
 void FProxyTableEditor::RegisterWidgets()
 {
-	UE::ChooserEditor::FObjectChooserWidgetFactories::ChooserTextConverter.Add(UObjectChooser_LookupProxy::StaticClass(), ConvertToText_LookupProxy);
-	UE::ChooserEditor::FObjectChooserWidgetFactories::ChooserWidgetCreators.Add(UObjectChooser_LookupProxy::StaticClass(), CreateLookupProxyWidget);
+	UE::ChooserEditor::FObjectChooserWidgetFactories::ChooserWidgetCreators.Add(FLookupProxy::StaticStruct(), CreateLookupProxyWidget);
 }
 	
 }
diff --git a/Engine/Plugins/Experimental/Chooser/Source/ProxyTableEditor/ProxyTableEditor.Build.cs b/Engine/Plugins/Experimental/Chooser/Source/ProxyTableEditor/ProxyTableEditor.Build.cs
--- a/Engine/Plugins/Experimental/Chooser/Source/ProxyTableEditor/ProxyTableEditor.Build.cs
+++ b/Engine/Plugins/Experimental/Chooser/Source/ProxyTableEditor/ProxyTableEditor.Build.cs
@@ -29,6 +29,7 @@ namespace UnrealBuildTool.Rules
 					"GraphEditor",
 					"GameplayTags",
 					"GameplayTagsEditor",
+					"StructUtils"
 					// ... add private dependencies that you statically link with here ...
 				}
 			);
