fix nomaster compile issues

diff --git a/Engine/Plugins/Experimental/Chooser/Source/Chooser/Public/IChooserParameterEnum.h b/Engine/Plugins/Experimental/Chooser/Source/Chooser/Public/IChooserParameterEnum.h
--- a/Engine/Plugins/Experimental/Chooser/Source/Chooser/Public/IChooserParameterEnum.h
+++ b/Engine/Plugins/Experimental/Chooser/Source/Chooser/Public/IChooserParameterEnum.h
@@ -6,6 +6,7 @@
 #include "UObject/Object.h"
 #include "UObject/Interface.h"
 #include "InstancedStruct.h"
+#include "IChooserParameterBase.h"
 #include "IChooserParameterEnum.generated.h"
 
 UINTERFACE(NotBlueprintType, meta = (CannotImplementInterfaceInBlueprint))
diff --git a/Engine/Plugins/Experimental/Chooser/Source/Chooser/Public/IChooserParameterFloat.h b/Engine/Plugins/Experimental/Chooser/Source/Chooser/Public/IChooserParameterFloat.h
--- a/Engine/Plugins/Experimental/Chooser/Source/Chooser/Public/IChooserParameterFloat.h
+++ b/Engine/Plugins/Experimental/Chooser/Source/Chooser/Public/IChooserParameterFloat.h
@@ -5,6 +5,8 @@
 #include "CoreMinimal.h"
 #include "UObject/Object.h"
 #include "UObject/Interface.h"
+#include "InstancedStruct.h"
+#include "IChooserParameterBase.h"
 #include "IChooserParameterFloat.generated.h"
 
 UINTERFACE(NotBlueprintType, meta = (CannotImplementInterfaceInBlueprint))
diff --git a/Engine/Plugins/Experimental/Chooser/Source/Chooser/Public/IChooserParameterGameplayTag.h b/Engine/Plugins/Experimental/Chooser/Source/Chooser/Public/IChooserParameterGameplayTag.h
--- a/Engine/Plugins/Experimental/Chooser/Source/Chooser/Public/IChooserParameterGameplayTag.h
+++ b/Engine/Plugins/Experimental/Chooser/Source/Chooser/Public/IChooserParameterGameplayTag.h
@@ -5,6 +5,7 @@
 #include "CoreMinimal.h"
 #include "UObject/Object.h"
 #include "UObject/Interface.h"
+#include "InstancedStruct.h"
 #include "IChooserParameterBase.h"
 #include "IChooserParameterGameplayTag.generated.h"
 
