Fix build error

diff --git a/Engine/Plugins/Experimental/Chooser/Source/ChooserEditor/Private/ChooserTableEditor.cpp b/Engine/Plugins/Experimental/Chooser/Source/ChooserEditor/Private/ChooserTableEditor.cpp
--- a/Engine/Plugins/Experimental/Chooser/Source/ChooserEditor/Private/ChooserTableEditor.cpp
+++ b/Engine/Plugins/Experimental/Chooser/Source/ChooserEditor/Private/ChooserTableEditor.cpp
@@ -545,9 +545,9 @@ void FChooserTableEditor::UpdateTableColumns()
 						TSharedRef<SWidget> Widget = FModuleManager::LoadModuleChecked<FStructViewerModule>("StructViewer").CreateStructViewer(Options, FOnStructPicked::CreateLambda([this, ColumnIndex](const UScriptStruct* ChosenStruct)
 						{
 							const FScopedTransaction Transaction(LOCTEXT("SetColumnInputType", "Set Column Input Type"));
-							UChooserTable* Chooser = Cast<UChooserTable>(EditingObjects[0]);
-							Chooser->ColumnsStructs[ColumnIndex].GetMutable<FChooserColumnBase>().SetInputType(ChosenStruct);
-							Chooser->Modify(true);
+							UChooserTable* ChooserTable = Cast<UChooserTable>(EditingObjects[0]);
+							ChooserTable->ColumnsStructs[ColumnIndex].GetMutable<FChooserColumnBase>().SetInputType(ChosenStruct);
+							ChooserTable->Modify(true);
 							UpdateTableColumns();
 							UpdateTableRows();
 						}));
