Follow up editor fixes for Chooser UObject to UStruct changes

diff --git a/Engine/Plugins/Experimental/Chooser/Source/ChooserEditor/Private/ChooserTableEditor.cpp b/Engine/Plugins/Experimental/Chooser/Source/ChooserEditor/Private/ChooserTableEditor.cpp
--- a/Engine/Plugins/Experimental/Chooser/Source/ChooserEditor/Private/ChooserTableEditor.cpp
+++ b/Engine/Plugins/Experimental/Chooser/Source/ChooserEditor/Private/ChooserTableEditor.cpp
@@ -147,7 +147,8 @@ void FChooserTableEditor::NotifyPreChange(FProperty* PropertyAboutToChange)
 
 void FChooserTableEditor::NotifyPostChange(const FPropertyChangedEvent& PropertyChangedEvent, FProperty* PropertyThatChanged)
 {
-	if (PropertyThatChanged->GetNameCPP() == GET_MEMBER_NAME_STRING_CHECKED(UChooserTable, OutputObjectType))
+	if (PropertyThatChanged->GetNameCPP() == GET_MEMBER_NAME_STRING_CHECKED(UChooserTable, OutputObjectType)              // if you change the Output type, we need to update all the Results asset pickers
+	   || PropertyChangedEvent.Property->GetNameCPP() == GET_MEMBER_NAME_STRING_CHECKED(UChooserTable, ResultsStructs) )   // if you change the type of a Result, we need to update that widget
 	{
 		// rebuild all result widgets
 		UpdateTableRows();
@@ -881,6 +882,40 @@ TSharedRef<SWidget> CreateEvaluateChooserWidget(UObject* TransactionObject, void
 			EvaluateChooser->Chooser = Cast<UChooserTable>(AssetData.GetAsset());
 		});
 }
+	
+class FChooserDetails : public IDetailCustomization
+{
+public:
+	FChooserDetails() {};
+	virtual ~FChooserDetails() override {};
+
+	static TSharedRef<IDetailCustomization> MakeInstance()
+	{
+		return MakeShareable( new FChooserDetails() );
+	}
+
+	// IDetailCustomization interface
+	virtual void CustomizeDetails(class IDetailLayoutBuilder& DetailBuilder) override;
+};
+		
+void FChooserDetails::CustomizeDetails(IDetailLayoutBuilder& DetailBuilder)
+{
+	TArray<TWeakObjectPtr<UObject>> Objects;
+	DetailBuilder.GetObjectsBeingCustomized(Objects);
+
+	UChooserTable* Chooser = Cast<UChooserTable>(Objects[0]);
+	
+	IDetailCategoryBuilder& HiddenCategory = DetailBuilder.EditCategory(TEXT("Hidden"));
+
+	TArray<TSharedRef<IPropertyHandle>> HiddenProperties;
+	HiddenCategory.GetDefaultProperties(HiddenProperties);
+	for(TSharedRef<IPropertyHandle>& PropertyHandle :  HiddenProperties)
+	{
+		// these (Results and Columns arrays) need to be hidden when showing the root ChooserTable properties
+		// but still need to be EditAnywhere so that the Properties exist for displaying when you select a row or column (eg by FChooserRowDetails below)
+		PropertyHandle->MarkHiddenByCustomization();
+	}
+}
 
 class FChooserRowDetails : public IDetailCustomization
 {
@@ -935,7 +970,7 @@ void FChooserRowDetails::CustomizeDetails(IDetailLayoutBuilder& DetailBuilder)
 				FText DisplayName = LOCTEXT("No Input Value", "No Input Value");
 				if (FChooserParameterBase* InputValue = Column.GetInputValue())
 				{
-					InputValue->GetDisplayName(DisplayName); // todo
+					InputValue->GetDisplayName(DisplayName); 
 				}
 				NewColumnProperty.DisplayName(DisplayName);
 				NewColumnProperty.ShowPropertyButtons(false); // hide array add button
@@ -993,6 +1028,8 @@ void FChooserTableEditor::RegisterWidgets()
 	
 
 	FPropertyEditorModule& PropertyModule = FModuleManager::LoadModuleChecked<FPropertyEditorModule>("PropertyEditor");
+	
+	PropertyModule.RegisterCustomClassLayout("ChooserTable", FOnGetDetailCustomizationInstance::CreateStatic(&FChooserDetails::MakeInstance));	
 	PropertyModule.RegisterCustomClassLayout("ChooserRowDetails", FOnGetDetailCustomizationInstance::CreateStatic(&FChooserRowDetails::MakeInstance));	
 	PropertyModule.RegisterCustomClassLayout("ChooserColumnDetails", FOnGetDetailCustomizationInstance::CreateStatic(&FChooserColumnDetails::MakeInstance));	
 }
