Fix for chooser class to script upgrade code not correctly upgrading column input parameters

diff --git a/Engine/Plugins/Experimental/Chooser/Source/Chooser/Internal/BoolColumn.h b/Engine/Plugins/Experimental/Chooser/Source/Chooser/Internal/BoolColumn.h
--- a/Engine/Plugins/Experimental/Chooser/Source/Chooser/Internal/BoolColumn.h
+++ b/Engine/Plugins/Experimental/Chooser/Source/Chooser/Internal/BoolColumn.h
@@ -83,7 +83,10 @@ class CHOOSER_API UDEPRECATED_ChooserColumnBool : public UObject, public IChoose
 	GENERATED_BODY()
 public:
 	UDEPRECATED_ChooserColumnBool() {};
-	UDEPRECATED_ChooserColumnBool(const FObjectInitializer& ObjectInitializer){};
+	UDEPRECATED_ChooserColumnBool(const FObjectInitializer& ObjectInitializer)
+	{
+		InputValue = ObjectInitializer.CreateDefaultSubobject<UDEPRECATED_ChooserParameterBool_ContextProperty>(this, "InputValue");
+	}	
 	
 	UPROPERTY(EditAnywhere, Category = "Input")
 	TScriptInterface<IChooserParameterBool> InputValue;
diff --git a/Engine/Plugins/Experimental/Chooser/Source/Chooser/Internal/EnumColumn.h b/Engine/Plugins/Experimental/Chooser/Source/Chooser/Internal/EnumColumn.h
--- a/Engine/Plugins/Experimental/Chooser/Source/Chooser/Internal/EnumColumn.h
+++ b/Engine/Plugins/Experimental/Chooser/Source/Chooser/Internal/EnumColumn.h
@@ -151,7 +151,10 @@ class CHOOSER_API UDEPRECATED_ChooserColumnEnum : public UObject, public IChoose
 	GENERATED_BODY()
 public:
 	UDEPRECATED_ChooserColumnEnum() = default;
-	UDEPRECATED_ChooserColumnEnum(const FObjectInitializer& ObjectInitializer){};
+	UDEPRECATED_ChooserColumnEnum(const FObjectInitializer& ObjectInitializer)
+	{
+		InputValue = ObjectInitializer.CreateDefaultSubobject<UDEPRECATED_ChooserParameterEnum_ContextProperty>(this, "InputValue");
+	}	
 
 	UPROPERTY(EditAnywhere, Category = "Input")
 	TScriptInterface<IChooserParameterEnum> InputValue;
diff --git a/Engine/Plugins/Experimental/Chooser/Source/Chooser/Internal/FloatRangeColumn.h b/Engine/Plugins/Experimental/Chooser/Source/Chooser/Internal/FloatRangeColumn.h
--- a/Engine/Plugins/Experimental/Chooser/Source/Chooser/Internal/FloatRangeColumn.h
+++ b/Engine/Plugins/Experimental/Chooser/Source/Chooser/Internal/FloatRangeColumn.h
@@ -100,7 +100,10 @@ class CHOOSER_API UDEPRECATED_ChooserColumnFloatRange : public UObject, public I
 	GENERATED_BODY()
 	public:
 	UDEPRECATED_ChooserColumnFloatRange() {}
-	UDEPRECATED_ChooserColumnFloatRange(const FObjectInitializer& ObjectInitializer) {};
+	UDEPRECATED_ChooserColumnFloatRange(const FObjectInitializer& ObjectInitializer)
+	{
+		InputValue = ObjectInitializer.CreateDefaultSubobject<UDEPRECATED_ChooserParameterFloat_ContextProperty>(this, "InputValue");
+	}	
 		
 	UPROPERTY(EditAnywhere, Category = "Input")
 	TScriptInterface<IChooserParameterFloat> InputValue;
diff --git a/Engine/Plugins/Experimental/Chooser/Source/Chooser/Internal/GameplayTagColumn.h b/Engine/Plugins/Experimental/Chooser/Source/Chooser/Internal/GameplayTagColumn.h
--- a/Engine/Plugins/Experimental/Chooser/Source/Chooser/Internal/GameplayTagColumn.h
+++ b/Engine/Plugins/Experimental/Chooser/Source/Chooser/Internal/GameplayTagColumn.h
@@ -91,7 +91,10 @@ class CHOOSER_API UDEPRECATED_ChooserColumnGameplayTag : public UObject, public
 	GENERATED_BODY()
 	public:
 	UDEPRECATED_ChooserColumnGameplayTag() {};
-	UDEPRECATED_ChooserColumnGameplayTag(const FObjectInitializer& ObjectInitializer){};
+	UDEPRECATED_ChooserColumnGameplayTag(const FObjectInitializer& ObjectInitializer)
+	{
+		InputValue = ObjectInitializer.CreateDefaultSubobject<UDEPRECATED_ChooserParameterGameplayTag_ContextProperty>(this, "InputValue");
+	}	
 	
 	UPROPERTY(EditAnywhere, Category = "Input")
 	TScriptInterface<IChooserParameterGameplayTag> InputValue;
