Add EvaluateProxyTable function to allow anim node functions to evaluate a proxy table directly.

diff --git a/Engine/Plugins/Experimental/Chooser/Source/ProxyTable/Private/ProxyTable.cpp b/Engine/Plugins/Experimental/Chooser/Source/ProxyTable/Private/ProxyTable.cpp
--- a/Engine/Plugins/Experimental/Chooser/Source/ProxyTable/Private/ProxyTable.cpp
+++ b/Engine/Plugins/Experimental/Chooser/Source/ProxyTable/Private/ProxyTable.cpp
@@ -1,5 +1,6 @@
 // Copyright Epic Games, Inc. All Rights Reserved.
 #include "ProxyTable.h"
+#include "ProxyTableFunctionLibrary.h"
 
 FLookupProxy::FLookupProxy()
 {
@@ -70,4 +71,9 @@ bool FProxyTableContextProperty::GetValue(const UObject* ContextObject, const UP
 	}
 
 	return false;
+}
+
+UObject* UProxyTableFunctionLibrary::EvaluateProxyTable(const UObject* ContextObject, const UProxyTable* ProxyTable, FName Key)
+{
+	return FindProxyObject(ProxyTable, Key, ContextObject);
 }
\ No newline at end of file
diff --git a/Engine/Plugins/Experimental/Chooser/Source/ProxyTable/Public/ProxyTableFunctionLibrary.h b/Engine/Plugins/Experimental/Chooser/Source/ProxyTable/Public/ProxyTableFunctionLibrary.h
new file mode 100644
--- /dev/null
+++ b/Engine/Plugins/Experimental/Chooser/Source/ProxyTable/Public/ProxyTableFunctionLibrary.h
@@ -0,0 +1,28 @@
+// Copyright Epic Games, Inc. All Rights Reserved.
+
+#pragma once
+
+#include "Kismet/BlueprintFunctionLibrary.h"
+#include "ProxyTable.h"
+#include "ProxyTableFunctionLibrary.generated.h"
+
+/**
+ * Proxy Table Function Library
+ */
+UCLASS()
+class PROXYTABLE_API UProxyTableFunctionLibrary : public UBlueprintFunctionLibrary
+{
+	GENERATED_BODY()
+
+public:
+	/**
+	* Evaluate a proxy table and return the selected UObject, or null
+	*
+	* @param ContextObject			(in) An Object from which the parameters to the Chooser Table will be read
+	* @param ProxyTable				(in) The ProxyTable asset
+	* @param Key					(in) The Key from the ProxyTable asset
+	*/
+	UFUNCTION(BlueprintPure, meta = (BlueprintThreadSafe), Category = "Animation")
+	static UObject* EvaluateProxyTable(const UObject* ContextObject, const UProxyTable* ProxyTable, FName Key);
+
+};
\ No newline at end of file
