Chooser output column support and OutputBool column

diff --git a/Engine/Plugins/Experimental/Chooser/Source/Chooser/Internal/BoolColumn.h b/Engine/Plugins/Experimental/Chooser/Source/Chooser/Internal/BoolColumn.h
--- a/Engine/Plugins/Experimental/Chooser/Source/Chooser/Internal/BoolColumn.h
+++ b/Engine/Plugins/Experimental/Chooser/Source/Chooser/Internal/BoolColumn.h
@@ -18,6 +18,7 @@ public:
 	TArray<FName> PropertyBindingChain;
 	
 	virtual bool GetValue(const UObject* ContextObject, bool& OutResult) const override;
+	virtual bool SetValue(UObject* ContextObject, bool InValue) const override;
 
 #if WITH_EDITOR
 	static bool CanBind(const FProperty& Property)
diff --git a/Engine/Plugins/Experimental/Chooser/Source/Chooser/Internal/OutputBoolColumn.h b/Engine/Plugins/Experimental/Chooser/Source/Chooser/Internal/OutputBoolColumn.h
new file mode 100644
--- /dev/null
+++ b/Engine/Plugins/Experimental/Chooser/Source/Chooser/Internal/OutputBoolColumn.h
@@ -0,0 +1,28 @@
+// Copyright Epic Games, Inc. All Rights Reserved.
+#pragma once
+
+#include "CoreMinimal.h"
+#include "IChooserColumn.h"
+#include "IChooserParameterBool.h"
+#include "ChooserPropertyAccess.h"
+#include "InstancedStruct.h"
+#include "OutputBoolColumn.generated.h"
+
+USTRUCT()
+struct CHOOSER_API FOutputBoolColumn : public FChooserColumnBase
+{
+	GENERATED_BODY()
+	public:
+	FOutputBoolColumn();
+	virtual bool HasFilters() const override { return false; }
+	virtual bool HasOutputs() const override { return true; }
+	virtual void SetOutputs(UObject* ContextObject, int RowIndex) const override;
+	
+	UPROPERTY(EditAnywhere, Meta = (ExcludeBaseStruct, BaseStruct = "/Script/Chooser.ChooserParameterBoolBase"), Category = "Hidden")
+	FInstancedStruct InputValue;
+	
+	UPROPERTY(EditAnywhere, Category=Runtime);
+	TArray<bool> RowValues; 
+
+	CHOOSER_COLUMN_BOILERPLATE(FChooserParameterBoolBase);
+};
\ No newline at end of file
diff --git a/Engine/Plugins/Experimental/Chooser/Source/Chooser/Private/BoolColumn.cpp b/Engine/Plugins/Experimental/Chooser/Source/Chooser/Private/BoolColumn.cpp
--- a/Engine/Plugins/Experimental/Chooser/Source/Chooser/Private/BoolColumn.cpp
+++ b/Engine/Plugins/Experimental/Chooser/Source/Chooser/Private/BoolColumn.cpp
@@ -19,6 +19,24 @@ bool FBoolContextProperty::GetValue(const UObject* ContextObject, bool& OutResul
 	return false;
 }
 
+bool FBoolContextProperty::SetValue(UObject* ContextObject, bool InValue) const
+{
+	UStruct* StructType = ContextObject->GetClass();
+	const void* Container = ContextObject;
+	
+	if (UE::Chooser::ResolvePropertyChain(Container, StructType, PropertyBindingChain))
+	{
+		if (FBoolProperty* Property = FindFProperty<FBoolProperty>(StructType, PropertyBindingChain.Last()))
+		{
+			// const cast is here just because ResolvePropertyChain expects a const void*&
+			*Property->ContainerPtrToValuePtr<bool>(const_cast<void*>(Container)) = InValue;
+			return true;
+		}
+	}
+
+	return false;
+}
+
 FBoolColumn::FBoolColumn()
 {
 	InputValue.InitializeAs(FBoolContextProperty::StaticStruct());
diff --git a/Engine/Plugins/Experimental/Chooser/Source/Chooser/Private/Chooser.cpp b/Engine/Plugins/Experimental/Chooser/Source/Chooser/Private/Chooser.cpp
--- a/Engine/Plugins/Experimental/Chooser/Source/Chooser/Private/Chooser.cpp
+++ b/Engine/Plugins/Experimental/Chooser/Source/Chooser/Private/Chooser.cpp
@@ -81,12 +81,15 @@ static FObjectChooserBase::EIteratorStatus StaticEvaluateChooser(const UObject*
 
 	for (const FInstancedStruct& ColumnData : Chooser->ColumnsStructs)
 	{
-		const FChooserColumnBase& Column = ColumnData.Get<FChooserColumnBase>(); 
-		Swap(IndicesIn, IndicesOut);
-		IndicesOut->SetNum(0, false);
-		Column.Filter(ContextObject, *IndicesIn, *IndicesOut);
+		const FChooserColumnBase& Column = ColumnData.Get<FChooserColumnBase>();
+		if (Column.HasFilters())
+		{
+			Swap(IndicesIn, IndicesOut);
+			IndicesOut->SetNum(0, false);
+			Column.Filter(ContextObject, *IndicesIn, *IndicesOut);
+		}
 	}
-
+	
 	// of the rows that passed all column filters, return the first one for which the result row succeeds (could fail eg for a nexted chooser where no rows passed)
 	for (uint32 SelectedIndex : *IndicesOut)
 	{
@@ -95,6 +98,12 @@ static FObjectChooserBase::EIteratorStatus StaticEvaluateChooser(const UObject*
 			const FObjectChooserBase& SelectedResult = Chooser->ResultsStructs[SelectedIndex].Get<FObjectChooserBase>();
 			if (SelectedResult.ChooseMulti(ContextObject, Callback) == FObjectChooserBase::EIteratorStatus::Stop)
 			{
+				// trigger all output columns
+				for (const FInstancedStruct& ColumnData : Chooser->ColumnsStructs)
+				{
+					const FChooserColumnBase& Column = ColumnData.Get<FChooserColumnBase>();
+					Column.SetOutputs(const_cast<UObject*>(ContextObject), SelectedIndex);
+				}
 				return FObjectChooserBase::EIteratorStatus::Stop;
 			}
 		}
diff --git a/Engine/Plugins/Experimental/Chooser/Source/Chooser/Private/OutputBoolColumn.cpp b/Engine/Plugins/Experimental/Chooser/Source/Chooser/Private/OutputBoolColumn.cpp
new file mode 100644
--- /dev/null
+++ b/Engine/Plugins/Experimental/Chooser/Source/Chooser/Private/OutputBoolColumn.cpp
@@ -0,0 +1,14 @@
+// Copyright Epic Games, Inc. All Rights Reserved.
+#include "OutputBoolColumn.h"
+#include "ChooserPropertyAccess.h"
+#include "BoolColumn.h"
+
+FOutputBoolColumn::FOutputBoolColumn()
+{
+	InputValue.InitializeAs(FBoolContextProperty::StaticStruct());
+}
+
+void FOutputBoolColumn::SetOutputs(UObject* ContextObject, int RowIndex) const
+{
+	InputValue.Get<FChooserParameterBoolBase>().SetValue(ContextObject, RowValues[RowIndex]);
+}
diff --git a/Engine/Plugins/Experimental/Chooser/Source/Chooser/Public/IChooserColumn.h b/Engine/Plugins/Experimental/Chooser/Source/Chooser/Public/IChooserColumn.h
--- a/Engine/Plugins/Experimental/Chooser/Source/Chooser/Public/IChooserColumn.h
+++ b/Engine/Plugins/Experimental/Chooser/Source/Chooser/Public/IChooserColumn.h
@@ -22,7 +22,6 @@ public:
 	virtual void ConvertToInstancedStruct(FInstancedStruct& OutInstancedStruct) const {}
 };
 
-
 USTRUCT()
 struct CHOOSER_API FChooserColumnBase
 {
@@ -33,6 +32,11 @@ public:
 	virtual void PostLoad() {};
 	virtual void Filter(const UObject* ContextObject, const TArray<uint32>& IndexListIn, TArray<uint32>& IndexListOut) const {}
 
+	virtual bool HasFilters() const { return true; }
+	virtual bool HasOutputs() const { return false; }
+	
+	virtual void SetOutputs(UObject* ContextObject, int RowIndex) const { }
+
 #if WITH_EDITOR
 	virtual void SetNumRows(uint32 NumRows) {}
 	virtual void DeleteRows(const TArray<uint32> & RowIndices) {}
diff --git a/Engine/Plugins/Experimental/Chooser/Source/Chooser/Public/IChooserParameterBool.h b/Engine/Plugins/Experimental/Chooser/Source/Chooser/Public/IChooserParameterBool.h
--- a/Engine/Plugins/Experimental/Chooser/Source/Chooser/Public/IChooserParameterBool.h
+++ b/Engine/Plugins/Experimental/Chooser/Source/Chooser/Public/IChooserParameterBool.h
@@ -29,4 +29,5 @@ struct FChooserParameterBoolBase : public FChooserParameterBase
     
     public:
     	virtual bool GetValue(const UObject* ContextObject, bool& OutResult) const { return false; }
+    	virtual bool SetValue(UObject* ContextObject, bool InValue) const { return false; }
 };
\ No newline at end of file
diff --git a/Engine/Plugins/Experimental/Chooser/Source/ChooserEditor/Private/BoolColumnEditor.cpp b/Engine/Plugins/Experimental/Chooser/Source/ChooserEditor/Private/BoolColumnEditor.cpp
--- a/Engine/Plugins/Experimental/Chooser/Source/ChooserEditor/Private/BoolColumnEditor.cpp
+++ b/Engine/Plugins/Experimental/Chooser/Source/ChooserEditor/Private/BoolColumnEditor.cpp
@@ -2,6 +2,7 @@
 
 #include "BoolColumnEditor.h"
 #include "BoolColumn.h"
+#include "OutputBoolColumn.h"
 #include "ContextPropertyWidget.h"
 #include "ChooserTableEditor.h"
 #include "ObjectChooserWidgetFactories.h"
@@ -34,6 +35,27 @@ TSharedRef<SWidget> CreateBoolColumnWidget(UChooserTable* Chooser, FChooserColum
 	});
 }
 	
+TSharedRef<SWidget> CreateOutputBoolColumnWidget(UChooserTable* Chooser, FChooserColumnBase* Column, int Row)
+{
+	FOutputBoolColumn* BoolColumn = static_cast<FOutputBoolColumn*>(Column);
+
+	return SNew (SCheckBox)
+	.OnCheckStateChanged_Lambda([Chooser, BoolColumn,Row](ECheckBoxState State)
+	{
+		if (Row < BoolColumn->RowValues.Num())
+		{
+			const FScopedTransaction Transaction(LOCTEXT("Change Bool Value", "Change Bool Value"));
+			Chooser->Modify(true);
+			BoolColumn->RowValues[Row] = (State == ECheckBoxState::Checked);
+		}
+	})
+	.IsChecked_Lambda([BoolColumn, Row]()
+	{
+		const bool value = (Row < BoolColumn->RowValues.Num()) ? BoolColumn->RowValues[Row] : false;
+		return value ? ECheckBoxState::Checked : ECheckBoxState::Unchecked;
+	});
+}
+	
 TSharedRef<SWidget> CreateBoolPropertyWidget(UObject* TransactionObject, void* Value, UClass* ContextClass)
 {
 	return CreatePropertyWidget<FBoolContextProperty>(TransactionObject, Value, ContextClass, GetDefault<UGraphEditorSettings>()->BooleanPinTypeColor);
@@ -43,6 +65,7 @@ void RegisterBoolWidgets()
 {
 	FObjectChooserWidgetFactories::ChooserWidgetCreators.Add(FBoolContextProperty::StaticStruct(), CreateBoolPropertyWidget);
 	FChooserTableEditor::ColumnWidgetCreators.Add(FBoolColumn::StaticStruct(), CreateBoolColumnWidget);
+	FChooserTableEditor::ColumnWidgetCreators.Add(FOutputBoolColumn::StaticStruct(), CreateOutputBoolColumnWidget);
 }
 	
 }
diff --git a/Engine/Plugins/Experimental/Chooser/Source/ChooserEditor/Private/ChooserTableEditor.cpp b/Engine/Plugins/Experimental/Chooser/Source/ChooserEditor/Private/ChooserTableEditor.cpp
--- a/Engine/Plugins/Experimental/Chooser/Source/ChooserEditor/Private/ChooserTableEditor.cpp
+++ b/Engine/Plugins/Experimental/Chooser/Source/ChooserEditor/Private/ChooserTableEditor.cpp
@@ -578,7 +578,20 @@ void FChooserTableEditor::UpdateTableColumns()
 					return FReply::Handled();
 				})
 				[
-					HeaderWidget ? HeaderWidget.ToSharedRef() : SNullWidget::NullWidget
+					SNew(SHorizontalBox)
+					+ SHorizontalBox::Slot().AutoWidth()
+					[
+						SNew(SBorder)
+						.BorderBackgroundColor(FLinearColor(0,0,0,0))
+						.Content()
+						[
+							SNew(SImage).Image(Column.HasOutputs() ? FCoreStyle::Get().GetBrush("Icons.ArrowRight") :  FCoreStyle::Get().GetBrush("Icons.Filter"))
+						]
+					]
+					+ SHorizontalBox::Slot()
+					[
+						HeaderWidget ? HeaderWidget.ToSharedRef() : SNullWidget::NullWidget
+					]
 				]
 				
 			]);
@@ -623,13 +636,35 @@ TSharedRef<SDockTab> FChooserTableEditor::SpawnTableTab( const FSpawnTabArgs& Ar
 			const FScopedTransaction Transaction(LOCTEXT("Add Column Transaction", "Add Column"));
 			Chooser->Modify(true);
 
-			Chooser->ColumnsStructs.SetNum(Chooser->ColumnsStructs.Num()+1);
-			Chooser->ColumnsStructs.Last().InitializeAs(ChosenStruct);
+			FInstancedStruct NewColumn;
+			NewColumn.InitializeAs(ChosenStruct);
+			const FChooserColumnBase& NewColumnRef = NewColumn.Get<FChooserColumnBase>();
+			int InsertIndex = 0;
+			if (NewColumnRef.HasOutputs())
+			{
+				InsertIndex = Chooser->ColumnsStructs.Num();
+				// add output columns at the end
+				Chooser->ColumnsStructs.Add(NewColumn);
+			}
+			else
+			{
+				// add other columns after the loast non-output column
+				while(InsertIndex < Chooser->ColumnsStructs.Num())
+				{
+					const FChooserColumnBase& Column = Chooser->ColumnsStructs[InsertIndex].Get<FChooserColumnBase>();
+					if (Column.HasOutputs())
+					{
+						break;
+					}
+					InsertIndex++;
+				}
+				Chooser->ColumnsStructs.Insert(NewColumn, InsertIndex);
+			}
 			
 			UpdateTableColumns();
 			UpdateTableRows();
 
-			SelectColumn(Chooser->ColumnsStructs.Num() - 1);
+			SelectColumn(InsertIndex);
 		}));
 		return Widget;
 	})
