Row Drag/Drop support in ChooserTableEditor

diff --git a/Engine/Plugins/Experimental/Chooser/Source/Chooser/Public/IChooserColumn.h b/Engine/Plugins/Experimental/Chooser/Source/Chooser/Public/IChooserColumn.h
--- a/Engine/Plugins/Experimental/Chooser/Source/Chooser/Public/IChooserColumn.h
+++ b/Engine/Plugins/Experimental/Chooser/Source/Chooser/Public/IChooserColumn.h
@@ -40,6 +40,7 @@ public:
 #if WITH_EDITOR
 	virtual void SetNumRows(uint32 NumRows) {}
 	virtual void DeleteRows(const TArray<uint32> & RowIndices) {}
+	virtual void MoveRow(int SourceIndex, int TargetIndex) {}
 	
 	virtual UScriptStruct* GetInputBaseType() const { return nullptr; };
 	virtual const UScriptStruct* GetInputType() const { return nullptr; };
@@ -58,6 +59,13 @@ public:
 			RowValues.RemoveAt(Index);\
 		}\
 	}\
+	virtual void MoveRow(int SourceRowIndex, int TargetRowIndex)\
+	{\
+		auto RowData = RowValues[SourceRowIndex];\
+    	RowValues.RemoveAt(SourceRowIndex);\
+    	if (SourceRowIndex < TargetRowIndex) { TargetRowIndex--; }\
+    	RowValues.Insert(RowData, TargetRowIndex);\
+	}\
 	virtual UScriptStruct* GetInputBaseType() const override { return ParameterType::StaticStruct(); };\
 	virtual const UScriptStruct* GetInputType() const override { return InputValue.IsValid() ? InputValue.GetScriptStruct() : nullptr; };\
 	virtual FChooserParameterBase* GetInputValue() override { return InputValue.IsValid() ? &InputValue.GetMutable<FChooserParameterBase>() : nullptr; };
diff --git a/Engine/Plugins/Experimental/Chooser/Source/ChooserEditor/Private/ChooserTableEditor.cpp b/Engine/Plugins/Experimental/Chooser/Source/ChooserEditor/Private/ChooserTableEditor.cpp
--- a/Engine/Plugins/Experimental/Chooser/Source/ChooserEditor/Private/ChooserTableEditor.cpp
+++ b/Engine/Plugins/Experimental/Chooser/Source/ChooserEditor/Private/ChooserTableEditor.cpp
@@ -32,6 +32,7 @@
 #include "IDetailCustomization.h"
 #include "PropertyCustomizationHelpers.h"
 #include "ScopedTransaction.h"
+#include "Widgets/Layout/SSeparator.h"
 
 #define LOCTEXT_NAMESPACE "ChooserEditor"
 
@@ -408,7 +409,22 @@ public:
 				}),
 				&CacheBorder
 				);
-				return ResultWidget.ToSharedRef();
+				
+				return SNew(SOverlay)
+                		+ SOverlay::Slot()
+                		[
+                			ResultWidget.ToSharedRef()
+                		]
+                		+ SOverlay::Slot().VAlign(VAlign_Bottom)
+                		[
+                			SNew(SSeparator).SeparatorImage(FCoreStyle::Get().GetBrush("FocusRectangle"))
+                			.Visibility_Lambda([this]() { return bDragActive && !bDropAbove ? EVisibility::Visible : EVisibility::Hidden; })
+                		]
+                		+ SOverlay::Slot().VAlign(VAlign_Top)
+                		[
+                			SNew(SSeparator).SeparatorImage(FCoreStyle::Get().GetBrush("FocusRectangle"))
+                			.Visibility_Lambda([this]() { return bDragActive && bDropAbove ? EVisibility::Visible : EVisibility::Hidden; })
+                		];
 			}
 			else
 			{
@@ -436,21 +452,69 @@ public:
 			}
 		}
 		else if (RowIndex->RowIndex == Chooser->ResultsStructs.Num())
-        {
+		{
 			// on the row past the end, show an Add button in the result column
 			if (ColumnName == Result)
 			{
 				return Editor->GetCreateRowComboButton().ToSharedRef();
 			}
 		}
+
 		return SNullWidget::NullWidget;
 	}
+	
+	virtual void OnDragEnter(const FGeometry& MyGeometry, const FDragDropEvent& DragDropEvent) override
+	{
+		if (TSharedPtr<FChooserRowDragDropOp> Operation = DragDropEvent.GetOperationAs<FChooserRowDragDropOp>())
+		{
+			bDragActive = true;
+			float Center = MyGeometry.Position.Y + MyGeometry.Size.Y;
+			bDropAbove = DragDropEvent.GetScreenSpacePosition().Y < Center;
+		}
+	}
+	virtual void OnDragLeave(const FDragDropEvent& DragDropEvent) override
+	{
+		bDragActive = false;
+	}
+	
+	virtual FReply OnDragOver(const FGeometry& MyGeometry, const FDragDropEvent& DragDropEvent) override
+	{
+		if (TSharedPtr<FChooserRowDragDropOp> Operation = DragDropEvent.GetOperationAs<FChooserRowDragDropOp>())
+		{
+			float Center = MyGeometry.AbsolutePosition.Y + MyGeometry.Size.Y/2;
+			bDropAbove = DragDropEvent.GetScreenSpacePosition().Y < Center;
+			return FReply::Handled();
+		}
+		return FReply::Unhandled();
+	}
+	
+    virtual FReply OnDrop(const FGeometry& MyGeometry, const FDragDropEvent& DragDropEvent) override
+	{
+		if (TSharedPtr<FChooserRowDragDropOp> Operation = DragDropEvent.GetOperationAs<FChooserRowDragDropOp>())
+		{
+			if (Chooser == Operation->ChooserEditor->GetChooser())
+			{
+				if (bDropAbove)
+				{
+					Editor->MoveRow(Operation->RowIndex, RowIndex->RowIndex);
+				}
+				else
+				{
+					Editor->MoveRow(Operation->RowIndex, RowIndex->RowIndex+1);
+				}
+				return FReply::Handled();		
+			}
+		}
+		return FReply::Unhandled();
+	}
 
 private:
 	TSharedPtr<FChooserTableEditor::FChooserTableRow> RowIndex;
 	UChooserTable* Chooser;
 	FChooserTableEditor* Editor;
 	TSharedPtr<SBorder> CacheBorder;
+	bool bDragActive = false;
+	bool bDropAbove = false;
 };
 
 
@@ -473,6 +537,31 @@ FReply FChooserTableEditor::SelectRootProperties()
 	return FReply::Handled();
 }
 
+void FChooserTableEditor::MoveRow(int SourceRowIndex, int TargetRowIndex)
+{
+	UChooserTable* Chooser = Cast<UChooserTable>(EditingObjects[0]);
+	TargetRowIndex = FMath::Min(TargetRowIndex,Chooser->ResultsStructs.Num());
+	
+	const FScopedTransaction Transaction(LOCTEXT("Move Row", "Move Row"));
+	
+	Chooser->Modify(true);
+
+	for (FInstancedStruct& ColStruct : Chooser->ColumnsStructs)
+	{
+		FChooserColumnBase& Column = ColStruct.GetMutable<FChooserColumnBase>();
+		Column.MoveRow(SourceRowIndex, TargetRowIndex);
+	}
+	
+	FInstancedStruct Result = Chooser->ResultsStructs[SourceRowIndex];
+	Chooser->ResultsStructs.RemoveAt(SourceRowIndex);
+	if (SourceRowIndex < TargetRowIndex)
+	{
+		TargetRowIndex--;
+	}
+	Chooser->ResultsStructs.Insert(Result, TargetRowIndex);
+
+	UpdateTableRows();
+}
 
 void FChooserTableEditor::UpdateTableColumns()
 {
diff --git a/Engine/Plugins/Experimental/Chooser/Source/ChooserEditor/Private/ChooserTableEditor.h b/Engine/Plugins/Experimental/Chooser/Source/ChooserEditor/Private/ChooserTableEditor.h
--- a/Engine/Plugins/Experimental/Chooser/Source/ChooserEditor/Private/ChooserTableEditor.h
+++ b/Engine/Plugins/Experimental/Chooser/Source/ChooserEditor/Private/ChooserTableEditor.h
@@ -98,6 +98,7 @@ namespace UE::ChooserEditor
 		void UpdateTableRows();
 		void SelectColumn(int Index);
 		void DeleteColumn(int Index);
+		void MoveRow(int SourceRowIndex, int TargetIndex);
 	private:
 
 		FReply SelectRootProperties();
diff --git a/Engine/Plugins/Experimental/Chooser/Source/ChooserEditor/Public/ContextPropertyWidget.h b/Engine/Plugins/Experimental/Chooser/Source/ChooserEditor/Public/ContextPropertyWidget.h
--- a/Engine/Plugins/Experimental/Chooser/Source/ChooserEditor/Public/ContextPropertyWidget.h
+++ b/Engine/Plugins/Experimental/Chooser/Source/ChooserEditor/Public/ContextPropertyWidget.h
@@ -35,8 +35,8 @@ namespace UE::ChooserEditor
 	
 		Args.OnCanAcceptPropertyOrChildren = FOnCanBindProperty::CreateLambda([](FProperty* InProperty)
 			{
-				// Make only editor visible properties visible for binding.
-				return InProperty->HasAnyPropertyFlags(CPF_Edit);
+				// Make only blueprint visible properties visible for binding.
+				return InProperty->HasAnyPropertyFlags(CPF_BlueprintVisible);
 			});	
 	
 		Args.OnAddBinding = FOnAddBinding::CreateLambda([TransactionObject, ContextProperty](FName InPropertyName, const TArray<FBindingChainElement>& InBindingChain)
