Public api cleanup for Chooser widget factories

diff --git a/Engine/Plugins/Experimental/Chooser/Source/Chooser/Chooser.Build.cs b/Engine/Plugins/Experimental/Chooser/Source/Chooser/Chooser.Build.cs
--- a/Engine/Plugins/Experimental/Chooser/Source/Chooser/Chooser.Build.cs
+++ b/Engine/Plugins/Experimental/Chooser/Source/Chooser/Chooser.Build.cs
@@ -9,6 +9,7 @@ namespace UnrealBuildTool.Rules
 			PublicDependencyModuleNames.AddRange(
 				new string[]
 				{
+					"StructUtils",
 					// ... add other public dependencies that you statically link with here ...
 				}
 			);
@@ -20,7 +21,6 @@ namespace UnrealBuildTool.Rules
 					"CoreUObject",
 					"Engine",
 					"GameplayTags",
-					"StructUtils"
 					// ... add private dependencies that you statically link with here ...
 				}
 			);
diff --git a/Engine/Plugins/Experimental/Chooser/Source/ChooserEditor/Private/BoolColumnEditor.cpp b/Engine/Plugins/Experimental/Chooser/Source/ChooserEditor/Private/BoolColumnEditor.cpp
--- a/Engine/Plugins/Experimental/Chooser/Source/ChooserEditor/Private/BoolColumnEditor.cpp
+++ b/Engine/Plugins/Experimental/Chooser/Source/ChooserEditor/Private/BoolColumnEditor.cpp
@@ -63,9 +63,9 @@ TSharedRef<SWidget> CreateBoolPropertyWidget(UObject* TransactionObject, void* V
 	
 void RegisterBoolWidgets()
 {
-	FObjectChooserWidgetFactories::ChooserWidgetCreators.Add(FBoolContextProperty::StaticStruct(), CreateBoolPropertyWidget);
-	FChooserTableEditor::ColumnWidgetCreators.Add(FBoolColumn::StaticStruct(), CreateBoolColumnWidget);
-	FChooserTableEditor::ColumnWidgetCreators.Add(FOutputBoolColumn::StaticStruct(), CreateOutputBoolColumnWidget);
+	FObjectChooserWidgetFactories::RegisterWidgetCreator(FBoolContextProperty::StaticStruct(), CreateBoolPropertyWidget);
+	FObjectChooserWidgetFactories::RegisterColumnWidgetCreator(FBoolColumn::StaticStruct(), CreateBoolColumnWidget);
+	FObjectChooserWidgetFactories::RegisterColumnWidgetCreator(FOutputBoolColumn::StaticStruct(), CreateOutputBoolColumnWidget);
 }
 	
 }
diff --git a/Engine/Plugins/Experimental/Chooser/Source/ChooserEditor/Private/ChooserTableEditor.cpp b/Engine/Plugins/Experimental/Chooser/Source/ChooserEditor/Private/ChooserTableEditor.cpp
--- a/Engine/Plugins/Experimental/Chooser/Source/ChooserEditor/Private/ChooserTableEditor.cpp
+++ b/Engine/Plugins/Experimental/Chooser/Source/ChooserEditor/Private/ChooserTableEditor.cpp
@@ -433,16 +433,8 @@ public:
 				{
 					FChooserColumnBase* Column = &Chooser->ColumnsStructs[ColumnIndex].GetMutable<FChooserColumnBase>();
 					const UStruct * ColumnStruct = Chooser->ColumnsStructs[ColumnIndex].GetScriptStruct();
-					TSharedPtr<SWidget> ColumnWidget;
-					while (ColumnStruct && !ColumnWidget.IsValid())
-					{
-						if (auto Creator = FChooserTableEditor::ColumnWidgetCreators.Find(ColumnStruct))
-						{
-							ColumnWidget = (*Creator)(Chooser, Column, RowIndex->RowIndex);
-							break;
-						}
-						ColumnStruct = ColumnStruct->GetSuperStruct();
-					}
+
+					TSharedPtr<SWidget> ColumnWidget = FObjectChooserWidgetFactories::CreateColumnWidget(Column, ColumnStruct, Chooser, RowIndex->RowIndex);
 					
 					if (ColumnWidget.IsValid())
 					{
@@ -969,10 +961,6 @@ void FChooserTableEditor::DeleteColumn(int Index)
 	}
 }
 
-/// Result widgets
-
-TMap<const UStruct*, TFunction<TSharedRef<SWidget> (UChooserTable* Chooser, FChooserColumnBase* Column, int Row)>> FChooserTableEditor::ColumnWidgetCreators;
-
 TSharedRef<SWidget> CreateAssetWidget(UObject* TransactionObject, void* Value, UClass* ContextClass)
 {
 	FAssetChooser* DIAsset = static_cast<FAssetChooser*>(Value);
@@ -1147,8 +1135,8 @@ void FChooserColumnDetails::CustomizeDetails(IDetailLayoutBuilder& DetailBuilder
 void FChooserTableEditor::RegisterWidgets()
 {
 	// todo: fallback widget
-	FObjectChooserWidgetFactories::ChooserWidgetCreators.Add(FAssetChooser::StaticStruct(), CreateAssetWidget);
-	FObjectChooserWidgetFactories::ChooserWidgetCreators.Add(FEvaluateChooser::StaticStruct(), CreateEvaluateChooserWidget);
+	FObjectChooserWidgetFactories::RegisterWidgetCreator(FAssetChooser::StaticStruct(), CreateAssetWidget);
+	FObjectChooserWidgetFactories::RegisterWidgetCreator(FEvaluateChooser::StaticStruct(), CreateEvaluateChooserWidget);
 	
 
 	FPropertyEditorModule& PropertyModule = FModuleManager::LoadModuleChecked<FPropertyEditorModule>("PropertyEditor");
diff --git a/Engine/Plugins/Experimental/Chooser/Source/ChooserEditor/Private/ChooserTableEditor.h b/Engine/Plugins/Experimental/Chooser/Source/ChooserEditor/Private/ChooserTableEditor.h
--- a/Engine/Plugins/Experimental/Chooser/Source/ChooserEditor/Private/ChooserTableEditor.h
+++ b/Engine/Plugins/Experimental/Chooser/Source/ChooserEditor/Private/ChooserTableEditor.h
@@ -141,8 +141,6 @@ namespace UE::ChooserEditor
 	public:
 		TSharedPtr<SComboButton>& GetCreateRowComboButton() { return CreateRowComboButton; };
 
-		static TMap<const UStruct*, TFunction<TSharedRef<SWidget> (UChooserTable* Chooser, FChooserColumnBase* Column, int Row)>> ColumnWidgetCreators;
-	
 		/** The name given to all instances of this type of editor */
 		static const FName ToolkitFName;
 
diff --git a/Engine/Plugins/Experimental/Chooser/Source/ChooserEditor/Private/EnumColumnEditor.cpp b/Engine/Plugins/Experimental/Chooser/Source/ChooserEditor/Private/EnumColumnEditor.cpp
--- a/Engine/Plugins/Experimental/Chooser/Source/ChooserEditor/Private/EnumColumnEditor.cpp
+++ b/Engine/Plugins/Experimental/Chooser/Source/ChooserEditor/Private/EnumColumnEditor.cpp
@@ -126,8 +126,8 @@ TSharedRef<SWidget> CreateEnumPropertyWidget(UObject* TransactionObject, void* V
 	
 void RegisterEnumWidgets()
 {
-	FObjectChooserWidgetFactories::ChooserWidgetCreators.Add(FEnumContextProperty::StaticStruct(), CreateEnumPropertyWidget);
-	FChooserTableEditor::ColumnWidgetCreators.Add(FEnumColumn::StaticStruct(), CreateEnumColumnWidget);
+	FObjectChooserWidgetFactories::RegisterWidgetCreator(FEnumContextProperty::StaticStruct(), CreateEnumPropertyWidget);
+	FObjectChooserWidgetFactories::RegisterColumnWidgetCreator(FEnumColumn::StaticStruct(), CreateEnumColumnWidget);
 }
 	
 }
diff --git a/Engine/Plugins/Experimental/Chooser/Source/ChooserEditor/Private/FloatRangeColumnEditor.cpp b/Engine/Plugins/Experimental/Chooser/Source/ChooserEditor/Private/FloatRangeColumnEditor.cpp
--- a/Engine/Plugins/Experimental/Chooser/Source/ChooserEditor/Private/FloatRangeColumnEditor.cpp
+++ b/Engine/Plugins/Experimental/Chooser/Source/ChooserEditor/Private/FloatRangeColumnEditor.cpp
@@ -81,8 +81,8 @@ TSharedRef<SWidget> CreateFloatRangeColumnWidget(UChooserTable* Chooser, FChoose
 	
 void RegisterFloatRangeWidgets()
 {
-	FObjectChooserWidgetFactories::ChooserWidgetCreators.Add(FFloatContextProperty::StaticStruct(), CreateFloatPropertyWidget);
-	FChooserTableEditor::ColumnWidgetCreators.Add(FFloatRangeColumn::StaticStruct(), CreateFloatRangeColumnWidget);
+	FObjectChooserWidgetFactories::RegisterWidgetCreator(FFloatContextProperty::StaticStruct(), CreateFloatPropertyWidget);
+	FObjectChooserWidgetFactories::RegisterColumnWidgetCreator(FFloatRangeColumn::StaticStruct(), CreateFloatRangeColumnWidget);
 }
 	
 }
diff --git a/Engine/Plugins/Experimental/Chooser/Source/ChooserEditor/Private/GameplayTagColumnEditor.cpp b/Engine/Plugins/Experimental/Chooser/Source/ChooserEditor/Private/GameplayTagColumnEditor.cpp
--- a/Engine/Plugins/Experimental/Chooser/Source/ChooserEditor/Private/GameplayTagColumnEditor.cpp
+++ b/Engine/Plugins/Experimental/Chooser/Source/ChooserEditor/Private/GameplayTagColumnEditor.cpp
@@ -56,8 +56,8 @@ TSharedRef<SWidget> CreateGameplayTagPropertyWidget(UObject* TransactionObject,
 	
 void RegisterGameplayTagWidgets()
 {
-	FObjectChooserWidgetFactories::ChooserWidgetCreators.Add(FGameplayTagContextProperty::StaticStruct(), CreateGameplayTagPropertyWidget);
-	FChooserTableEditor::ColumnWidgetCreators.Add(FGameplayTagColumn::StaticStruct(), CreateGameplayTagColumnWidget);
+	FObjectChooserWidgetFactories::RegisterWidgetCreator(FGameplayTagContextProperty::StaticStruct(), CreateGameplayTagPropertyWidget);
+	FObjectChooserWidgetFactories::RegisterColumnWidgetCreator(FGameplayTagColumn::StaticStruct(), CreateGameplayTagColumnWidget);
 }
 	
 }
diff --git a/Engine/Plugins/Experimental/Chooser/Source/ChooserEditor/Private/ObjectChooserWidgetFactories.cpp b/Engine/Plugins/Experimental/Chooser/Source/ChooserEditor/Private/ObjectChooserWidgetFactories.cpp
--- a/Engine/Plugins/Experimental/Chooser/Source/ChooserEditor/Private/ObjectChooserWidgetFactories.cpp
+++ b/Engine/Plugins/Experimental/Chooser/Source/ChooserEditor/Private/ObjectChooserWidgetFactories.cpp
@@ -17,7 +17,35 @@
 
 namespace UE::ChooserEditor
 {
+	
+void FObjectChooserWidgetFactories::RegisterWidgetCreator(const UStruct* Type, FChooserWidgetCreator Creator)
+{
+	ChooserWidgetCreators.Add(Type, Creator);	
+}
+	
+void  FObjectChooserWidgetFactories::RegisterColumnWidgetCreator(const UStruct* ColumnType, FColumnWidgetCreator Creator)
+{
+	ColumnWidgetCreators.Add(ColumnType, Creator);	
+}
+	
+TSharedPtr<SWidget> FObjectChooserWidgetFactories::CreateColumnWidget(FChooserColumnBase* Column, const UStruct* ColumnType, UChooserTable* Chooser, int RowIndex)
+{
+	if (Column)
+	{
+		while (ColumnType)
+		{
+			if (FColumnWidgetCreator* Creator = FObjectChooserWidgetFactories::ColumnWidgetCreators.Find(ColumnType))
+			{
+				return (*Creator)(Chooser, Column, RowIndex);
+			}
+			ColumnType = ColumnType->GetSuperStruct();
+		}
+	}
 
+	return nullptr;
+}
+	
+	
 TSharedPtr<SWidget> FObjectChooserWidgetFactories::CreateWidget(UObject* TransactionObject, void* Value, const UStruct* ValueType, UClass* ContextClass)
 {
 	if (Value)
@@ -95,6 +123,7 @@ TSharedPtr<SWidget> FObjectChooserWidgetFactories::CreateWidget(UObject* Transac
 	return Widget;
 }
 
+TMap<const UStruct*, TFunction<TSharedRef<SWidget>(UChooserTable* Chooser, FChooserColumnBase* Column, int Row)>> FObjectChooserWidgetFactories::ColumnWidgetCreators;
 TMap<const UStruct*, FChooserWidgetCreator> FObjectChooserWidgetFactories::ChooserWidgetCreators;
 
 void FObjectChooserWidgetFactories::RegisterWidgets()
diff --git a/Engine/Plugins/Experimental/Chooser/Source/ChooserEditor/Public/ObjectChooserWidgetFactories.h b/Engine/Plugins/Experimental/Chooser/Source/ChooserEditor/Public/ObjectChooserWidgetFactories.h
--- a/Engine/Plugins/Experimental/Chooser/Source/ChooserEditor/Public/ObjectChooserWidgetFactories.h
+++ b/Engine/Plugins/Experimental/Chooser/Source/ChooserEditor/Public/ObjectChooserWidgetFactories.h
@@ -2,23 +2,36 @@
 
 #pragma once
 
+#include "Chooser.h"
 #include "CoreMinimal.h"
 #include "Widgets/SWidget.h"
 #include "Widgets/Layout/SBorder.h"
 #include "StructViewerModule.h"
 
+class UChooserTable;
+struct FChooserColumnBase;
+
 namespace UE::ChooserEditor
 {
 	typedef TFunction<TSharedRef<SWidget>(UObject* TransactionObject, void* Value, UClass* ContextClass)> FChooserWidgetCreator;
+	typedef TFunction<TSharedRef<SWidget>(UChooserTable* Chooser, FChooserColumnBase* Column, int Row)> FColumnWidgetCreator;
 
 	class CHOOSEREDITOR_API FObjectChooserWidgetFactories
 	{
 	public:
-		static TMap<const UStruct*, FChooserWidgetCreator> ChooserWidgetCreators;
 
 		static TSharedPtr<SWidget> CreateWidget(UObject* TransactionObject, void* Value, const UStruct* ValueType, UClass* ContextClass);
 		static TSharedPtr<SWidget> CreateWidget(UObject* TransactionObject, const UScriptStruct* BaseType, void* Value, const UStruct* ValueType, UClass* ContextClass, const FOnStructPicked& CreateClassCallback, TSharedPtr<SBorder>* InnerWidget = nullptr);
 		
+		static TSharedPtr<SWidget> CreateColumnWidget(FChooserColumnBase* Column, const UStruct* ColumnType, UChooserTable* Chooser, int Row);
+		
 		static void RegisterWidgets();
+
+		static void RegisterWidgetCreator(const UStruct* Type, FChooserWidgetCreator Creator);
+		static void RegisterColumnWidgetCreator(const UStruct* ColumnType, FColumnWidgetCreator Creator);
+		
+	private:
+		static TMap<const UStruct*, FColumnWidgetCreator>  ColumnWidgetCreators;
+		static TMap<const UStruct*, FChooserWidgetCreator> ChooserWidgetCreators;
 	};
 }
diff --git a/Engine/Plugins/Experimental/Chooser/Source/ProxyTableEditor/Private/ProxyTableEditor.cpp b/Engine/Plugins/Experimental/Chooser/Source/ProxyTableEditor/Private/ProxyTableEditor.cpp
--- a/Engine/Plugins/Experimental/Chooser/Source/ProxyTableEditor/Private/ProxyTableEditor.cpp
+++ b/Engine/Plugins/Experimental/Chooser/Source/ProxyTableEditor/Private/ProxyTableEditor.cpp
@@ -668,7 +668,7 @@ TSharedRef<SWidget> CreateLookupProxyWidget(UObject* TransactionObject, void* Va
 
 void FProxyTableEditor::RegisterWidgets()
 {
-	UE::ChooserEditor::FObjectChooserWidgetFactories::ChooserWidgetCreators.Add(FLookupProxy::StaticStruct(), CreateLookupProxyWidget);
+	UE::ChooserEditor::FObjectChooserWidgetFactories::RegisterWidgetCreator(FLookupProxy::StaticStruct(), CreateLookupProxyWidget);
 }
 	
 }
