Support for moving chooser columns around in column header context menu

diff --git a/Engine/Plugins/Experimental/Chooser/Source/ChooserEditor/Private/ChooserTableEditor.cpp b/Engine/Plugins/Experimental/Chooser/Source/ChooserEditor/Private/ChooserTableEditor.cpp
--- a/Engine/Plugins/Experimental/Chooser/Source/ChooserEditor/Private/ChooserTableEditor.cpp
+++ b/Engine/Plugins/Experimental/Chooser/Source/ChooserEditor/Private/ChooserTableEditor.cpp
@@ -604,6 +604,35 @@ void FChooserTableEditor::UpdateTableColumns()
 						)
 					);
 
+				if (ColumnIndex > 0)
+				{
+					MenuBuilder.AddMenuEntry(LOCTEXT("Move Left","Move Left"),LOCTEXT("Move Left ToolTip", "Move this column to the left."),FSlateIcon(),
+						FUIAction(
+							FExecuteAction::CreateLambda([this,Chooser,ColumnIndex, &Column]()
+							{
+								const FScopedTransaction Transaction(LOCTEXT("Move Column Left Transaction", "Move Column Left"));
+								Chooser->Modify(true);
+								Chooser->ColumnsStructs.Swap(ColumnIndex, ColumnIndex - 1);
+								UpdateTableColumns();
+							})
+							));
+				}
+				if (ColumnIndex < Chooser->ColumnsStructs.Num() - 1)
+				{
+					MenuBuilder.AddMenuEntry(LOCTEXT("Move Right","Move Right"),LOCTEXT("Move Right ToolTip", "Move this column to the right."),FSlateIcon(),
+						FUIAction(
+							FExecuteAction::CreateLambda([this,Chooser,ColumnIndex, &Column]()
+							{
+								const FScopedTransaction Transaction(LOCTEXT("Move Column Right Transaction", "Move Column Right"));
+								Chooser->Modify(true);
+								Chooser->ColumnsStructs.Swap(ColumnIndex, ColumnIndex + 1);
+								UpdateTableColumns();
+							})
+							));
+				}
+
+				// if (ColumnIndex < Chooser->ColumnsStructs.Num() - 1)
+
 				MenuBuilder.AddMenuEntry(LOCTEXT("Delete Column","Delete"),LOCTEXT("Delete Column ToolTip", "Remove this column and all its data from the table"),FSlateIcon(),
 					FUIAction(
 						FExecuteAction::CreateLambda([this,Chooser,ColumnIndex, &Column]()
