Null check in EvaluateProxy function

diff --git a/Engine/Plugins/Experimental/Chooser/Source/ProxyTable/Private/ProxyTable.cpp b/Engine/Plugins/Experimental/Chooser/Source/ProxyTable/Private/ProxyTable.cpp
--- a/Engine/Plugins/Experimental/Chooser/Source/ProxyTable/Private/ProxyTable.cpp
+++ b/Engine/Plugins/Experimental/Chooser/Source/ProxyTable/Private/ProxyTable.cpp
@@ -9,21 +9,24 @@ FLookupProxy::FLookupProxy()
 
 static UObject* FindProxyObject(const UProxyTable* Table, FName Key, const UObject* ContextObject)
 {
-	for (const FProxyEntry& Entry : Table->Entries)
+	if (Table)
 	{
-		if (Entry.Key == Key && Entry.ValueStruct.IsValid())
+		for (const FProxyEntry& Entry : Table->Entries)
 		{
-			const FObjectChooserBase& EntryValue = Entry.ValueStruct.Get<FObjectChooserBase>();
-			return EntryValue.ChooseObject(ContextObject);
+			if (Entry.Key == Key && Entry.ValueStruct.IsValid())
+			{
+				const FObjectChooserBase& EntryValue = Entry.ValueStruct.Get<FObjectChooserBase>();
+				return EntryValue.ChooseObject(ContextObject);
+			}
 		}
-	}
 
-	// search parent tables (uncooked data only)
-	for (const TObjectPtr<UProxyTable> ParentTable : Table->InheritEntriesFrom)
-	{
-		if (UObject* Value = FindProxyObject(ParentTable, Key, ContextObject))
+		// search parent tables (uncooked data only)
+		for (const TObjectPtr<UProxyTable> ParentTable : Table->InheritEntriesFrom)
 		{
-			return Value;
+			if (UObject* Value = FindProxyObject(ParentTable, Key, ContextObject))
+			{
+				return Value;
+			}
 		}
 	}
 
@@ -37,12 +40,9 @@ UObject* FLookupProxy::ChooseObject(const UObject* ContextObject) const
 		const UProxyTable* Table;
 		if (ProxyTable.Get<FChooserParameterProxyTableBase>().GetValue(ContextObject, Table))
 		{
-			if (Table)
+			if (UObject* Value = FindProxyObject(Table, Key, ContextObject))
 			{
-				if (UObject* Value = FindProxyObject(Table, Key, ContextObject))
-				{
-					return Value;
-				}
+				return Value;
 			}
 		}
 	}
