Class result type for Choosers and Proxy Tables

diff --git a/Engine/Plugins/Experimental/Chooser/Source/Chooser/Internal/ObjectChooser_Class.h b/Engine/Plugins/Experimental/Chooser/Source/Chooser/Internal/ObjectChooser_Class.h
new file mode 100644
--- /dev/null
+++ b/Engine/Plugins/Experimental/Chooser/Source/Chooser/Internal/ObjectChooser_Class.h
@@ -0,0 +1,19 @@
+// Copyright Epic Games, Inc. All Rights Reserved.
+
+#pragma once
+
+#include "CoreMinimal.h"
+#include "IObjectChooser.h"
+#include "ObjectChooser_Class.generated.h"
+
+USTRUCT(DisplayName = "Class")
+struct CHOOSER_API FClassChooser : public FObjectChooserBase
+{
+	GENERATED_BODY()
+	
+	// FObjectChooserBase interface
+	virtual UObject* ChooseObject(const UObject* ContextObject) const final override;
+public: 
+	UPROPERTY(EditAnywhere, Category = "Parameters")
+	TObjectPtr<UClass> Class;
+};
diff --git a/Engine/Plugins/Experimental/Chooser/Source/Chooser/Private/ObjectChooser_Class.cpp b/Engine/Plugins/Experimental/Chooser/Source/Chooser/Private/ObjectChooser_Class.cpp
new file mode 100644
--- /dev/null
+++ b/Engine/Plugins/Experimental/Chooser/Source/Chooser/Private/ObjectChooser_Class.cpp
@@ -0,0 +1,8 @@
+// Copyright Epic Games, Inc. All Rights Reserved.
+
+#include "ObjectChooser_Class.h"
+
+UObject* FClassChooser::ChooseObject(const UObject* ContextObject) const
+{
+	return Class;
+}
\ No newline at end of file
diff --git a/Engine/Plugins/Experimental/Chooser/Source/ChooserEditor/Private/ChooserTableEditor.cpp b/Engine/Plugins/Experimental/Chooser/Source/ChooserEditor/Private/ChooserTableEditor.cpp
--- a/Engine/Plugins/Experimental/Chooser/Source/ChooserEditor/Private/ChooserTableEditor.cpp
+++ b/Engine/Plugins/Experimental/Chooser/Source/ChooserEditor/Private/ChooserTableEditor.cpp
@@ -26,6 +26,7 @@
 #include "Widgets/Layout/SWidgetSwitcher.h"
 #include "DragAndDrop/DecoratedDragDropOp.h"
 #include "ObjectChooser_Asset.h"
+#include "ObjectChooser_Class.h"
 #include "ObjectChooserClassFilter.h"
 #include "ObjectChooserWidgetFactories.h"
 #include "GraphEditorSettings.h"
@@ -1008,6 +1009,28 @@ TSharedRef<SWidget> CreateAssetWidget(UObject* TransactionObject, void* Value, U
 			DIAsset->Asset = AssetData.GetAsset();
 		});
 }
+	
+TSharedRef<SWidget> CreateClassWidget(UObject* TransactionObject, void* Value, UClass* ContextClass)
+{
+	FClassChooser* ClassChooser = static_cast<FClassChooser*>(Value);
+
+	UClass* Class = ClassChooser->Class;
+
+	UChooserTable* Chooser = Cast<UChooserTable>(TransactionObject);
+	
+	return SNew(SClassPropertyEntryBox)
+			.MetaClass((Chooser!=nullptr && Chooser->OutputObjectType!=nullptr) ? Chooser->OutputObjectType.Get() : UObject::StaticClass())
+			.SelectedClass_Lambda([ClassChooser]()
+			{
+				return ClassChooser->Class;
+			})
+			.OnSetClass_Lambda([TransactionObject, ClassChooser](const UClass* SelectedClass)
+			{
+				const FScopedTransaction Transaction(LOCTEXT("Edit Class", "Edit Class"));
+				TransactionObject->Modify(true);
+				ClassChooser->Class = const_cast<UClass*>(SelectedClass);
+			});
+}
 
 TSharedRef<SWidget> CreateEvaluateChooserWidget(UObject* TransactionObject, void* Value, UClass* ContextObject)
 {
@@ -1165,6 +1188,7 @@ void FChooserTableEditor::RegisterWidgets()
 {
 	// todo: fallback widget
 	FObjectChooserWidgetFactories::RegisterWidgetCreator(FAssetChooser::StaticStruct(), CreateAssetWidget);
+	FObjectChooserWidgetFactories::RegisterWidgetCreator(FClassChooser::StaticStruct(), CreateClassWidget);
 	FObjectChooserWidgetFactories::RegisterWidgetCreator(FEvaluateChooser::StaticStruct(), CreateEvaluateChooserWidget);
 	
 
