ProxyTable and ChooserTable asset definition conversion and cleanup

diff --git a/Engine/Plugins/Experimental/Chooser/Source/ChooserEditor/Private/AssetDefinition_ChooserTable.cpp b/Engine/Plugins/Experimental/Chooser/Source/ChooserEditor/Private/AssetDefinition_ChooserTable.cpp
--- a/Engine/Plugins/Experimental/Chooser/Source/ChooserEditor/Private/AssetDefinition_ChooserTable.cpp
+++ b/Engine/Plugins/Experimental/Chooser/Source/ChooserEditor/Private/AssetDefinition_ChooserTable.cpp
@@ -7,8 +7,12 @@
 
 EAssetCommandResult UAssetDefinition_ChooserTable::OpenAssets(const FAssetOpenArgs& OpenArgs) const
 {
-	TArray<UObject*> ChooserTables = TArray<UObject*>(OpenArgs.LoadObjects<UChooserTable>());
-	UE::ChooserEditor::FChooserTableEditor::CreateEditor(OpenArgs.GetToolkitMode(), OpenArgs.ToolkitHost, ChooserTables);
+	TArray<UObject*> Objects = OpenArgs.LoadObjects<UObject>();
+	
+	for (UObject* Object : Objects)
+	{
+		UE::ChooserEditor::FChooserTableEditor::CreateEditor(OpenArgs.GetToolkitMode(), OpenArgs.ToolkitHost, Object);
+	}
 
 	return EAssetCommandResult::Handled;
 }
diff --git a/Engine/Plugins/Experimental/Chooser/Source/ChooserEditor/Private/AssetDefinition_ChooserTable.h b/Engine/Plugins/Experimental/Chooser/Source/ChooserEditor/Private/AssetDefinition_ChooserTable.h
--- a/Engine/Plugins/Experimental/Chooser/Source/ChooserEditor/Private/AssetDefinition_ChooserTable.h
+++ b/Engine/Plugins/Experimental/Chooser/Source/ChooserEditor/Private/AssetDefinition_ChooserTable.h
@@ -19,8 +19,7 @@ public:
 	virtual TSoftClassPtr<UObject> GetAssetClass() const override { return UChooserTable::StaticClass(); }
 	virtual TConstArrayView<FAssetCategoryPath> GetAssetCategories() const override
 	{
-		static const auto Categories = { FAssetCategoryPath(EAssetCategoryPaths::Animation, NSLOCTEXT("AssetTypeActions", "DataInterfacesSubMenu", "Data Interfaces")) };
-		
+		static const auto Categories = { FAssetCategoryPath(EAssetCategoryPaths::Misc) };
 		return Categories;
 	}
 	virtual EAssetCommandResult OpenAssets(const FAssetOpenArgs& OpenArgs) const override;
diff --git a/Engine/Plugins/Experimental/Chooser/Source/ChooserEditor/Private/ChooserEditorModule.h b/Engine/Plugins/Experimental/Chooser/Source/ChooserEditor/Private/ChooserEditorModule.h
--- a/Engine/Plugins/Experimental/Chooser/Source/ChooserEditor/Private/ChooserEditorModule.h
+++ b/Engine/Plugins/Experimental/Chooser/Source/ChooserEditor/Private/ChooserEditorModule.h
@@ -8,8 +8,6 @@
 namespace UE::ChooserEditor
 {
 
-class FAssetTypeActions_ChooserTable;
-
 class FModule : public IModuleInterface
 {
 private:
diff --git a/Engine/Plugins/Experimental/Chooser/Source/ProxyTableEditor/Private/AssetDefinition_ProxyTable.cpp b/Engine/Plugins/Experimental/Chooser/Source/ProxyTableEditor/Private/AssetDefinition_ProxyTable.cpp
new file mode 100644
--- /dev/null
+++ b/Engine/Plugins/Experimental/Chooser/Source/ProxyTableEditor/Private/AssetDefinition_ProxyTable.cpp
@@ -0,0 +1,20 @@
+// Copyright Epic Games, Inc. All Rights Reserved.
+
+#include "AssetDefinition_ProxyTable.h"
+#include "ProxyTableEditor.h"
+
+#define LOCTEXT_NAMESPACE "UAssetDefinition_ProxyTable" 
+
+EAssetCommandResult UAssetDefinition_ProxyTable::OpenAssets(const FAssetOpenArgs& OpenArgs) const
+{
+	TArray<UObject*> Objects = OpenArgs.LoadObjects<UObject>();
+	
+	for (UObject* Object : Objects)
+	{
+		UE::ProxyTableEditor::FProxyTableEditor::CreateEditor(OpenArgs.GetToolkitMode(), OpenArgs.ToolkitHost, Object);
+	}
+
+	return EAssetCommandResult::Handled;
+}
+
+#undef LOCTEXT_NAMESPACE
\ No newline at end of file
diff --git a/Engine/Plugins/Experimental/Chooser/Source/ProxyTableEditor/Private/AssetDefinition_ProxyTable.h b/Engine/Plugins/Experimental/Chooser/Source/ProxyTableEditor/Private/AssetDefinition_ProxyTable.h
new file mode 100644
--- /dev/null
+++ b/Engine/Plugins/Experimental/Chooser/Source/ProxyTableEditor/Private/AssetDefinition_ProxyTable.h
@@ -0,0 +1,27 @@
+// Copyright Epic Games, Inc. All Rights Reserved.
+
+#pragma once
+
+#include "ProxyTable.h"
+#include "AssetDefinitionDefault.h"
+
+#include "AssetDefinition_ProxyTable.generated.h"
+
+UCLASS()
+class UAssetDefinition_ProxyTable : public UAssetDefinitionDefault
+{
+	GENERATED_BODY()
+
+public:
+	// UAssetDefinition Begin
+	virtual FText GetAssetDisplayName() const override { return NSLOCTEXT("AssetTypeActions", "AssetTypeActions_ProxyTable", "Proxy Table"); }
+	virtual FLinearColor GetAssetColor() const override { return FLinearColor(FColor(128,128,64)); }
+	virtual TSoftClassPtr<UObject> GetAssetClass() const override { return UProxyTable::StaticClass(); }
+	virtual TConstArrayView<FAssetCategoryPath> GetAssetCategories() const override
+	{
+		static const auto Categories = { FAssetCategoryPath(EAssetCategoryPaths::Misc) };
+		return Categories;
+	}
+	virtual EAssetCommandResult OpenAssets(const FAssetOpenArgs& OpenArgs) const override;
+	// UAssetDefinition End
+};
diff --git a/Engine/Plugins/Experimental/Chooser/Source/ProxyTableEditor/Private/AssetTypeActions_ProxyTable.cpp b/Engine/Plugins/Experimental/Chooser/Source/ProxyTableEditor/Private/AssetTypeActions_ProxyTable.cpp
deleted file mode 100644
diff --git a/Engine/Plugins/Experimental/Chooser/Source/ProxyTableEditor/Private/AssetTypeActions_ProxyTable.h b/Engine/Plugins/Experimental/Chooser/Source/ProxyTableEditor/Private/AssetTypeActions_ProxyTable.h
deleted file mode 100644
diff --git a/Engine/Plugins/Experimental/Chooser/Source/ProxyTableEditor/Private/ProxyTableEditorModule.cpp b/Engine/Plugins/Experimental/Chooser/Source/ProxyTableEditor/Private/ProxyTableEditorModule.cpp
--- a/Engine/Plugins/Experimental/Chooser/Source/ProxyTableEditor/Private/ProxyTableEditorModule.cpp
+++ b/Engine/Plugins/Experimental/Chooser/Source/ProxyTableEditor/Private/ProxyTableEditorModule.cpp
@@ -2,7 +2,6 @@
 
 #include "ProxyTableEditorModule.h"
 #include "IAssetTools.h"
-#include "AssetTypeActions_ProxyTable.h"
 #include "ProxyTableEditor.h"
 
 #define LOCTEXT_NAMESPACE "ChooserEditorModule"
@@ -12,21 +11,11 @@ namespace UE::ProxyTableEditor
 
 void FModule::StartupModule()
 {
-	IAssetTools& AssetTools = FModuleManager::LoadModuleChecked<FAssetToolsModule>("AssetTools").Get();
-	
-	AssetTypeActions_ProxyTable = MakeShared<FAssetTypeActions_ProxyTable>();
-	AssetTools.RegisterAssetTypeActions(AssetTypeActions_ProxyTable.ToSharedRef());
 	FProxyTableEditor::RegisterWidgets();
 }
-		 
 
 void FModule::ShutdownModule()
 {
-	if(FModuleManager::Get().IsModuleLoaded("AssetTools"))
-	{
-		IAssetTools& AssetTools = FModuleManager::GetModuleChecked<FAssetToolsModule>("AssetTools").Get();
-		AssetTools.UnregisterAssetTypeActions(AssetTypeActions_ProxyTable.ToSharedRef());
-	}
 }
 
 }
diff --git a/Engine/Plugins/Experimental/Chooser/Source/ProxyTableEditor/Private/ProxyTableEditorModule.h b/Engine/Plugins/Experimental/Chooser/Source/ProxyTableEditor/Private/ProxyTableEditorModule.h
--- a/Engine/Plugins/Experimental/Chooser/Source/ProxyTableEditor/Private/ProxyTableEditorModule.h
+++ b/Engine/Plugins/Experimental/Chooser/Source/ProxyTableEditor/Private/ProxyTableEditorModule.h
@@ -8,16 +8,11 @@
 namespace UE::ProxyTableEditor
 {
 
-class FAssetTypeActions_ChooserTable;
-class FAssetTypeActions_ProxyTable;
-
 class FModule : public IModuleInterface
 {
 private:
 	virtual void StartupModule() override;
 	virtual void ShutdownModule() override;
-
-	TSharedPtr<FAssetTypeActions_ProxyTable> AssetTypeActions_ProxyTable;
 };
 
 }
\ No newline at end of file
diff --git a/Engine/Plugins/Experimental/Chooser/Source/ProxyTableEditor/ProxyTableEditor.Build.cs b/Engine/Plugins/Experimental/Chooser/Source/ProxyTableEditor/ProxyTableEditor.Build.cs
--- a/Engine/Plugins/Experimental/Chooser/Source/ProxyTableEditor/ProxyTableEditor.Build.cs
+++ b/Engine/Plugins/Experimental/Chooser/Source/ProxyTableEditor/ProxyTableEditor.Build.cs
@@ -13,6 +13,7 @@ namespace UnrealBuildTool.Rules
 					"CoreUObject",
 					"Engine",
 					"AssetTools",
+					"AssetDefinition",
 					"Chooser",
 					"ChooserEditor",
 					"ProxyTable",
