support for moving rows in proxy tables with drag/drop

diff --git a/Engine/Plugins/Experimental/Chooser/Source/ProxyTableEditor/Private/ProxyTableEditor.cpp b/Engine/Plugins/Experimental/Chooser/Source/ProxyTableEditor/Private/ProxyTableEditor.cpp
--- a/Engine/Plugins/Experimental/Chooser/Source/ProxyTableEditor/Private/ProxyTableEditor.cpp
+++ b/Engine/Plugins/Experimental/Chooser/Source/ProxyTableEditor/Private/ProxyTableEditor.cpp
@@ -30,7 +30,7 @@
 #include "ObjectChooserWidgetFactories.h"
 #include "ContextPropertyWidget.h"
 #include "IObjectChooser.h"
-#include "Misc/TransactionObjectEvent.h"
+#include "Widgets/Layout/SSeparator.h"
 
 #define LOCTEXT_NAMESPACE "ProxyTableEditor"
 
@@ -286,14 +286,14 @@ class FProxyRowDragDropOp : public FDecoratedDragDropOp
 public:
 	DRAG_DROP_OPERATOR_TYPE(FWidgetTemplateDragDropOp, FDecoratedDragDropOp)
 
-	FProxyTableEditor* ChooserEditor;
+	FProxyTableEditor* Editor;
 	uint32 RowIndex;
 
 	/** Constructs the drag drop operation */
 	static TSharedRef<FProxyRowDragDropOp> New(FProxyTableEditor* InEditor, uint32 InRowIndex)
 	{
 		TSharedRef<FProxyRowDragDropOp> Operation = MakeShareable(new FProxyRowDragDropOp());
-		Operation->ChooserEditor = InEditor;
+		Operation->Editor = InEditor;
 		Operation->RowIndex = InRowIndex;
 		Operation->DefaultHoverText = LOCTEXT("Proxy Row", "Proxy Row");
 		// UE::ChooserEditor::FObjectChooserWidgetFactories::ConvertToText(InEditor->GetProxyTable()->Entries[InRowIndex].Value.GetObject(), Operation->DefaultHoverText);
@@ -402,7 +402,22 @@ public:
 					}),
 					&CacheBorder
 					);
-				return ResultWidget.ToSharedRef();
+
+				return SNew(SOverlay)
+					+ SOverlay::Slot()
+					[
+						ResultWidget.ToSharedRef()
+					]
+					+ SOverlay::Slot().VAlign(VAlign_Bottom)
+					[
+						SNew(SSeparator).SeparatorImage(FCoreStyle::Get().GetBrush("FocusRectangle"))
+						.Visibility_Lambda([this]() { return bDragActive && !bDropAbove ? EVisibility::Visible : EVisibility::Hidden; })
+					]
+					+ SOverlay::Slot().VAlign(VAlign_Top)
+					[
+						SNew(SSeparator).SeparatorImage(FCoreStyle::Get().GetBrush("FocusRectangle"))
+						.Visibility_Lambda([this]() { return bDragActive && bDropAbove ? EVisibility::Visible : EVisibility::Hidden; })
+					];
 			}
 			else if (ColumnName == Key)
 			{
@@ -428,11 +443,59 @@ public:
 		return SNullWidget::NullWidget;
 	}
 
+	
+	virtual void OnDragEnter(const FGeometry& MyGeometry, const FDragDropEvent& DragDropEvent) override
+	{
+		if (TSharedPtr<FProxyRowDragDropOp> Operation = DragDropEvent.GetOperationAs<FProxyRowDragDropOp>())
+		{
+			bDragActive = true;
+			float Center = MyGeometry.Position.Y + MyGeometry.Size.Y;
+			bDropAbove = DragDropEvent.GetScreenSpacePosition().Y < Center;
+		}
+	}
+	virtual void OnDragLeave(const FDragDropEvent& DragDropEvent) override
+	{
+		bDragActive = false;
+	}
+	
+	virtual FReply OnDragOver(const FGeometry& MyGeometry, const FDragDropEvent& DragDropEvent) override
+	{
+		if (TSharedPtr<FProxyRowDragDropOp> Operation = DragDropEvent.GetOperationAs<FProxyRowDragDropOp>())
+		{
+			float Center = MyGeometry.AbsolutePosition.Y + MyGeometry.Size.Y/2;
+			bDropAbove = DragDropEvent.GetScreenSpacePosition().Y < Center;
+			return FReply::Handled();
+		}
+		return FReply::Unhandled();
+	}
+	
+	virtual FReply OnDrop(const FGeometry& MyGeometry, const FDragDropEvent& DragDropEvent) override
+	{
+		if (TSharedPtr<FProxyRowDragDropOp> Operation = DragDropEvent.GetOperationAs<FProxyRowDragDropOp>())
+		{
+			if (ProxyTable == Operation->Editor->GetProxyTable())
+			{
+				if (bDropAbove)
+				{
+					Editor->MoveRow(Operation->RowIndex, RowIndex->RowIndex);
+				}
+				else
+				{
+					Editor->MoveRow(Operation->RowIndex, RowIndex->RowIndex+1);
+				}
+				return FReply::Handled();		
+			}
+		}
+		return FReply::Unhandled();
+	}	
+
 private:
 	TSharedPtr<FProxyTableEditor::FProxyTableRow> RowIndex;
 	UProxyTable* ProxyTable;
 	FProxyTableEditor* Editor;
 	TSharedPtr<SBorder> CacheBorder;
+	bool bDragActive = false;
+	bool bDropAbove = false;	
 };
 
 
@@ -455,6 +518,25 @@ FReply FProxyTableEditor::SelectRootProperties()
 	return FReply::Handled();
 }
 
+void FProxyTableEditor::MoveRow(int SourceRowIndex, int TargetRowIndex)
+{
+	UProxyTable* Table = Cast<UProxyTable>(EditingObjects[0]);
+	TargetRowIndex = FMath::Min(TargetRowIndex,Table->Entries.Num());
+	
+	const FScopedTransaction Transaction(LOCTEXT("Move Row", "Move Row"));
+	
+	Table->Modify(true);
+
+	FProxyEntry Entry = Table->Entries[SourceRowIndex];
+	Table->Entries.RemoveAt(SourceRowIndex);
+	if (SourceRowIndex < TargetRowIndex)
+	{
+		TargetRowIndex--;
+	}
+	Table->Entries.Insert(Entry, TargetRowIndex);
+
+	UpdateTableRows();
+}
 
 void FProxyTableEditor::UpdateTableColumns()
 {
diff --git a/Engine/Plugins/Experimental/Chooser/Source/ProxyTableEditor/Private/ProxyTableEditor.h b/Engine/Plugins/Experimental/Chooser/Source/ProxyTableEditor/Private/ProxyTableEditor.h
--- a/Engine/Plugins/Experimental/Chooser/Source/ProxyTableEditor/Private/ProxyTableEditor.h
+++ b/Engine/Plugins/Experimental/Chooser/Source/ProxyTableEditor/Private/ProxyTableEditor.h
@@ -86,6 +86,7 @@ namespace UE::ProxyTableEditor
 
 		void UpdateTableColumns();
 		void UpdateTableRows();
+		void MoveRow(int SourceRowIndex, int TargetIndex);
 	private:
 
 		FReply SelectRootProperties();
