add "Table Settings" button to toolbar for chooser table and proxy table editors

diff --git a/Engine/Plugins/Experimental/Chooser/Source/ChooserEditor/ChooserEditor.Build.cs b/Engine/Plugins/Experimental/Chooser/Source/ChooserEditor/ChooserEditor.Build.cs
--- a/Engine/Plugins/Experimental/Chooser/Source/ChooserEditor/ChooserEditor.Build.cs
+++ b/Engine/Plugins/Experimental/Chooser/Source/ChooserEditor/ChooserEditor.Build.cs
@@ -18,6 +18,7 @@ namespace UnrealBuildTool.Rules
 					"UnrealEd",
 					"EditorWidgets",
 					"ToolWidgets",
+					"ToolMenus",
 					"SlateCore",
 					"Slate",
 					"PropertyEditor",
diff --git a/Engine/Plugins/Experimental/Chooser/Source/ChooserEditor/Private/ChooserEditorModule.cpp b/Engine/Plugins/Experimental/Chooser/Source/ChooserEditor/Private/ChooserEditorModule.cpp
--- a/Engine/Plugins/Experimental/Chooser/Source/ChooserEditor/Private/ChooserEditorModule.cpp
+++ b/Engine/Plugins/Experimental/Chooser/Source/ChooserEditor/Private/ChooserEditorModule.cpp
@@ -7,6 +7,7 @@
 #include "EnumColumnEditor.h"
 #include "FloatRangeColumnEditor.h"
 #include "GameplayTagColumnEditor.h"
+#include "ChooserTableEditorCommands.h"
 
 #define LOCTEXT_NAMESPACE "ChooserEditorModule"
 
@@ -20,10 +21,14 @@ void FModule::StartupModule()
 	RegisterFloatRangeWidgets();
 	RegisterBoolWidgets();
 	RegisterEnumWidgets();
+	
+	FChooserTableEditorCommands::Register();
+
 }
 
 void FModule::ShutdownModule()
 {
+	FChooserTableEditorCommands::Unregister();
 }
 
 }
diff --git a/Engine/Plugins/Experimental/Chooser/Source/ChooserEditor/Private/ChooserTableEditor.cpp b/Engine/Plugins/Experimental/Chooser/Source/ChooserEditor/Private/ChooserTableEditor.cpp
--- a/Engine/Plugins/Experimental/Chooser/Source/ChooserEditor/Private/ChooserTableEditor.cpp
+++ b/Engine/Plugins/Experimental/Chooser/Source/ChooserEditor/Private/ChooserTableEditor.cpp
@@ -34,6 +34,7 @@
 #include "PropertyCustomizationHelpers.h"
 #include "ScopedTransaction.h"
 #include "Widgets/Layout/SSeparator.h"
+#include "ChooserTableEditorCommands.h"
 
 #define LOCTEXT_NAMESPACE "ChooserEditor"
 
@@ -43,7 +44,7 @@ namespace UE::ChooserEditor
 const FName FChooserTableEditor::ToolkitFName( TEXT( "ChooserTableEditor" ) );
 const FName FChooserTableEditor::PropertiesTabId( TEXT( "ChooserEditor_Properties" ) );
 const FName FChooserTableEditor::TableTabId( TEXT( "ChooserEditor_Table" ) );
-
+	
 void FChooserTableEditor::RegisterTabSpawners(const TSharedRef<class FTabManager>& InTabManager)
 {
 	WorkspaceMenuCategory = InTabManager->AddLocalWorkspaceMenuCategory(LOCTEXT("WorkspaceMenu_ChooserTableEditor", "Chooser Table Editor"));
@@ -77,7 +78,43 @@ FChooserTableEditor::~FChooserTableEditor()
 	
 	DetailsView.Reset();
 }
+	
+void FChooserTableEditor::RegisterToolbar()
+{
+	UToolMenus* ToolMenus = UToolMenus::Get();
+	UToolMenu* ToolBar;
+	FName ParentName;
+	const FName MenuName = GetToolMenuToolbarName(ParentName);
+	if (ToolMenus->IsMenuRegistered(MenuName))
+	{
+		ToolBar = ToolMenus->ExtendMenu(MenuName);
+	}
+	else
+	{
+		ToolBar = UToolMenus::Get()->RegisterMenu(MenuName, ParentName, EMultiBoxType::ToolBar);
+	}
 
+	const FChooserTableEditorCommands& Commands = FChooserTableEditorCommands::Get();
+	FToolMenuInsert InsertAfterAssetSection("Asset", EToolMenuInsertType::After);
+	{
+		FToolMenuSection& Section = ToolBar->AddSection("Chooser", TAttribute<FText>(), InsertAfterAssetSection);
+		Section.AddEntry(FToolMenuEntry::InitToolBarButton(
+			Commands.EditChooserSettings,
+			TAttribute<FText>(),
+			TAttribute<FText>(),
+			FSlateIcon("EditorStyle", "FullBlueprintEditor.EditGlobalOptions")));
+	}
+
+}
+
+void FChooserTableEditor::BindCommands()
+{
+	const FChooserTableEditorCommands& Commands = FChooserTableEditorCommands::Get();
+
+	ToolkitCommands->MapAction(
+		Commands.EditChooserSettings,
+		FExecuteAction::CreateSP(this, &FChooserTableEditor::SelectRootProperties));
+}
 
 void FChooserTableEditor::InitEditor( const EToolkitMode::Type Mode, const TSharedPtr< class IToolkitHost >& InitToolkitHost, const TArray<UObject*>& ObjectsToEdit, FGetDetailsViewObjects GetDetailsViewObjects )
 {
@@ -115,7 +152,9 @@ void FChooserTableEditor::InitEditor( const EToolkitMode::Type Mode, const TShar
 	const bool bCreateDefaultToolbar = true;
 	FAssetEditorToolkit::InitAssetEditor( Mode, InitToolkitHost, FChooserTableEditor::ChooserEditorAppIdentifier, StandaloneDefaultLayout, bCreateDefaultStandaloneMenu, bCreateDefaultToolbar, ObjectsToEdit );
 
+	BindCommands();
 	RegenerateMenusAndToolbars();
+	RegisterToolbar();
 
 	SelectRootProperties();
 }
@@ -519,15 +558,13 @@ TSharedRef<ITableRow> FChooserTableEditor::GenerateTableRow(TSharedPtr<FChooserT
 		.Entry(InItem).Chooser(Chooser).Editor(this);
 }
 
-FReply FChooserTableEditor::SelectRootProperties()
+void FChooserTableEditor::SelectRootProperties()
 {
 	if( DetailsView.IsValid() )
 	{
 		// Make sure details window is pointing to our object
 		DetailsView->SetObjects( EditingObjects );
 	}
-
-	return FReply::Handled();
 }
 
 void FChooserTableEditor::MoveRow(int SourceRowIndex, int TargetRowIndex)
@@ -563,11 +600,8 @@ void FChooserTableEditor::UpdateTableColumns()
 	HeaderRow->ClearColumns();
 	
 	HeaderRow->AddColumn(SHeaderRow::Column("Handles")
-					.ManualWidth(30)
-					.HeaderContent()
-					[					
-						SNew(SButton).OnClicked_Raw(this, &FChooserTableEditor::SelectRootProperties)
-					]);
+					.DefaultLabel(FText())
+					.ManualWidth(30));
 	
 	HeaderRow->AddColumn(SHeaderRow::Column("Result")
 					.DefaultLabel(LOCTEXT("ResultColumnName", "Result"))
@@ -1186,11 +1220,9 @@ void FChooserColumnDetails::CustomizeDetails(IDetailLayoutBuilder& DetailBuilder
 
 void FChooserTableEditor::RegisterWidgets()
 {
-	// todo: fallback widget
 	FObjectChooserWidgetFactories::RegisterWidgetCreator(FAssetChooser::StaticStruct(), CreateAssetWidget);
 	FObjectChooserWidgetFactories::RegisterWidgetCreator(FClassChooser::StaticStruct(), CreateClassWidget);
 	FObjectChooserWidgetFactories::RegisterWidgetCreator(FEvaluateChooser::StaticStruct(), CreateEvaluateChooserWidget);
-	
 
 	FPropertyEditorModule& PropertyModule = FModuleManager::LoadModuleChecked<FPropertyEditorModule>("PropertyEditor");
 	
diff --git a/Engine/Plugins/Experimental/Chooser/Source/ChooserEditor/Private/ChooserTableEditor.h b/Engine/Plugins/Experimental/Chooser/Source/ChooserEditor/Private/ChooserTableEditor.h
--- a/Engine/Plugins/Experimental/Chooser/Source/ChooserEditor/Private/ChooserTableEditor.h
+++ b/Engine/Plugins/Experimental/Chooser/Source/ChooserEditor/Private/ChooserTableEditor.h
@@ -101,7 +101,9 @@ namespace UE::ChooserEditor
 		void MoveRow(int SourceRowIndex, int TargetIndex);
 	private:
 
-		FReply SelectRootProperties();
+		void SelectRootProperties();
+		void RegisterToolbar();
+		void BindCommands();
 	
 		/** Create the properties tab and its content */
 		TSharedRef<SDockTab> SpawnPropertiesTab( const FSpawnTabArgs& Args );
diff --git a/Engine/Plugins/Experimental/Chooser/Source/ChooserEditor/Private/ChooserTableEditorCommands.h b/Engine/Plugins/Experimental/Chooser/Source/ChooserEditor/Private/ChooserTableEditorCommands.h
new file mode 100644
--- /dev/null
+++ b/Engine/Plugins/Experimental/Chooser/Source/ChooserEditor/Private/ChooserTableEditorCommands.h
@@ -0,0 +1,28 @@
+// Copyright Epic Games, Inc. All Rights Reserved.
+
+#pragma once
+
+#include "CoreMinimal.h"
+#include "Framework/Commands/Commands.h"
+
+#define LOCTEXT_NAMESPACE "ChooserTableEditorCommands"
+	
+class FChooserTableEditorCommands : public TCommands<FChooserTableEditorCommands>
+{
+public:
+	/** Constructor */
+	FChooserTableEditorCommands() 
+		: TCommands<FChooserTableEditorCommands>("ChooserTableEditor", NSLOCTEXT("Contexts", "ChooserTableEditor", "Chooser Table Editor"), NAME_None, FAppStyle::GetAppStyleSetName())
+	{
+	}
+
+	TSharedPtr<FUICommandInfo> EditChooserSettings;
+
+	/** Initialize commands */
+	virtual void RegisterCommands() override
+	{
+		UI_COMMAND(EditChooserSettings, "Table Settings", "Edit the root properties of the ChooserTable asset.", EUserInterfaceActionType::Button, FInputChord())
+	}
+};
+	
+#undef LOCTEXT_NAMESPACE
diff --git a/Engine/Plugins/Experimental/Chooser/Source/ProxyTableEditor/Private/ProxyTableEditor.cpp b/Engine/Plugins/Experimental/Chooser/Source/ProxyTableEditor/Private/ProxyTableEditor.cpp
--- a/Engine/Plugins/Experimental/Chooser/Source/ProxyTableEditor/Private/ProxyTableEditor.cpp
+++ b/Engine/Plugins/Experimental/Chooser/Source/ProxyTableEditor/Private/ProxyTableEditor.cpp
@@ -31,6 +31,7 @@
 #include "ContextPropertyWidget.h"
 #include "IObjectChooser.h"
 #include "Widgets/Layout/SSeparator.h"
+#include "ProxyTableEditorCommands.h"
 
 #define LOCTEXT_NAMESPACE "ProxyTableEditor"
 
@@ -76,6 +77,42 @@ FProxyTableEditor::~FProxyTableEditor()
 	DetailsView.Reset();
 }
 
+void FProxyTableEditor::RegisterToolbar()
+{
+	UToolMenus* ToolMenus = UToolMenus::Get();
+	UToolMenu* ToolBar;
+	FName ParentName;
+	const FName MenuName = GetToolMenuToolbarName(ParentName);
+	if (ToolMenus->IsMenuRegistered(MenuName))
+	{
+		ToolBar = ToolMenus->ExtendMenu(MenuName);
+	}
+	else
+	{
+		ToolBar = UToolMenus::Get()->RegisterMenu(MenuName, ParentName, EMultiBoxType::ToolBar);
+	}
+
+	const FProxyTableEditorCommands& Commands = FProxyTableEditorCommands::Get();
+	FToolMenuInsert InsertAfterAssetSection("Asset", EToolMenuInsertType::After);
+	{
+		FToolMenuSection& Section = ToolBar->AddSection("Proxy Table", TAttribute<FText>(), InsertAfterAssetSection);
+		Section.AddEntry(FToolMenuEntry::InitToolBarButton(
+			Commands.EditTableSettings,
+			TAttribute<FText>(),
+			TAttribute<FText>(),
+			FSlateIcon("EditorStyle", "FullBlueprintEditor.EditGlobalOptions")));
+	}
+
+}
+
+void FProxyTableEditor::BindCommands()
+{
+	const FProxyTableEditorCommands& Commands = FProxyTableEditorCommands::Get();
+
+	ToolkitCommands->MapAction(
+		Commands.EditTableSettings,
+		FExecuteAction::CreateSP(this, &FProxyTableEditor::SelectRootProperties));
+}
 
 void FProxyTableEditor::InitEditor( const EToolkitMode::Type Mode, const TSharedPtr< class IToolkitHost >& InitToolkitHost, const TArray<UObject*>& ObjectsToEdit, FGetDetailsViewObjects GetDetailsViewObjects )
 {
@@ -113,7 +150,9 @@ void FProxyTableEditor::InitEditor( const EToolkitMode::Type Mode, const TShared
 	const bool bCreateDefaultToolbar = true;
 	FAssetEditorToolkit::InitAssetEditor( Mode, InitToolkitHost, FProxyTableEditor::ProxyEditorAppIdentifier, StandaloneDefaultLayout, bCreateDefaultStandaloneMenu, bCreateDefaultToolbar, ObjectsToEdit );
 
+	BindCommands();
 	RegenerateMenusAndToolbars();
+	RegisterToolbar();
 
 	SelectRootProperties();
 }
@@ -507,15 +546,13 @@ TSharedRef<ITableRow> FProxyTableEditor::GenerateTableRow(TSharedPtr<FProxyTable
 		.Entry(InItem).ProxyTable(ProxyTable).Editor(this);
 }
 
-FReply FProxyTableEditor::SelectRootProperties()
+void FProxyTableEditor::SelectRootProperties()
 {
 	if( DetailsView.IsValid() )
 	{
 		// Make sure details window is pointing to our object
 		DetailsView->SetObjects( EditingObjects );
 	}
-
-	return FReply::Handled();
 }
 
 void FProxyTableEditor::MoveRow(int SourceRowIndex, int TargetRowIndex)
@@ -542,11 +579,8 @@ void FProxyTableEditor::UpdateTableColumns()
 {
 	HeaderRow->ClearColumns();
 	HeaderRow->AddColumn(SHeaderRow::Column("Handles")
-					.ManualWidth(30)
-					.HeaderContent()
-					[					
-						SNew(SButton).OnClicked_Raw(this, &FProxyTableEditor::SelectRootProperties)
-					]);
+					.DefaultLabel(FText())
+					.ManualWidth(30));
 	
 	HeaderRow->AddColumn(SHeaderRow::Column("Key")
 					.DefaultLabel(LOCTEXT("KeyColumnName", "Key"))
diff --git a/Engine/Plugins/Experimental/Chooser/Source/ProxyTableEditor/Private/ProxyTableEditor.h b/Engine/Plugins/Experimental/Chooser/Source/ProxyTableEditor/Private/ProxyTableEditor.h
--- a/Engine/Plugins/Experimental/Chooser/Source/ProxyTableEditor/Private/ProxyTableEditor.h
+++ b/Engine/Plugins/Experimental/Chooser/Source/ProxyTableEditor/Private/ProxyTableEditor.h
@@ -89,7 +89,9 @@ namespace UE::ProxyTableEditor
 		void MoveRow(int SourceRowIndex, int TargetIndex);
 	private:
 
-		FReply SelectRootProperties();
+		void SelectRootProperties();
+		void RegisterToolbar();
+		void BindCommands();
 	
 		/** Create the properties tab and its content */
 		TSharedRef<SDockTab> SpawnPropertiesTab( const FSpawnTabArgs& Args );
diff --git a/Engine/Plugins/Experimental/Chooser/Source/ProxyTableEditor/Private/ProxyTableEditorCommands.h b/Engine/Plugins/Experimental/Chooser/Source/ProxyTableEditor/Private/ProxyTableEditorCommands.h
new file mode 100644
--- /dev/null
+++ b/Engine/Plugins/Experimental/Chooser/Source/ProxyTableEditor/Private/ProxyTableEditorCommands.h
@@ -0,0 +1,28 @@
+// Copyright Epic Games, Inc. All Rights Reserved.
+
+#pragma once
+
+#include "CoreMinimal.h"
+#include "Framework/Commands/Commands.h"
+
+#define LOCTEXT_NAMESPACE "ProxyTableEditorCommands"
+	
+class FProxyTableEditorCommands : public TCommands<FProxyTableEditorCommands>
+{
+public:
+	/** Constructor */
+	FProxyTableEditorCommands() 
+		: TCommands<FProxyTableEditorCommands>("ProxyTableEditor", NSLOCTEXT("Contexts", "ProxyTableEditor", "Proxy Table Editor"), NAME_None, FAppStyle::GetAppStyleSetName())
+	{
+	}
+
+	TSharedPtr<FUICommandInfo> EditTableSettings;
+
+	/** Initialize commands */
+	virtual void RegisterCommands() override
+	{
+		UI_COMMAND(EditTableSettings, "Table Settings", "Edit the root properties of the ChooserTable asset.", EUserInterfaceActionType::Button, FInputChord())
+	}
+};
+	
+#undef LOCTEXT_NAMESPACE
diff --git a/Engine/Plugins/Experimental/Chooser/Source/ProxyTableEditor/Private/ProxyTableEditorModule.cpp b/Engine/Plugins/Experimental/Chooser/Source/ProxyTableEditor/Private/ProxyTableEditorModule.cpp
--- a/Engine/Plugins/Experimental/Chooser/Source/ProxyTableEditor/Private/ProxyTableEditorModule.cpp
+++ b/Engine/Plugins/Experimental/Chooser/Source/ProxyTableEditor/Private/ProxyTableEditorModule.cpp
@@ -3,6 +3,7 @@
 #include "ProxyTableEditorModule.h"
 #include "IAssetTools.h"
 #include "ProxyTableEditor.h"
+#include "ProxyTableEditorCommands.h"
 
 #define LOCTEXT_NAMESPACE "ChooserEditorModule"
 
@@ -12,6 +13,8 @@ namespace UE::ProxyTableEditor
 void FModule::StartupModule()
 {
 	FProxyTableEditor::RegisterWidgets();
+
+	FProxyTableEditorCommands::Register();
 }
 
 void FModule::ShutdownModule()
diff --git a/Engine/Plugins/Experimental/Chooser/Source/ProxyTableEditor/ProxyTableEditor.Build.cs b/Engine/Plugins/Experimental/Chooser/Source/ProxyTableEditor/ProxyTableEditor.Build.cs
--- a/Engine/Plugins/Experimental/Chooser/Source/ProxyTableEditor/ProxyTableEditor.Build.cs
+++ b/Engine/Plugins/Experimental/Chooser/Source/ProxyTableEditor/ProxyTableEditor.Build.cs
@@ -20,6 +20,7 @@ namespace UnrealBuildTool.Rules
 					"UnrealEd",
 					"EditorWidgets",
 					"ToolWidgets",
+					"ToolMenus",
 					"SlateCore",
 					"Slate",
 					"PropertyEditor",
