Change "Key" in Proxy Table asset to be a Proxy Asset rather than an FName

diff --git a/Engine/Plugins/Experimental/Chooser/Source/Chooser/Public/ChooserPropertyAccess.h b/Engine/Plugins/Experimental/Chooser/Source/Chooser/Public/ChooserPropertyAccess.h
--- a/Engine/Plugins/Experimental/Chooser/Source/Chooser/Public/ChooserPropertyAccess.h
+++ b/Engine/Plugins/Experimental/Chooser/Source/Chooser/Public/ChooserPropertyAccess.h
@@ -3,11 +3,36 @@
 #pragma once
 
 #include "CoreMinimal.h"
+#include "UObject/Interface.h"
+
+#include "ChooserPropertyAccess.generated.h"
 
 #if WITH_EDITOR
 struct FBindingChainElement;
 #endif
 
+UINTERFACE(MinimalAPI)
+class UHasContextClass : public UInterface
+{
+	GENERATED_BODY()
+};
+
+class IHasContextClass
+{
+	GENERATED_BODY()
+public:
+	virtual UClass* GetContextClass() { return nullptr; };
+};
+
+USTRUCT()
+struct FChooserPropertyBinding 
+{
+	GENERATED_BODY()
+	
+	UPROPERTY()
+	TArray<FName> PropertyBindingChain;
+};
+
 namespace UE::Chooser
 {
 	CHOOSER_API bool ResolvePropertyChain(const void*& Container, UStruct*& StructType, const TArray<FName>& PropertyBindingChain);
diff --git a/Engine/Plugins/Experimental/Chooser/Source/ChooserEditor/Private/BoolColumnEditor.cpp b/Engine/Plugins/Experimental/Chooser/Source/ChooserEditor/Private/BoolColumnEditor.cpp
--- a/Engine/Plugins/Experimental/Chooser/Source/ChooserEditor/Private/BoolColumnEditor.cpp
+++ b/Engine/Plugins/Experimental/Chooser/Source/ChooserEditor/Private/BoolColumnEditor.cpp
@@ -56,7 +56,7 @@ TSharedRef<SWidget> CreateOutputBoolColumnWidget(UChooserTable* Chooser, FChoose
 	});
 }
 	
-TSharedRef<SWidget> CreateBoolPropertyWidget(UObject* TransactionObject, void* Value, UClass* ContextClass)
+TSharedRef<SWidget> CreateBoolPropertyWidget(UObject* TransactionObject, void* Value, UClass* ContextClass, UClass* ResultBaseClass)
 {
 	return CreatePropertyWidget<FBoolContextProperty>(TransactionObject, Value, ContextClass, GetDefault<UGraphEditorSettings>()->BooleanPinTypeColor);
 }
diff --git a/Engine/Plugins/Experimental/Chooser/Source/ChooserEditor/Private/ChooserEditorModule.cpp b/Engine/Plugins/Experimental/Chooser/Source/ChooserEditor/Private/ChooserEditorModule.cpp
--- a/Engine/Plugins/Experimental/Chooser/Source/ChooserEditor/Private/ChooserEditorModule.cpp
+++ b/Engine/Plugins/Experimental/Chooser/Source/ChooserEditor/Private/ChooserEditorModule.cpp
@@ -8,6 +8,7 @@
 #include "FloatRangeColumnEditor.h"
 #include "GameplayTagColumnEditor.h"
 #include "ChooserTableEditorCommands.h"
+#include "ChooserPropertyAccess.h"
 
 #define LOCTEXT_NAMESPACE "ChooserEditorModule"
 
@@ -24,6 +25,10 @@ void FModule::StartupModule()
 	
 	FChooserTableEditorCommands::Register();
 
+	FPropertyEditorModule& PropertyModule = FModuleManager::LoadModuleChecked<FPropertyEditorModule>("PropertyEditor");
+	
+	PropertyModule.RegisterCustomPropertyTypeLayout(FChooserPropertyBinding::StaticStruct()->GetFName(), FOnGetPropertyTypeCustomizationInstance::CreateLambda([] { return MakeShared<FPropertyAccessChainCustomization>(); }));
+
 }
 
 void FModule::ShutdownModule()
diff --git a/Engine/Plugins/Experimental/Chooser/Source/ChooserEditor/Private/ChooserEditorModule.h b/Engine/Plugins/Experimental/Chooser/Source/ChooserEditor/Private/ChooserEditorModule.h
--- a/Engine/Plugins/Experimental/Chooser/Source/ChooserEditor/Private/ChooserEditorModule.h
+++ b/Engine/Plugins/Experimental/Chooser/Source/ChooserEditor/Private/ChooserEditorModule.h
@@ -3,6 +3,7 @@
 #pragma once
 
 #include "CoreMinimal.h"
+#include "PropertyAccessChainCustomization.h"
 #include "Modules/ModuleInterface.h"
 
 namespace UE::ChooserEditor
@@ -13,7 +14,6 @@ class FModule : public IModuleInterface
 private:
 	virtual void StartupModule() override;
 	virtual void ShutdownModule() override;
-
 };
 
 }
\ No newline at end of file
diff --git a/Engine/Plugins/Experimental/Chooser/Source/ChooserEditor/Private/ChooserTableEditor.cpp b/Engine/Plugins/Experimental/Chooser/Source/ChooserEditor/Private/ChooserTableEditor.cpp
--- a/Engine/Plugins/Experimental/Chooser/Source/ChooserEditor/Private/ChooserTableEditor.cpp
+++ b/Engine/Plugins/Experimental/Chooser/Source/ChooserEditor/Private/ChooserTableEditor.cpp
@@ -439,13 +439,13 @@ public:
 			}
 			else if (ColumnName == Result) 
 			{
-				TSharedPtr<SWidget> ResultWidget = FObjectChooserWidgetFactories::CreateWidget(Chooser, FObjectChooserBase::StaticStruct(), Chooser->ResultsStructs[RowIndex->RowIndex].GetMutableMemory(), Chooser->ResultsStructs[RowIndex->RowIndex].GetScriptStruct(),Chooser->ContextObjectType,
+				TSharedPtr<SWidget> ResultWidget = FObjectChooserWidgetFactories::CreateWidget(Chooser, FObjectChooserBase::StaticStruct(), Chooser->ResultsStructs[RowIndex->RowIndex].GetMutableMemory(), Chooser->ResultsStructs[RowIndex->RowIndex].GetScriptStruct(),Chooser->ContextObjectType, Chooser->OutputObjectType,
 				FOnStructPicked::CreateLambda([this, RowIndex=RowIndex->RowIndex](const UScriptStruct* ChosenStruct)
 				{
 					const FScopedTransaction Transaction(LOCTEXT("Change Row Result Type", "Change Row Result Type"));
 					Chooser->ResultsStructs[RowIndex].InitializeAs(ChosenStruct);
 					Chooser->Modify(true);
-					FObjectChooserWidgetFactories::CreateWidget(Chooser, FObjectChooserBase::StaticStruct(), Chooser->ResultsStructs[RowIndex].GetMutableMemory(), ChosenStruct, Chooser->ContextObjectType, FOnStructPicked(), &CacheBorder);
+					FObjectChooserWidgetFactories::CreateWidget(Chooser, FObjectChooserBase::StaticStruct(), Chooser->ResultsStructs[RowIndex].GetMutableMemory(), ChosenStruct, Chooser->ContextObjectType, Chooser->OutputObjectType, FOnStructPicked(), &CacheBorder);
 				}),
 				&CacheBorder
 				);
@@ -615,7 +615,7 @@ void FChooserTableEditor::UpdateTableColumns()
 		TSharedPtr<SWidget> HeaderWidget = nullptr;
 		if (FChooserParameterBase* InputValue = Column.GetInputValue())
 		{
-			HeaderWidget = FObjectChooserWidgetFactories::CreateWidget(Chooser, InputValue, Column.GetInputType(), Chooser->ContextObjectType);
+			HeaderWidget = FObjectChooserWidgetFactories::CreateWidget(Chooser, InputValue, Column.GetInputType(), Chooser->ContextObjectType, Chooser->OutputObjectType);
 		}
 
 		HeaderRow->AddColumn(SHeaderRow::FColumn::FArguments()
@@ -1025,16 +1025,14 @@ void FChooserTableEditor::DeleteColumn(int Index)
 	}
 }
 
-TSharedRef<SWidget> CreateAssetWidget(UObject* TransactionObject, void* Value, UClass* ContextClass)
+TSharedRef<SWidget> CreateAssetWidget(UObject* TransactionObject, void* Value, UClass* ContextClass, UClass* ResultBaseClass)
 {
 	FAssetChooser* DIAsset = static_cast<FAssetChooser*>(Value);
 
 	UObject* Asset = DIAsset->Asset;
-
-	UChooserTable* Chooser = Cast<UChooserTable>(TransactionObject);
 	
 	return SNew(SObjectPropertyEntryBox)
-		.AllowedClass((Chooser!=nullptr && Chooser->OutputObjectType!=nullptr) ? Chooser->OutputObjectType.Get() : UObject::StaticClass())
+		.AllowedClass(ResultBaseClass ? ResultBaseClass : UObject::StaticClass())
 		.ObjectPath_Lambda([DIAsset](){ return DIAsset->Asset ? DIAsset->Asset.GetPath() : "";})
 		.OnObjectChanged_Lambda([TransactionObject, DIAsset](const FAssetData& AssetData)
 		{
@@ -1044,16 +1042,14 @@ TSharedRef<SWidget> CreateAssetWidget(UObject* TransactionObject, void* Value, U
 		});
 }
 	
-TSharedRef<SWidget> CreateClassWidget(UObject* TransactionObject, void* Value, UClass* ContextClass)
+TSharedRef<SWidget> CreateClassWidget(UObject* TransactionObject, void* Value, UClass* ContextClas, UClass* ResultBaseClass)
 {
 	FClassChooser* ClassChooser = static_cast<FClassChooser*>(Value);
 
 	UClass* Class = ClassChooser->Class;
 
-	UChooserTable* Chooser = Cast<UChooserTable>(TransactionObject);
-	
 	return SNew(SClassPropertyEntryBox)
-			.MetaClass((Chooser!=nullptr && Chooser->OutputObjectType!=nullptr) ? Chooser->OutputObjectType.Get() : UObject::StaticClass())
+			.MetaClass(ResultBaseClass ? ResultBaseClass : UObject::StaticClass())
 			.SelectedClass_Lambda([ClassChooser]()
 			{
 				return ClassChooser->Class;
@@ -1066,13 +1062,28 @@ TSharedRef<SWidget> CreateClassWidget(UObject* TransactionObject, void* Value, U
 			});
 }
 
-TSharedRef<SWidget> CreateEvaluateChooserWidget(UObject* TransactionObject, void* Value, UClass* ContextObject)
+TSharedRef<SWidget> CreateEvaluateChooserWidget(UObject* TransactionObject, void* Value, UClass* ContextObject, UClass* ResultBaseClass)
 {
 	FEvaluateChooser* EvaluateChooser = static_cast<FEvaluateChooser*>(Value);
 	
 	return SNew(SObjectPropertyEntryBox)
 		.AllowedClass(UChooserTable::StaticClass())
 		.ObjectPath_Lambda([EvaluateChooser](){ return EvaluateChooser->Chooser ? EvaluateChooser->Chooser.GetPath() : "";})
+		.OnShouldFilterAsset_Lambda([ResultBaseClass](const FAssetData& InAssetData)
+		{
+			if (ResultBaseClass == nullptr)
+		 	{
+		 		return false;
+		 	}
+		 	if (InAssetData.IsInstanceOf(UChooserTable::StaticClass()))
+		 	{
+		 		if (UChooserTable* Chooser = Cast<UChooserTable>(InAssetData.GetAsset()))
+		 		{
+		 			return !(Chooser->OutputObjectType && Chooser->OutputObjectType->IsChildOf(ResultBaseClass));
+		 		}
+		 	}
+			return true;
+		})
 		.OnObjectChanged_Lambda([TransactionObject, EvaluateChooser](const FAssetData& AssetData)
 		{
 			const FScopedTransaction Transaction(LOCTEXT("Edit Chooser", "Edit Chooser"));
diff --git a/Engine/Plugins/Experimental/Chooser/Source/ChooserEditor/Private/EnumColumnEditor.cpp b/Engine/Plugins/Experimental/Chooser/Source/ChooserEditor/Private/EnumColumnEditor.cpp
--- a/Engine/Plugins/Experimental/Chooser/Source/ChooserEditor/Private/EnumColumnEditor.cpp
+++ b/Engine/Plugins/Experimental/Chooser/Source/ChooserEditor/Private/EnumColumnEditor.cpp
@@ -119,7 +119,7 @@ namespace UE::ChooserEditor
 
 	
 
-TSharedRef<SWidget> CreateEnumPropertyWidget(UObject* TransactionObject, void* Value, UClass* ContextClass)
+TSharedRef<SWidget> CreateEnumPropertyWidget(UObject* TransactionObject, void* Value, UClass* ContextClass, UClass* ResultBaseClass)
 {
 	return CreatePropertyWidget<FEnumContextProperty>(TransactionObject, Value, ContextClass, GetDefault<UGraphEditorSettings>()->BytePinTypeColor);
 }
diff --git a/Engine/Plugins/Experimental/Chooser/Source/ChooserEditor/Private/FloatRangeColumnEditor.cpp b/Engine/Plugins/Experimental/Chooser/Source/ChooserEditor/Private/FloatRangeColumnEditor.cpp
--- a/Engine/Plugins/Experimental/Chooser/Source/ChooserEditor/Private/FloatRangeColumnEditor.cpp
+++ b/Engine/Plugins/Experimental/Chooser/Source/ChooserEditor/Private/FloatRangeColumnEditor.cpp
@@ -12,7 +12,7 @@
 
 namespace UE::ChooserEditor
 {
-TSharedRef<SWidget> CreateFloatPropertyWidget(UObject* TransactionObject, void* Value, UClass* ContextClass)
+TSharedRef<SWidget> CreateFloatPropertyWidget(UObject* TransactionObject, void* Value, UClass* ContextClass, UClass* ResultBaseClass)
 {
 	return CreatePropertyWidget<FFloatContextProperty>(TransactionObject, Value, ContextClass, GetDefault<UGraphEditorSettings>()->FloatPinTypeColor);
 }
diff --git a/Engine/Plugins/Experimental/Chooser/Source/ChooserEditor/Private/GameplayTagColumnEditor.cpp b/Engine/Plugins/Experimental/Chooser/Source/ChooserEditor/Private/GameplayTagColumnEditor.cpp
--- a/Engine/Plugins/Experimental/Chooser/Source/ChooserEditor/Private/GameplayTagColumnEditor.cpp
+++ b/Engine/Plugins/Experimental/Chooser/Source/ChooserEditor/Private/GameplayTagColumnEditor.cpp
@@ -49,7 +49,7 @@ TSharedRef<SWidget> CreateGameplayTagColumnWidget(UChooserTable* Chooser, FChoos
 	);
 }
 
-TSharedRef<SWidget> CreateGameplayTagPropertyWidget(UObject* TransactionObject, void* Value, UClass* ContextClass)
+TSharedRef<SWidget> CreateGameplayTagPropertyWidget(UObject* TransactionObject, void* Value, UClass* ContextClass, UClass* ResultBaseClass)
 {
 	return CreatePropertyWidget<FGameplayTagContextProperty>(TransactionObject, Value, ContextClass, GetDefault<UGraphEditorSettings>()->StructPinTypeColor);
 }
diff --git a/Engine/Plugins/Experimental/Chooser/Source/ChooserEditor/Private/ObjectChooserWidgetFactories.cpp b/Engine/Plugins/Experimental/Chooser/Source/ChooserEditor/Private/ObjectChooserWidgetFactories.cpp
--- a/Engine/Plugins/Experimental/Chooser/Source/ChooserEditor/Private/ObjectChooserWidgetFactories.cpp
+++ b/Engine/Plugins/Experimental/Chooser/Source/ChooserEditor/Private/ObjectChooserWidgetFactories.cpp
@@ -46,7 +46,7 @@ TSharedPtr<SWidget> FObjectChooserWidgetFactories::CreateColumnWidget(FChooserCo
 }
 	
 	
-TSharedPtr<SWidget> FObjectChooserWidgetFactories::CreateWidget(UObject* TransactionObject, void* Value, const UStruct* ValueType, UClass* ContextClass)
+TSharedPtr<SWidget> FObjectChooserWidgetFactories::CreateWidget(UObject* TransactionObject, void* Value, const UStruct* ValueType, UClass* ContextClass, UClass* ResultBaseClass)
 {
 	if (Value)
 	{
@@ -54,7 +54,7 @@ TSharedPtr<SWidget> FObjectChooserWidgetFactories::CreateWidget(UObject* Transac
 		{
 			if (FChooserWidgetCreator* Creator = ChooserWidgetCreators.Find(ValueType))
 			{
-				return (*Creator)(TransactionObject, Value, ContextClass);
+				return (*Creator)(TransactionObject, Value, ContextClass, ResultBaseClass);
 			}
 			ValueType = ValueType->GetSuperStruct();
 		}
@@ -63,9 +63,9 @@ TSharedPtr<SWidget> FObjectChooserWidgetFactories::CreateWidget(UObject* Transac
 	return nullptr;
 }
 
-TSharedPtr<SWidget> FObjectChooserWidgetFactories::CreateWidget(UObject* TransactionObject, const UScriptStruct* BaseType, void* Value, const UStruct* ValueType, UClass* ContextClass, const FOnStructPicked& CreateClassCallback, TSharedPtr<SBorder>* InnerWidget)
+TSharedPtr<SWidget> FObjectChooserWidgetFactories::CreateWidget(UObject* TransactionObject, const UScriptStruct* BaseType, void* Value, const UStruct* ValueType, UClass* ContextClass, UClass* ResultBaseClass, const FOnStructPicked& CreateClassCallback, TSharedPtr<SBorder>* InnerWidget)
 {
-	TSharedPtr<SWidget> LeftWidget = CreateWidget(TransactionObject, Value, ValueType, ContextClass);
+	TSharedPtr<SWidget> LeftWidget = CreateWidget(TransactionObject, Value, ValueType, ContextClass, ResultBaseClass);
 	
 	if (!LeftWidget.IsValid())
 	{
diff --git a/Engine/Plugins/Experimental/Chooser/Source/ChooserEditor/Private/PropertyAccessChainCustomization.cpp b/Engine/Plugins/Experimental/Chooser/Source/ChooserEditor/Private/PropertyAccessChainCustomization.cpp
new file mode 100644
--- /dev/null
+++ b/Engine/Plugins/Experimental/Chooser/Source/ChooserEditor/Private/PropertyAccessChainCustomization.cpp
@@ -0,0 +1,185 @@
+﻿// Copyright Epic Games, Inc. All Rights Reserved.
+
+#include "PropertyAccessChainCustomization.h"
+
+#include "Chooser.h"
+#include "ChooserPropertyAccess.h"
+#include "PropertyHandle.h"
+#include "DetailWidgetRow.h"
+#include "SClassViewer.h"
+#include "IDetailChildrenBuilder.h"
+#include "IPropertyAccessEditor.h"
+#include "ScopedTransaction.h"
+#include "Features/IModularFeatures.h"
+
+#define LOCTEXT_NAMESPACE "PropertyAccessChainCustomization"
+
+namespace UE::ChooserEditor
+{
+
+void FPropertyAccessChainCustomization::CustomizeHeader(TSharedRef<IPropertyHandle> PropertyHandle, FDetailWidgetRow& HeaderRow, IPropertyTypeCustomizationUtils& CustomizationUtils)
+{
+	FProperty* Property = PropertyHandle->GetProperty();
+
+	UClass* ContextClass = nullptr;
+	static FName TypeMetaData = "BindingType";
+
+	FString TypeFilter = PropertyHandle->GetMetaData(TypeMetaData);
+
+	TArray<UObject*> OuterObjects;
+	PropertyHandle->GetOuterObjects(OuterObjects);
+	UObject* OuterObject = OuterObjects[0];
+	
+	while (OuterObject && !OuterObject->Implements<UHasContextClass>())
+	{
+		OuterObject = OuterObject->GetOuter();
+	}
+
+	if (OuterObject)
+	{
+		IHasContextClass* HasContext = static_cast<IHasContextClass*>(OuterObject->GetInterfaceAddress(UHasContextClass::StaticClass()));
+		ContextClass = HasContext->GetContextClass();
+	}
+	
+	TSharedPtr<SWidget> Widget = SNullWidget::NullWidget;
+
+	FPropertyBindingWidgetArgs Args;
+	Args.bAllowPropertyBindings = true;
+	
+	Args.OnCanBindProperty = FOnCanBindProperty::CreateLambda([TypeFilter](FProperty* Property)
+	{
+		if (TypeFilter == "")
+		{
+			return true;
+		}
+		else
+		{
+			return Property == nullptr || Property->GetCPPType() == TypeFilter;
+		}
+		
+	});
+	
+	Args.OnCanBindToClass = FOnCanBindToClass::CreateLambda([ContextClass ](UClass* InClass)
+	{
+		return true;
+	});
+
+	Args.CurrentBindingColor = MakeAttributeLambda([]() {
+		return FLinearColor::Gray;
+	});
+
+	Args.OnCanAcceptPropertyOrChildren = FOnCanBindProperty::CreateLambda([](FProperty* InProperty)
+		{
+			// Make only blueprint visible properties visible for binding.
+			return InProperty->HasAnyPropertyFlags(CPF_BlueprintVisible);
+		});	
+
+	Args.OnAddBinding = FOnAddBinding::CreateLambda([PropertyHandle](FName InPropertyName, const TArray<FBindingChainElement>& InBindingChain)
+		{
+			TArray<UObject*> OuterObjects;
+			PropertyHandle->GetOuterObjects(OuterObjects);
+			const FScopedTransaction Transaction(NSLOCTEXT("ChooserPropertyBinding", "Change Property Binding", "Change Property Binding"));
+		
+			for (uint32 i=0; i<PropertyHandle->GetNumOuterObjects(); i++)
+			{
+				void* ValuePtr = nullptr;
+				PropertyHandle->GetValueData(ValuePtr);// todo get per object value data?
+				FChooserPropertyBinding* PropertyValue = reinterpret_cast<FChooserPropertyBinding*>(ValuePtr);
+			
+				if (PropertyValue != nullptr)
+				{
+					PropertyHandle->NotifyPreChange();
+					
+					OuterObjects[i]->Modify(true); // todo
+					Chooser::CopyPropertyChain(InBindingChain, PropertyValue->PropertyBindingChain);
+
+					PropertyHandle->NotifyPostChange(EPropertyChangeType::ValueSet);
+				}
+			}
+		});
+
+	Args.CurrentBindingToolTipText = MakeAttributeLambda([PropertyHandle]()
+			{
+				const FText Bind = NSLOCTEXT("ContextPropertyWidget", "Bind", "Bind");
+				FText CurrentValue = Bind;
+				
+				void* ValuePtr = nullptr;
+				PropertyHandle->GetValueData(ValuePtr);// todo get per object value data?
+				FChooserPropertyBinding* PropertyValue = reinterpret_cast<FChooserPropertyBinding*>(ValuePtr);
+
+				if (PropertyValue != nullptr && PropertyValue->PropertyBindingChain.Num()>0)
+				{
+					TArray<FText> BindingChainText;
+					BindingChainText.Reserve(PropertyValue->PropertyBindingChain.Num());
+				 
+					for (const FName& Name : PropertyValue->PropertyBindingChain)
+					{
+						BindingChainText.Add(FText::FromName(Name));
+					}
+					
+					CurrentValue = FText::Join(NSLOCTEXT("ContextPropertyWidget", "PropertyPathSeparator","."), BindingChainText);
+				}
+	
+				return CurrentValue;	
+			});
+	
+	Args.CurrentBindingText = MakeAttributeLambda([PropertyHandle]()
+			{
+				const FText Bind = NSLOCTEXT("ContextPropertyWidget", "Bind", "Bind");
+				FText CurrentValue = Bind;
+		
+				void* ValuePtr = nullptr;
+				PropertyHandle->GetValueData(ValuePtr);// todo get per object value data?
+				FChooserPropertyBinding* PropertyValue = reinterpret_cast<FChooserPropertyBinding*>(ValuePtr);
+	
+				int BindingChainLength = PropertyValue->PropertyBindingChain.Num();
+				if (BindingChainLength > 0)
+				{
+					if (BindingChainLength == 1)
+					{
+						// single property, just use the property name
+						CurrentValue = FText::FromName(PropertyValue->PropertyBindingChain.Last());
+					}
+					else
+					{
+						// for longer chains always show the last struct/object name, and the final property name (full path in tooltip)
+						CurrentValue = FText::Join(NSLOCTEXT("ContextPropertyWidget", "PropertyPathSeparator","."),
+							TArray<FText>({
+								FText::FromName(PropertyValue->PropertyBindingChain[BindingChainLength-2]),
+								FText::FromName(PropertyValue->PropertyBindingChain[BindingChainLength-1])
+							}));
+					}
+				}
+	
+				return CurrentValue;
+			});
+
+	Args.CurrentBindingImage = MakeAttributeLambda([]() -> const FSlateBrush*
+		{
+			static FName PropertyIcon(TEXT("Kismet.Tabs.Variables"));
+			return FAppStyle::GetBrush(PropertyIcon);
+		});
+
+	IPropertyAccessEditor& PropertyAccessEditor = IModularFeatures::Get().GetModularFeature<IPropertyAccessEditor>("PropertyAccessEditor");
+	FBindingContextStruct StructInfo;
+	StructInfo.Struct = ContextClass;
+	Widget =  PropertyAccessEditor.MakePropertyBindingWidget({StructInfo}, Args);
+
+	HeaderRow
+	.NameContent()
+	[
+		PropertyHandle->CreatePropertyNameWidget()
+	]
+	.ValueContent()
+	[
+		Widget.ToSharedRef()
+	];
+}
+	
+void FPropertyAccessChainCustomization::CustomizeChildren(TSharedRef<IPropertyHandle> PropertyHandle, IDetailChildrenBuilder& ChildBuilder, IPropertyTypeCustomizationUtils& CustomizationUtils)
+{
+}
+
+}
+
+#undef LOCTEXT_NAMESPACE
\ No newline at end of file
diff --git a/Engine/Plugins/Experimental/Chooser/Source/ChooserEditor/Private/PropertyAccessChainCustomization.h b/Engine/Plugins/Experimental/Chooser/Source/ChooserEditor/Private/PropertyAccessChainCustomization.h
new file mode 100644
--- /dev/null
+++ b/Engine/Plugins/Experimental/Chooser/Source/ChooserEditor/Private/PropertyAccessChainCustomization.h
@@ -0,0 +1,21 @@
+﻿// Copyright Epic Games, Inc. All Rights Reserved.
+
+#pragma once
+
+#include "IPropertyTypeCustomization.h"
+#include "PropertyEditorModule.h"
+
+namespace UE::ChooserEditor
+{
+
+class FPropertyAccessChainCustomization : public IPropertyTypeCustomization
+{
+public:
+	FPropertyAccessChainCustomization(){}
+private:
+	// IPropertyTypeCustomization interface
+	virtual void CustomizeHeader(TSharedRef<IPropertyHandle> PropertyHandle, FDetailWidgetRow& HeaderRow, IPropertyTypeCustomizationUtils& CustomizationUtils) override;
+	virtual void CustomizeChildren(TSharedRef<IPropertyHandle> PropertyHandle, IDetailChildrenBuilder& ChildBuilder, IPropertyTypeCustomizationUtils& CustomizationUtils) override;
+};
+
+}
\ No newline at end of file
diff --git a/Engine/Plugins/Experimental/Chooser/Source/ChooserEditor/Public/ObjectChooserWidgetFactories.h b/Engine/Plugins/Experimental/Chooser/Source/ChooserEditor/Public/ObjectChooserWidgetFactories.h
--- a/Engine/Plugins/Experimental/Chooser/Source/ChooserEditor/Public/ObjectChooserWidgetFactories.h
+++ b/Engine/Plugins/Experimental/Chooser/Source/ChooserEditor/Public/ObjectChooserWidgetFactories.h
@@ -13,15 +13,15 @@ struct FChooserColumnBase;
 
 namespace UE::ChooserEditor
 {
-	typedef TFunction<TSharedRef<SWidget>(UObject* TransactionObject, void* Value, UClass* ContextClass)> FChooserWidgetCreator;
+	typedef TFunction<TSharedRef<SWidget>(UObject* TransactionObject, void* Value, UClass* ContextClass, UClass* ResultBaseClass)> FChooserWidgetCreator;
 	typedef TFunction<TSharedRef<SWidget>(UChooserTable* Chooser, FChooserColumnBase* Column, int Row)> FColumnWidgetCreator;
 
 	class CHOOSEREDITOR_API FObjectChooserWidgetFactories
 	{
 	public:
 
-		static TSharedPtr<SWidget> CreateWidget(UObject* TransactionObject, void* Value, const UStruct* ValueType, UClass* ContextClass);
-		static TSharedPtr<SWidget> CreateWidget(UObject* TransactionObject, const UScriptStruct* BaseType, void* Value, const UStruct* ValueType, UClass* ContextClass, const FOnStructPicked& CreateClassCallback, TSharedPtr<SBorder>* InnerWidget = nullptr);
+		static TSharedPtr<SWidget> CreateWidget(UObject* TransactionObject, void* Value, const UStruct* ValueType, UClass* ContextClass, UClass* ResultBaseClass);
+		static TSharedPtr<SWidget> CreateWidget(UObject* TransactionObject, const UScriptStruct* BaseType, void* Value, const UStruct* ValueType, UClass* ContextClass, UClass* ResultBaseClass, const FOnStructPicked& CreateClassCallback, TSharedPtr<SBorder>* InnerWidget = nullptr);
 		
 		static TSharedPtr<SWidget> CreateColumnWidget(FChooserColumnBase* Column, const UStruct* ColumnType, UChooserTable* Chooser, int Row);
 		
diff --git a/Engine/Plugins/Experimental/Chooser/Source/ProxyTable/Internal/ProxyTable.h b/Engine/Plugins/Experimental/Chooser/Source/ProxyTable/Internal/ProxyTable.h
--- a/Engine/Plugins/Experimental/Chooser/Source/ProxyTable/Internal/ProxyTable.h
+++ b/Engine/Plugins/Experimental/Chooser/Source/ProxyTable/Internal/ProxyTable.h
@@ -9,20 +9,39 @@
 #include "InstancedStruct.h"
 #include "ProxyTable.generated.h"
 
+UCLASS(MinimalAPI,BlueprintType)
+class UProxyAsset : public UObject, public IHasContextClass
+{
+	GENERATED_UCLASS_BODY()
+public:
+	UProxyAsset() {}
+
+	UPROPERTY(EditAnywhere, Category = "Proxy", Meta = (AllowAbstract=true))
+	TObjectPtr<UClass> Type;
+	
+	UPROPERTY(EditAnywhere, Category = "Proxy Table Reference")
+	TObjectPtr<UClass> ContextClass;
+	
+	UPROPERTY(EditAnywhere, Meta = (ExcludeBaseStruct, BaseStruct ="/Script/ProxyTable.ChooserParameterProxyTableBase"), Category = "Proxy Table Reference")
+	FInstancedStruct ProxyTable;
+
+	virtual UClass* GetContextClass() override { return ContextClass; }
+};
 
 USTRUCT()
 struct FProxyEntry
 {
 	GENERATED_BODY()
+	
+	UPROPERTY(EditAnywhere, Category = "Data")
+	TObjectPtr<UProxyAsset> Proxy;
+
+	// temporarily leaving this property for backwards compatibility with old content which used FNames rather than UProxyAsset
 	UPROPERTY(EditAnywhere, Category = "Data")
 	FName Key;
 	
-	UPROPERTY()
-	TScriptInterface<IObjectChooser> Value;
-	
 	UPROPERTY(DisplayName="Value", EditAnywhere, Meta = (ExcludeBaseStruct, BaseStruct = "/Script/Chooser.ObjectChooserBase"), Category = "Data")
 	FInstancedStruct ValueStruct;
-	
 };
 
 UCLASS(MinimalAPI,BlueprintType)
@@ -48,10 +67,10 @@ struct PROXYTABLE_API FProxyTableContextProperty :  public FChooserParameterProx
 {
 	GENERATED_BODY()
 public:
-
-	UPROPERTY()
-	TArray<FName> PropertyBindingChain;
 	
+	UPROPERTY(EditAnywhere, Meta = (BindingType = "UProxyTable*"), Category = "Binding")
+	FChooserPropertyBinding Binding;
+
 	virtual bool GetValue(const UObject* ContextObject, const UProxyTable*& OutResult) const override;
 
 #if WITH_EDITOR
@@ -62,7 +81,7 @@ public:
 
 	void SetBinding(const TArray<FBindingChainElement>& InBindingChain)
 	{
-		UE::Chooser::CopyPropertyChain(InBindingChain, PropertyBindingChain);
+		UE::Chooser::CopyPropertyChain(InBindingChain, Binding.PropertyBindingChain);
 	}
 #endif
 };
@@ -76,52 +95,6 @@ struct PROXYTABLE_API FLookupProxy : public FObjectChooserBase
 	public:
 	FLookupProxy();
 	
-	UPROPERTY(EditAnywhere, Meta = (ExcludeBaseStruct, BaseStruct ="/Script/ProxyTable.ChooserParameterProxyTableBase"), Category = "Input")
-	FInstancedStruct ProxyTable;
-	
 	UPROPERTY(EditAnywhere, Category="Parameters")
-	FName Key;
-};
-
-// deprecated classes for upgrading old data
-
-UCLASS(ClassGroup = "LiveLink", deprecated)
-class PROXYTABLE_API UDEPRECATED_ChooserParameterProxyTable_ContextProperty :  public UObject, public IChooserParameterProxyTable
-{
-	GENERATED_BODY()
-public:
-
-	UPROPERTY()
-	TArray<FName> PropertyBindingChain;
-
-	virtual void ConvertToInstancedStruct(FInstancedStruct& OutInstancedStruct) const override
-	{
-		OutInstancedStruct.InitializeAs(FProxyTableContextProperty::StaticStruct());
-		FProxyTableContextProperty& Property = OutInstancedStruct.GetMutable<FProxyTableContextProperty>();
-		Property.PropertyBindingChain = PropertyBindingChain;
-	}
-};
-
-UCLASS(ClassGroup = "LiveLink", deprecated)
-class PROXYTABLE_API UDEPRECATED_ObjectChooser_LookupProxy : public UObject, public IObjectChooser
-{
-	GENERATED_BODY()
-	
-	UPROPERTY(EditAnywhere, Category = "Input")
-	TScriptInterface<IChooserParameterProxyTable> ProxyTable;
-	
-	UPROPERTY(EditAnywhere, Category="Parameters")
-	FName Key;
-	
-	virtual void ConvertToInstancedStruct(FInstancedStruct& OutInstancedStruct) const override
-   	{
-   		OutInstancedStruct.InitializeAs(FLookupProxy::StaticStruct());
-   		FLookupProxy& AssetChooser = OutInstancedStruct.GetMutable<FLookupProxy>();
-		
-		if (IChooserParameterProxyTable* ProxyTableInterface = ProxyTable.GetInterface())
-		{
-			ProxyTableInterface->ConvertToInstancedStruct(AssetChooser.ProxyTable);
-		}
-		AssetChooser.Key = Key;
-   	}
+	TObjectPtr<UProxyAsset> Proxy;
 };
\ No newline at end of file
diff --git a/Engine/Plugins/Experimental/Chooser/Source/ProxyTable/Private/ProxyTable.cpp b/Engine/Plugins/Experimental/Chooser/Source/ProxyTable/Private/ProxyTable.cpp
--- a/Engine/Plugins/Experimental/Chooser/Source/ProxyTable/Private/ProxyTable.cpp
+++ b/Engine/Plugins/Experimental/Chooser/Source/ProxyTable/Private/ProxyTable.cpp
@@ -4,16 +4,15 @@
 
 FLookupProxy::FLookupProxy()
 {
-	ProxyTable.InitializeAs(FProxyTableContextProperty::StaticStruct());
 }
 
-static UObject* FindProxyObject(const UProxyTable* Table, FName Key, const UObject* ContextObject)
+static UObject* FindProxyObject(const UProxyTable* Table, const UProxyAsset* Proxy, const UObject* ContextObject)
 {
 	if (Table)
 	{
 		for (const FProxyEntry& Entry : Table->Entries)
 		{
-			if (Entry.Key == Key && Entry.ValueStruct.IsValid())
+			if (Entry.Proxy == Proxy && Entry.ValueStruct.IsValid())
 			{
 				const FObjectChooserBase& EntryValue = Entry.ValueStruct.Get<FObjectChooserBase>();
 				return EntryValue.ChooseObject(ContextObject);
@@ -23,7 +22,7 @@ static UObject* FindProxyObject(const UProxyTable* Table, FName Key, const UObje
 		// search parent tables (uncooked data only)
 		for (const TObjectPtr<UProxyTable> ParentTable : Table->InheritEntriesFrom)
 		{
-			if (UObject* Value = FindProxyObject(ParentTable, Key, ContextObject))
+			if (UObject* Value = FindProxyObject(ParentTable, Proxy, ContextObject))
 			{
 				return Value;
 			}
@@ -33,14 +32,14 @@ static UObject* FindProxyObject(const UProxyTable* Table, FName Key, const UObje
 	return nullptr;
 }
 
-UObject* FLookupProxy::ChooseObject(const UObject* ContextObject) const
+static UObject* FindProxyObject(const UProxyAsset* Proxy, const UObject* ContextObject)
 {
-	if (ProxyTable.IsValid())
+	if (Proxy && Proxy->ProxyTable.IsValid())
 	{
 		const UProxyTable* Table;
-		if (ProxyTable.Get<FChooserParameterProxyTableBase>().GetValue(ContextObject, Table))
+		if (Proxy->ProxyTable.Get<FChooserParameterProxyTableBase>().GetValue(ContextObject, Table))
 		{
-			if (UObject* Value = FindProxyObject(Table, Key, ContextObject))
+			if (UObject* Value = FindProxyObject(Table, Proxy, ContextObject))
 			{
 				return Value;
 			}
@@ -50,20 +49,31 @@ UObject* FLookupProxy::ChooseObject(const UObject* ContextObject) const
 	return nullptr;
 }
 
+UObject* FLookupProxy::ChooseObject(const UObject* ContextObject) const
+{
+	return FindProxyObject(Proxy, ContextObject);
+}
+
 UProxyTable::UProxyTable(const FObjectInitializer& Initializer)
 	:Super(Initializer)
 {
 
 }
 
+UProxyAsset::UProxyAsset(const FObjectInitializer& Initializer)
+	: Super(Initializer)
+{
+	ProxyTable.InitializeAs(FProxyTableContextProperty::StaticStruct());
+}
+
 bool FProxyTableContextProperty::GetValue(const UObject* ContextObject, const UProxyTable*& OutResult) const
 {
 	UStruct* StructType = ContextObject->GetClass();
 	const void* Container = ContextObject;
 	
-	if (UE::Chooser::ResolvePropertyChain(Container, StructType, PropertyBindingChain))
+	if (UE::Chooser::ResolvePropertyChain(Container, StructType, Binding.PropertyBindingChain))
 	{
-		if (const FObjectProperty* Property = FindFProperty<FObjectProperty>(StructType, PropertyBindingChain.Last()))
+		if (const FObjectProperty* Property = FindFProperty<FObjectProperty>(StructType, Binding.PropertyBindingChain.Last()))
 		{
 			OutResult = *Property->ContainerPtrToValuePtr<UProxyTable*>(Container);
 			return true;
@@ -73,7 +83,40 @@ bool FProxyTableContextProperty::GetValue(const UObject* ContextObject, const UP
 	return false;
 }
 
+UObject* UProxyTableFunctionLibrary::EvaluateProxyAsset(const UObject* ContextObject, const UProxyAsset* Proxy)
+{
+	return FindProxyObject(Proxy, ContextObject);
+}
+
+// fallback for FName based Keys:
+
+static UObject* FindProxyObject(const UProxyTable* Table, FName Key, const UObject* ContextObject)
+{
+	if (Table)
+	{
+		for (const FProxyEntry& Entry : Table->Entries)
+		{
+			if (Entry.Key == Key && Entry.ValueStruct.IsValid())
+			{
+				const FObjectChooserBase& EntryValue = Entry.ValueStruct.Get<FObjectChooserBase>();
+				return EntryValue.ChooseObject(ContextObject);
+			}
+		}
+
+		// search parent tables (uncooked data only)
+		for (const TObjectPtr<UProxyTable> ParentTable : Table->InheritEntriesFrom)
+		{
+			if (UObject* Value = FindProxyObject(ParentTable, Key, ContextObject))
+			{
+				return Value;
+			}
+		}
+	}
+
+	return nullptr;
+}
+
 UObject* UProxyTableFunctionLibrary::EvaluateProxyTable(const UObject* ContextObject, const UProxyTable* ProxyTable, FName Key)
 {
 	return FindProxyObject(ProxyTable, Key, ContextObject);
-}
\ No newline at end of file
+}
diff --git a/Engine/Plugins/Experimental/Chooser/Source/ProxyTable/Public/ProxyTableFunctionLibrary.h b/Engine/Plugins/Experimental/Chooser/Source/ProxyTable/Public/ProxyTableFunctionLibrary.h
--- a/Engine/Plugins/Experimental/Chooser/Source/ProxyTable/Public/ProxyTableFunctionLibrary.h
+++ b/Engine/Plugins/Experimental/Chooser/Source/ProxyTable/Public/ProxyTableFunctionLibrary.h
@@ -16,13 +16,22 @@ class PROXYTABLE_API UProxyTableFunctionLibrary : public UBlueprintFunctionLibra
 
 public:
 	/**
-	* Evaluate a proxy table and return the selected UObject, or null
-	*
-	* @param ContextObject			(in) An Object from which the parameters to the Chooser Table will be read
-	* @param ProxyTable				(in) The ProxyTable asset
-	* @param Key					(in) The Key from the ProxyTable asset
+		* Temporary backwards compatibility function!  please switch to EvaluateProxyAsset
+		*
+		* @param ContextObject			(in) An Object from which the parameters to the Chooser Table will be read
+		* @param ProxyTable				(in) The ProxyTable asset
+		* @param Key					(in) The Key from the ProxyTable asset
 	*/
 	UFUNCTION(BlueprintPure, meta = (BlueprintThreadSafe), Category = "Animation")
 	static UObject* EvaluateProxyTable(const UObject* ContextObject, const UProxyTable* ProxyTable, FName Key);
+	
+	/**
+	* Resolve a proxy asset and return the selected UObject, or null
+	*
+	* @param ContextObject			(in) An Object from which the Proxy Table will be read, and parameters to any nested Chooser Tables that need to evaluate
+	* @param ProxyAsset				(in) The ProxyAsset asset
+	*/
+	UFUNCTION(BlueprintPure, meta = (BlueprintThreadSafe), Category = "Animation")
+	static UObject* EvaluateProxyAsset(const UObject* ContextObject, const UProxyAsset* Proxy);
 
 };
\ No newline at end of file
diff --git a/Engine/Plugins/Experimental/Chooser/Source/ProxyTableEditor/Private/AssetDefinition_ProxyAsset.cpp b/Engine/Plugins/Experimental/Chooser/Source/ProxyTableEditor/Private/AssetDefinition_ProxyAsset.cpp
new file mode 100644
--- /dev/null
+++ b/Engine/Plugins/Experimental/Chooser/Source/ProxyTableEditor/Private/AssetDefinition_ProxyAsset.cpp
@@ -0,0 +1,16 @@
+// Copyright Epic Games, Inc. All Rights Reserved.
+
+#include "AssetDefinition_ProxyAsset.h"
+#include "Toolkits/SimpleAssetEditor.h"
+
+#define LOCTEXT_NAMESPACE "UAssetDefinition_ProxyAsset" 
+
+EAssetCommandResult UAssetDefinition_ProxyAsset::OpenAssets(const FAssetOpenArgs& OpenArgs) const
+{
+	TArray<UObject*> Objects = OpenArgs.LoadObjects<UObject>();
+	FSimpleAssetEditor::CreateEditor(OpenArgs.GetToolkitMode(), OpenArgs.ToolkitHost, Objects);
+
+	return EAssetCommandResult::Handled;
+}
+
+#undef LOCTEXT_NAMESPACE
\ No newline at end of file
diff --git a/Engine/Plugins/Experimental/Chooser/Source/ProxyTableEditor/Private/AssetDefinition_ProxyAsset.h b/Engine/Plugins/Experimental/Chooser/Source/ProxyTableEditor/Private/AssetDefinition_ProxyAsset.h
new file mode 100644
--- /dev/null
+++ b/Engine/Plugins/Experimental/Chooser/Source/ProxyTableEditor/Private/AssetDefinition_ProxyAsset.h
@@ -0,0 +1,27 @@
+// Copyright Epic Games, Inc. All Rights Reserved.
+
+#pragma once
+
+#include "ProxyTable.h"
+#include "AssetDefinitionDefault.h"
+
+#include "AssetDefinition_ProxyAsset.generated.h"
+
+UCLASS()
+class UAssetDefinition_ProxyAsset : public UAssetDefinitionDefault
+{
+	GENERATED_BODY()
+
+public:
+	// UAssetDefinition Begin
+	virtual FText GetAssetDisplayName() const override { return NSLOCTEXT("AssetTypeActions", "AssetTypeActions_ProxyAsset", "Proxy Asset"); }
+	virtual FLinearColor GetAssetColor() const override { return FLinearColor(FColor(128,128,64)); }
+	virtual TSoftClassPtr<UObject> GetAssetClass() const override { return UProxyAsset::StaticClass(); }
+	virtual TConstArrayView<FAssetCategoryPath> GetAssetCategories() const override
+	{
+		static const auto Categories = { FAssetCategoryPath(EAssetCategoryPaths::Misc) };
+		return Categories;
+	}
+	virtual EAssetCommandResult OpenAssets(const FAssetOpenArgs& OpenArgs) const override;
+	// UAssetDefinition End
+};
diff --git a/Engine/Plugins/Experimental/Chooser/Source/ProxyTableEditor/Private/ProxyAssetFactory.cpp b/Engine/Plugins/Experimental/Chooser/Source/ProxyTableEditor/Private/ProxyAssetFactory.cpp
new file mode 100644
--- /dev/null
+++ b/Engine/Plugins/Experimental/Chooser/Source/ProxyTableEditor/Private/ProxyAssetFactory.cpp
@@ -0,0 +1,21 @@
+﻿// Copyright Epic Games, Inc. All Rights Reserved.
+
+#include "ProxyAssetFactory.h"
+#include "ProxyTable.h"
+
+UProxyAssetFactory::UProxyAssetFactory()
+{
+	bCreateNew = true;
+	bEditAfterNew = true;
+	SupportedClass = UProxyAsset::StaticClass();
+}
+
+bool UProxyAssetFactory::ConfigureProperties()
+{
+	return true;
+}
+
+UObject* UProxyAssetFactory::FactoryCreateNew(UClass* Class, UObject* InParent, FName Name, EObjectFlags Flags, UObject* Context, FFeedbackContext* Warn, FName CallingContext)
+{
+	return NewObject<UProxyAsset>(InParent, Class, Name, Flags | RF_Public | RF_Standalone | RF_Transactional | RF_LoadCompleted);
+}
\ No newline at end of file
diff --git a/Engine/Plugins/Experimental/Chooser/Source/ProxyTableEditor/Private/ProxyAssetFactory.h b/Engine/Plugins/Experimental/Chooser/Source/ProxyTableEditor/Private/ProxyAssetFactory.h
new file mode 100644
--- /dev/null
+++ b/Engine/Plugins/Experimental/Chooser/Source/ProxyTableEditor/Private/ProxyAssetFactory.h
@@ -0,0 +1,22 @@
+﻿// Copyright Epic Games, Inc. All Rights Reserved.
+
+#pragma once
+
+#include "CoreMinimal.h"
+#include "Factories/Factory.h"
+#include "ProxyAssetFactory.generated.h"
+
+UCLASS(MinimalAPI)
+class UProxyAssetFactory : public UFactory
+{
+	GENERATED_BODY()
+
+	UProxyAssetFactory();
+	
+	virtual bool ConfigureProperties() override;
+	virtual UObject* FactoryCreateNew(UClass* Class, UObject* InParent, FName Name, EObjectFlags Flags, UObject* Context, FFeedbackContext* Warn, FName CallingContext) override;
+	virtual UObject* FactoryCreateNew(UClass* Class, UObject* InParent, FName Name, EObjectFlags Flags, UObject* Context, FFeedbackContext* Warn) override
+	{
+		return FactoryCreateNew(Class, InParent, Name, Flags, Context, Warn, NAME_None);
+	}
+};
\ No newline at end of file
diff --git a/Engine/Plugins/Experimental/Chooser/Source/ProxyTableEditor/Private/ProxyTableEditor.cpp b/Engine/Plugins/Experimental/Chooser/Source/ProxyTableEditor/Private/ProxyTableEditor.cpp
--- a/Engine/Plugins/Experimental/Chooser/Source/ProxyTableEditor/Private/ProxyTableEditor.cpp
+++ b/Engine/Plugins/Experimental/Chooser/Source/ProxyTableEditor/Private/ProxyTableEditor.cpp
@@ -31,6 +31,7 @@
 #include "ContextPropertyWidget.h"
 #include "IObjectChooser.h"
 #include "Widgets/Layout/SSeparator.h"
+#include "PropertyCustomizationHelpers.h"
 #include "ProxyTableEditorCommands.h"
 
 #define LOCTEXT_NAMESPACE "ProxyTableEditor"
@@ -335,7 +336,6 @@ public:
 		Operation->Editor = InEditor;
 		Operation->RowIndex = InRowIndex;
 		Operation->DefaultHoverText = LOCTEXT("Proxy Row", "Proxy Row");
-		// UE::ChooserEditor::FObjectChooserWidgetFactories::ConvertToText(InEditor->GetProxyTable()->Entries[InRowIndex].Value.GetObject(), Operation->DefaultHoverText);
 		Operation->CurrentHoverText = Operation->DefaultHoverText;
 			
 		Operation->Construct();
@@ -425,10 +425,18 @@ public:
 			}
 			else if (ColumnName == Value) 
 			{
+				UClass* ContextClass = nullptr;
+				UClass* ObjectType = nullptr;
+				if (ProxyTable->Entries[RowIndex->RowIndex].Proxy)
+				{
+					ContextClass = ProxyTable->Entries[RowIndex->RowIndex].Proxy->ContextClass;
+					ObjectType = ProxyTable->Entries[RowIndex->RowIndex].Proxy->Type;
+				}
+				
 				TSharedPtr<SWidget> ResultWidget = ChooserEditor::FObjectChooserWidgetFactories::CreateWidget(ProxyTable, FObjectChooserBase::StaticStruct(),
 					ProxyTable->Entries[RowIndex->RowIndex].ValueStruct.GetMutableMemory(),
 					ProxyTable->Entries[RowIndex->RowIndex].ValueStruct.GetScriptStruct(),
-					nullptr/*ProxyTable->ContextObjectType*/,
+					ContextClass, ObjectType,
 					FOnStructPicked::CreateLambda([this, RowIndex=RowIndex->RowIndex](const UScriptStruct* ChosenStruct)
 					{
 						const FScopedTransaction Transaction(LOCTEXT("Change Value Type", "Change Value Type"));
@@ -437,7 +445,9 @@ public:
 						ChooserEditor::FObjectChooserWidgetFactories::CreateWidget(ProxyTable, FObjectChooserBase::StaticStruct(),
 								ProxyTable->Entries[RowIndex].ValueStruct.GetMutableMemory(),
 								ProxyTable->Entries[RowIndex].ValueStruct.GetScriptStruct(),
-								nullptr/*ProxyTable->ContextObjectType*/, FOnStructPicked(), &CacheBorder);
+								ProxyTable->Entries[RowIndex].Proxy->ContextClass,
+								ProxyTable->Entries[RowIndex].Proxy->Type,
+								FOnStructPicked(), &CacheBorder);
 					}),
 					&CacheBorder
 					);
@@ -460,13 +470,23 @@ public:
 			}
 			else if (ColumnName == Key)
 			{
-				return SNew(SEditableTextBox)
-					.Text_Lambda([this](){ return ProxyTable->Entries.Num() > RowIndex->RowIndex ?  FText::FromName(ProxyTable->Entries[RowIndex->RowIndex].Key) : FText();})
-					.OnTextCommitted_Lambda([this](const FText& Text, ETextCommit::Type CommitType)
+				return SNew(SObjectPropertyEntryBox)
+					.AllowedClass(UProxyAsset::StaticClass())
+					.ObjectPath_Lambda([this]()
+					{
+						return ProxyTable->Entries.Num() > RowIndex->RowIndex ?  ProxyTable->Entries[RowIndex->RowIndex].Proxy.GetPath() : FString();
+					})
+					.OnObjectChanged_Lambda([this](const FAssetData& AssetData)
 					{
 						if (ProxyTable->Entries.Num() > RowIndex->RowIndex)
 						{
-							ProxyTable->Entries[RowIndex->RowIndex].Key = FName(Text.ToString());
+							const FScopedTransaction Transaction(LOCTEXT("Edit Proxy Asset", "Edit Proxy Asset"));
+							ProxyTable->Modify(true);
+							ProxyTable->Entries[RowIndex->RowIndex].Proxy = Cast<UProxyAsset>(AssetData.GetAsset());
+							
+							// ideally just need to rebuild the widget for the "Value" to update the UObject type filtering.
+							// For now just trigger a full refresh
+							Editor->UpdateTableRows();
 						}
 					});
 			}
@@ -573,6 +593,7 @@ void FProxyTableEditor::MoveRow(int SourceRowIndex, int TargetRowIndex)
 	Table->Entries.Insert(Entry, TargetRowIndex);
 
 	UpdateTableRows();
+	
 }
 
 void FProxyTableEditor::UpdateTableColumns()
@@ -583,7 +604,7 @@ void FProxyTableEditor::UpdateTableColumns()
 					.ManualWidth(30));
 	
 	HeaderRow->AddColumn(SHeaderRow::Column("Key")
-					.DefaultLabel(LOCTEXT("KeyColumnName", "Key"))
+					.DefaultLabel(LOCTEXT("KeyColumnName", "Proxy"))
 					.ManualWidth(500));
 
 	HeaderRow->AddColumn(SHeaderRow::Column("Value")
@@ -760,26 +781,34 @@ TSharedRef<FProxyTableEditor> FProxyTableEditor::CreateEditor( const EToolkitMod
 
 /// Result widgets
 
-TSharedRef<SWidget> CreateLookupProxyWidget(UObject* TransactionObject, void* Value, UClass* ContextObject)
+TSharedRef<SWidget> CreateLookupProxyWidget(UObject* TransactionObject, void* Value, UClass* ContextObject, UClass* ResultBaseClass)
 {
 	FLookupProxy* LookupProxy = static_cast<FLookupProxy*>(Value);
 	
-	return SNew(SHorizontalBox)
-		+ SHorizontalBox::Slot()
-		[
-			UE::ChooserEditor::CreatePropertyWidget<FProxyTableContextProperty>(TransactionObject, LookupProxy->ProxyTable.GetMutablePtr<FProxyTableContextProperty>(), ContextObject, GetDefault<UGraphEditorSettings>()->ObjectPinTypeColor)
-		]
-		+ SHorizontalBox::Slot()
-		[
-			SNew(SEditableTextBox)
-			.Text_Lambda([LookupProxy]() { return FText::FromName(LookupProxy->Key);})
-			.OnTextChanged_Lambda([TransactionObject,LookupProxy](const FText& NewText)
+	return SNew(SObjectPropertyEntryBox)
+		.AllowedClass(UProxyAsset::StaticClass())
+		.ObjectPath_Lambda([LookupProxy](){ return LookupProxy->Proxy.GetPath();})
+		.OnShouldFilterAsset_Lambda([ResultBaseClass](const FAssetData& AssetData)
+		{
+			if (ResultBaseClass == nullptr)
 			{
-				FScopedTransaction ScopedTransaction(LOCTEXT("Change LookupProxy Key Name", "Change LookupProxy Key Name"));
-				TransactionObject->Modify(true);
-				LookupProxy->Key = FName(NewText.ToString());
-			})
-		];
+				return false;
+			}
+			if (AssetData.IsInstanceOf(UProxyAsset::StaticClass()))
+			{
+				if (UProxyAsset* Proxy = Cast<UProxyAsset>(AssetData.GetAsset()))
+				{
+					return !(Proxy->Type && Proxy->Type->IsChildOf(ResultBaseClass));
+				}
+			}
+			return true;
+		})
+		.OnObjectChanged_Lambda([TransactionObject, LookupProxy](const FAssetData& AssetData)
+		{
+			const FScopedTransaction Transaction(LOCTEXT("Edit Chooser", "Edit Chooser"));
+			TransactionObject->Modify(true);
+			LookupProxy->Proxy = Cast<UProxyAsset>(AssetData.GetAsset());
+		});
 }
 
 void FProxyTableEditor::RegisterWidgets()
diff --git a/Engine/Plugins/Experimental/Chooser/Source/ProxyTableEditor/Private/ProxyTableEditorModule.cpp b/Engine/Plugins/Experimental/Chooser/Source/ProxyTableEditor/Private/ProxyTableEditorModule.cpp
--- a/Engine/Plugins/Experimental/Chooser/Source/ProxyTableEditor/Private/ProxyTableEditorModule.cpp
+++ b/Engine/Plugins/Experimental/Chooser/Source/ProxyTableEditor/Private/ProxyTableEditorModule.cpp
@@ -1,11 +1,10 @@
 // Copyright Epic Games, Inc. All Rights Reserved.
 
 #include "ProxyTableEditorModule.h"
-#include "IAssetTools.h"
 #include "ProxyTableEditor.h"
 #include "ProxyTableEditorCommands.h"
 
-#define LOCTEXT_NAMESPACE "ChooserEditorModule"
+#define LOCTEXT_NAMESPACE "ProxyTableEditorModule"
 
 namespace UE::ProxyTableEditor
 {
