Allow binding chooser column parameters in details panel, and fix issues with Enum column refreshing ui when enum changes

diff --git a/Engine/Plugins/Experimental/Chooser/Source/Chooser/Internal/BoolColumn.h b/Engine/Plugins/Experimental/Chooser/Source/Chooser/Internal/BoolColumn.h
--- a/Engine/Plugins/Experimental/Chooser/Source/Chooser/Internal/BoolColumn.h
+++ b/Engine/Plugins/Experimental/Chooser/Source/Chooser/Internal/BoolColumn.h
@@ -15,12 +15,26 @@ struct CHOOSER_API FBoolContextProperty :  public FChooserParameterBoolBase
 public:
 
 	UPROPERTY()
-	TArray<FName> PropertyBindingChain;
+	TArray<FName> PropertyBindingChain_DEPRECATED;
+	
+	UPROPERTY(EditAnywhere, Meta = (BindingType = "bool", BindingColor = "BooleanPinTypeColor"), Category = "Binding")
+	FChooserPropertyBinding Binding;
 	
 	virtual bool GetValue(const UObject* ContextObject, bool& OutResult) const override;
 	virtual bool SetValue(UObject* ContextObject, bool InValue) const override;
 
+	virtual void PostLoad() override
+	{
+		if (PropertyBindingChain_DEPRECATED.Num() > 0)
+		{
+			Binding.PropertyBindingChain = PropertyBindingChain_DEPRECATED;
+			PropertyBindingChain_DEPRECATED.SetNum(0);
+		}
+	}
+	
+
 #if WITH_EDITOR
+	
 	static bool CanBind(const FProperty& Property)
 	{
 		static FString BoolTypeName = "bool";
@@ -29,14 +43,14 @@ public:
 
 	void SetBinding(const TArray<FBindingChainElement>& InBindingChain)
 	{
-		UE::Chooser::CopyPropertyChain(InBindingChain, PropertyBindingChain);
+		UE::Chooser::CopyPropertyChain(InBindingChain, Binding.PropertyBindingChain);
 	}
 
 	virtual void GetDisplayName(FText& OutName) const override
 	{
-		if (!PropertyBindingChain.IsEmpty())
+		if (!Binding.PropertyBindingChain.IsEmpty())
 		{
-			OutName = FText::FromName(PropertyBindingChain.Last());
+			OutName = FText::FromName(Binding.PropertyBindingChain.Last());
 		}
 	}
 #endif
@@ -49,14 +63,22 @@ struct CHOOSER_API FBoolColumn : public FChooserColumnBase
 	public:
 	FBoolColumn();
 	
-	UPROPERTY(EditAnywhere, Meta = (ExcludeBaseStruct, BaseStruct = "/Script/Chooser.ChooserParameterBoolBase"), Category = "Hidden")
+	UPROPERTY(EditAnywhere, Meta = (ExcludeBaseStruct, BaseStruct = "/Script/Chooser.ChooserParameterBoolBase"), Category = "Data")
 	FInstancedStruct InputValue;
 	
-	UPROPERTY(EditAnywhere, Category=Runtime);
+	UPROPERTY(EditAnywhere, Category= "Data");
 	TArray<bool> RowValues; 
 	
 	virtual void Filter(const UObject* ContextObject, const TArray<uint32>& IndexListIn, TArray<uint32>& IndexListOut) const override;
 
+	virtual void PostLoad() override
+	{
+		if (InputValue.IsValid())
+		{
+			InputValue.GetMutable<FChooserParameterBase>().PostLoad();
+		}
+	}
+
 	CHOOSER_COLUMN_BOILERPLATE(FChooserParameterBoolBase);
 };
 
@@ -74,7 +96,7 @@ public:
 	{
 		OutInstancedStruct.InitializeAs(FBoolContextProperty::StaticStruct());
 		FBoolContextProperty& Property = OutInstancedStruct.GetMutable<FBoolContextProperty>();
-		Property.PropertyBindingChain = PropertyBindingChain;
+		Property.Binding.PropertyBindingChain = PropertyBindingChain;
 	}
 };
 
diff --git a/Engine/Plugins/Experimental/Chooser/Source/Chooser/Internal/Chooser.h b/Engine/Plugins/Experimental/Chooser/Source/Chooser/Internal/Chooser.h
--- a/Engine/Plugins/Experimental/Chooser/Source/Chooser/Internal/Chooser.h
+++ b/Engine/Plugins/Experimental/Chooser/Source/Chooser/Internal/Chooser.h
@@ -6,12 +6,13 @@
 #include "IObjectChooser.h"
 #include "InstancedStruct.h"
 #include "IChooserColumn.h"
+#include "ChooserPropertyAccess.h"
 
 #include "Chooser.generated.h"
 
 
 UCLASS(MinimalAPI, BlueprintType)
-class UChooserTable : public UObject
+class UChooserTable : public UObject, public IHasContextClass
 {
 	GENERATED_UCLASS_BODY()
 public:
@@ -42,6 +43,8 @@ public:
 #if WITH_EDITOR
 	virtual void PostLoad() override;
 #endif
+
+	virtual UClass* GetContextClass() override { return ContextObjectType; }
 };
 
 
diff --git a/Engine/Plugins/Experimental/Chooser/Source/Chooser/Internal/EnumColumn.h b/Engine/Plugins/Experimental/Chooser/Source/Chooser/Internal/EnumColumn.h
--- a/Engine/Plugins/Experimental/Chooser/Source/Chooser/Internal/EnumColumn.h
+++ b/Engine/Plugins/Experimental/Chooser/Source/Chooser/Internal/EnumColumn.h
@@ -5,6 +5,7 @@
 #include "IChooserColumn.h"
 #include "IChooserParameterEnum.h"
 #include "InstancedStruct.h"
+#include "ChooserPropertyAccess.h"
 #include "EnumColumn.generated.h"
 
 struct FBindingChainElement;
@@ -15,10 +16,26 @@ struct CHOOSER_API FEnumContextProperty : public FChooserParameterEnumBase
 	GENERATED_BODY()
 	
 	UPROPERTY()
-	TArray<FName> PropertyBindingChain;
+	TArray<FName> PropertyBindingChain_DEPRECATED;
+	
+	UPROPERTY(EditAnywhere, Meta = (BindingType = "enum", BindingColor = "BytePinTypeColor"), Category = "Binding")
+	FChooserEnumPropertyBinding Binding;
 
 	virtual bool GetValue(const UObject* ContextObject, uint8& OutResult) const override;
 
+	virtual void PostLoad() override
+	{
+		if (PropertyBindingChain_DEPRECATED.Num() > 0)
+		{
+			Binding.PropertyBindingChain = PropertyBindingChain_DEPRECATED;
+			PropertyBindingChain_DEPRECATED.SetNum(0);
+#if WITH_EDITORONLY_DATA
+			Binding.Enum = Enum_DEPRECATED;
+			Enum_DEPRECATED = nullptr;
+#endif
+		}
+	}
+
 #if WITH_EDITOR
 	static bool CanBind(const FProperty& Property)
 	{
@@ -39,43 +56,31 @@ struct CHOOSER_API FEnumContextProperty : public FChooserParameterEnumBase
 
 	virtual void GetDisplayName(FText& OutName) const override
 	{
-		if (!PropertyBindingChain.IsEmpty())
+		if (!Binding.PropertyBindingChain.IsEmpty())
 		{
-			OutName = FText::FromName(PropertyBindingChain.Last());
+			OutName = FText::FromName(Binding.PropertyBindingChain.Last());
 		}
 	}
 
-	virtual const UEnum* GetEnum() const override { return Enum; }
-#endif
+	virtual const UEnum* GetEnum() const override { return Binding.Enum; }
 
+#endif
+	
 private:
 #if WITH_EDITORONLY_DATA
 	UPROPERTY()
-	TObjectPtr<const UEnum> Enum = nullptr;
+	TObjectPtr<const UEnum> Enum_DEPRECATED = nullptr;
 #endif
 };
 
-
-
-UENUM()
-enum class EChooserEnumComparison : uint8
-{
-	Equal UMETA(DisplayName = "Value =="),
-	NotEqual UMETA(DisplayName = "Value !="),
-	GreaterThan UMETA(DisplayName = "Value >"),
-	GreaterThanEqual UMETA(DisplayName = "Value >="),
-	LessThan UMETA(DisplayName = "Value <"),
-	LessThanEqual UMETA(DisplayName = "Value <="),
-};
-
 USTRUCT()
 struct FChooserEnumRowData
 {
 	GENERATED_BODY()
 
 	UPROPERTY(EditAnywhere, Category = Runtime)
-	EChooserEnumComparison Comparison = EChooserEnumComparison::Equal;
-
+	bool CompareNotEqual = false;
+	
 	UPROPERTY(EditAnywhere, Category = Runtime)
 	uint8 Value = 0;
 
@@ -90,10 +95,10 @@ struct CHOOSER_API FEnumColumn : public FChooserColumnBase
 public:
 	FEnumColumn();
 
-	UPROPERTY(EditAnywhere, Meta = (ExcludeBaseStruct, BaseStruct = "/Script/Chooser.ChooserParameterEnumBase"), Category = "Hidden")
+	UPROPERTY(EditAnywhere, Meta = (ExcludeBaseStruct, BaseStruct = "/Script/Chooser.ChooserParameterEnumBase"), Category = "Data")
 	FInstancedStruct InputValue;
 
-	UPROPERTY(EditAnywhere, Category = Runtime)
+	UPROPERTY(EditAnywhere, Category = "Data")
 	// array of results (cells for this column for each row in the table)
 	// should match the length of the Results array
 	TArray<FChooserEnumRowData> RowValues;
@@ -101,29 +106,21 @@ public:
 	virtual void Filter(const UObject* ContextObject, const TArray<uint32>& IndexListIn, TArray<uint32>& IndexListOut) const override;
 	
 	CHOOSER_COLUMN_BOILERPLATE_NoSetInputType(FChooserParameterEnumBase);
-	
+
 #if WITH_EDITOR
 	virtual void SetInputType(const UScriptStruct* Type) override
 	{
 		InputValue.InitializeAs(Type);
-		InputChanged();
 	};
 
-	FSimpleMulticastDelegate OnEnumChanged;
-	
-	void InputChanged()
+	virtual void PostLoad() override
 	{
+		Super::PostLoad();
+		
 		if (InputValue.IsValid())
 		{
-			InputValue.GetMutable<FChooserParameterEnumBase>().OnEnumChanged.AddLambda([this](){ OnEnumChanged.Broadcast(); });
+			InputValue.GetMutable<FChooserParameterBase>().PostLoad();
 		}
-		OnEnumChanged.Broadcast();
-	}
-
-	virtual void PostLoad() override
-	{
-		Super::PostLoad();
-		InputChanged();
 	}
 #endif
 };
@@ -141,7 +138,7 @@ public:
 	{
 		OutInstancedStruct.InitializeAs(FEnumContextProperty::StaticStruct());
 		FEnumContextProperty& Property = OutInstancedStruct.GetMutable<FEnumContextProperty>();
-		Property.PropertyBindingChain = PropertyBindingChain;
+		Property.Binding.PropertyBindingChain = PropertyBindingChain;
 	}
 };
 
diff --git a/Engine/Plugins/Experimental/Chooser/Source/Chooser/Internal/FloatRangeColumn.h b/Engine/Plugins/Experimental/Chooser/Source/Chooser/Internal/FloatRangeColumn.h
--- a/Engine/Plugins/Experimental/Chooser/Source/Chooser/Internal/FloatRangeColumn.h
+++ b/Engine/Plugins/Experimental/Chooser/Source/Chooser/Internal/FloatRangeColumn.h
@@ -16,7 +16,20 @@ struct CHOOSER_API FFloatContextProperty :  public FChooserParameterFloatBase
 	virtual bool GetValue(const UObject* ContextObject, float& OutResult) const override;
 
 	UPROPERTY()
-	TArray<FName> PropertyBindingChain;
+	TArray<FName> PropertyBindingChain_DEPRECATED;
+	
+	UPROPERTY(EditAnywhere, Meta = (BindingType = "double", BindingColor = "FloatPinTypeColor"), Category = "Binding")
+	
+	FChooserPropertyBinding Binding;
+
+	virtual void PostLoad() override
+	{
+		if (PropertyBindingChain_DEPRECATED.Num() > 0)
+		{
+			Binding.PropertyBindingChain = PropertyBindingChain_DEPRECATED;
+			PropertyBindingChain_DEPRECATED.SetNum(0);
+		}
+	}	
 
 #if WITH_EDITOR
 	static bool CanBind(const FProperty& Property)
@@ -29,14 +42,14 @@ struct CHOOSER_API FFloatContextProperty :  public FChooserParameterFloatBase
 	
 	void SetBinding(const TArray<FBindingChainElement>& InBindingChain)
 	{
-		UE::Chooser::CopyPropertyChain(InBindingChain, PropertyBindingChain);
+		UE::Chooser::CopyPropertyChain(InBindingChain, Binding.PropertyBindingChain);
 	}
 
 	virtual void GetDisplayName(FText& OutName) const override
 	{
-		if (!PropertyBindingChain.IsEmpty())
+		if (!Binding.PropertyBindingChain.IsEmpty())
 		{
-			OutName = FText::FromName(PropertyBindingChain.Last());
+			OutName = FText::FromName(Binding.PropertyBindingChain.Last());
 		}
 	}
 #endif
@@ -62,16 +75,24 @@ struct CHOOSER_API FFloatRangeColumn : public FChooserColumnBase
 	public:
 	FFloatRangeColumn();
 		
-	UPROPERTY(EditAnywhere, Meta = (ExcludeBaseStruct, BaseStruct = "/Script/Chooser.ChooserParameterFloatBase"), Category = "Hidden")
+	UPROPERTY(EditAnywhere, Meta = (ExcludeBaseStruct, BaseStruct = "/Script/Chooser.ChooserParameterFloatBase"), Category = "Data")
 	FInstancedStruct InputValue;
 	
-	UPROPERTY(EditAnywhere, Category=Runtime)
+	UPROPERTY(EditAnywhere, Category="Data")
 	// array of results (cells for this column for each row in the table)
 	// should match the length of the Results array 
 	TArray<FChooserFloatRangeRowData> RowValues;
 	
 	virtual void Filter(const UObject* ContextObject, const TArray<uint32>& IndexListIn, TArray<uint32>& IndexListOut) const override;
 	
+	virtual void PostLoad() override
+	{
+		if (InputValue.IsValid())
+		{
+			InputValue.GetMutable<FChooserParameterBase>().PostLoad();
+		}
+	}
+	
 	CHOOSER_COLUMN_BOILERPLATE(FChooserParameterFloatBase);
 
 };
@@ -90,7 +111,7 @@ public:
 	{
 		OutInstancedStruct.InitializeAs(FFloatContextProperty::StaticStruct());
 		FFloatContextProperty& Property = OutInstancedStruct.GetMutable<FFloatContextProperty>();
-		Property.PropertyBindingChain = PropertyBindingChain;
+		Property.Binding.PropertyBindingChain = PropertyBindingChain;
 	}
 };
 
diff --git a/Engine/Plugins/Experimental/Chooser/Source/Chooser/Internal/GameplayTagColumn.h b/Engine/Plugins/Experimental/Chooser/Source/Chooser/Internal/GameplayTagColumn.h
--- a/Engine/Plugins/Experimental/Chooser/Source/Chooser/Internal/GameplayTagColumn.h
+++ b/Engine/Plugins/Experimental/Chooser/Source/Chooser/Internal/GameplayTagColumn.h
@@ -18,10 +18,21 @@ struct CHOOSER_API FGameplayTagContextProperty :  public FChooserParameterGamepl
 	GENERATED_BODY()
 public:
 	UPROPERTY()
-	TArray<FName> PropertyBindingChain;
+	TArray<FName> PropertyBindingChain_DEPRECATED;
+	
+	UPROPERTY(EditAnywhere, Meta = (BindingType = "FGameplayTagContainer", BindingColor = "StructPinTypeColor"), Category = "Binding")
+	FChooserPropertyBinding Binding;
 	
 	virtual bool GetValue(const UObject* ContextObject, const FGameplayTagContainer*& OutResult) const override;
-
+	
+	virtual void PostLoad() override
+	{
+		if (PropertyBindingChain_DEPRECATED.Num() > 0)
+		{
+			Binding.PropertyBindingChain = PropertyBindingChain_DEPRECATED;
+			PropertyBindingChain_DEPRECATED.SetNum(0);
+		}
+	}
 #if WITH_EDITOR
 	static bool CanBind(const FProperty& Property)
 	{
@@ -31,14 +42,14 @@ public:
 
 	void SetBinding(const TArray<FBindingChainElement>& InBindingChain)
 	{
-		UE::Chooser::CopyPropertyChain(InBindingChain, PropertyBindingChain);
+		UE::Chooser::CopyPropertyChain(InBindingChain, Binding.PropertyBindingChain);
 	}
 
 	virtual void GetDisplayName(FText& OutName) const override
 	{
-		if (!PropertyBindingChain.IsEmpty())
+		if (!Binding.PropertyBindingChain.IsEmpty())
 		{
-			OutName = FText::FromName(PropertyBindingChain.Last());
+			OutName = FText::FromName(Binding.PropertyBindingChain.Last());
 		}
 	}
 #endif
@@ -51,10 +62,10 @@ struct CHOOSER_API FGameplayTagColumn : public FChooserColumnBase
 	public:
 	FGameplayTagColumn();
 	
-	UPROPERTY(EditAnywhere, Meta = (ExcludeBaseStruct, BaseStruct = "/Script/Chooser.ChooserParameterGameplayTagBase"), Category = "Hidden")
+	UPROPERTY(EditAnywhere, Meta = (ExcludeBaseStruct, BaseStruct = "/Script/Chooser.ChooserParameterGameplayTagBase"), Category = "Data")
 	FInstancedStruct InputValue;
 
-	UPROPERTY(EditAnywhere, Category=Runtime)
+	UPROPERTY(EditAnywhere, Category="Data")
 	EGameplayContainerMatchType	TagMatchType = EGameplayContainerMatchType::Any;
 
 	UPROPERTY(EditAnywhere, Category=Runtime)
@@ -64,6 +75,14 @@ struct CHOOSER_API FGameplayTagColumn : public FChooserColumnBase
 	
 	virtual void Filter(const UObject* ContextObject, const TArray<uint32>& IndexListIn, TArray<uint32>& IndexListOut) const override;
 
+	virtual void PostLoad() override
+	{
+		if (InputValue.IsValid())
+		{
+			InputValue.GetMutable<FChooserParameterBase>().PostLoad();
+		}
+	}
+
 	CHOOSER_COLUMN_BOILERPLATE(FChooserParameterGameplayTagBase);
 };
 
@@ -81,7 +100,7 @@ public:
 	{
 		OutInstancedStruct.InitializeAs(FGameplayTagContextProperty::StaticStruct());
 		FGameplayTagContextProperty& Property = OutInstancedStruct.GetMutable<FGameplayTagContextProperty>();
-		Property.PropertyBindingChain = PropertyBindingChain;
+		Property.Binding.PropertyBindingChain = PropertyBindingChain;
 	}
 };
 
@@ -106,6 +125,7 @@ class CHOOSER_API UDEPRECATED_ChooserColumnGameplayTag : public UObject, public
 	// array of results (cells for this column for each row in the table)
 	// should match the length of the Results array 
 	TArray<FGameplayTagContainer> RowValues;
+
 	
 	virtual void ConvertToInstancedStruct(FInstancedStruct& OutInstancedStruct) const override
 	{
diff --git a/Engine/Plugins/Experimental/Chooser/Source/Chooser/Private/BoolColumn.cpp b/Engine/Plugins/Experimental/Chooser/Source/Chooser/Private/BoolColumn.cpp
--- a/Engine/Plugins/Experimental/Chooser/Source/Chooser/Private/BoolColumn.cpp
+++ b/Engine/Plugins/Experimental/Chooser/Source/Chooser/Private/BoolColumn.cpp
@@ -7,9 +7,9 @@ bool FBoolContextProperty::GetValue(const UObject* ContextObject, bool& OutResul
 	UStruct* StructType = ContextObject->GetClass();
 	const void* Container = ContextObject;
 	
-	if (UE::Chooser::ResolvePropertyChain(Container, StructType, PropertyBindingChain))
+	if (UE::Chooser::ResolvePropertyChain(Container, StructType, Binding.PropertyBindingChain))
 	{
-		if (const FBoolProperty* Property = FindFProperty<FBoolProperty>(StructType, PropertyBindingChain.Last()))
+		if (const FBoolProperty* Property = FindFProperty<FBoolProperty>(StructType, Binding.PropertyBindingChain.Last()))
 		{
 			OutResult = *Property->ContainerPtrToValuePtr<bool>(Container);
 			return true;
@@ -24,9 +24,9 @@ bool FBoolContextProperty::SetValue(UObject* ContextObject, bool InValue) const
 	UStruct* StructType = ContextObject->GetClass();
 	const void* Container = ContextObject;
 	
-	if (UE::Chooser::ResolvePropertyChain(Container, StructType, PropertyBindingChain))
+	if (UE::Chooser::ResolvePropertyChain(Container, StructType, Binding.PropertyBindingChain))
 	{
-		if (FBoolProperty* Property = FindFProperty<FBoolProperty>(StructType, PropertyBindingChain.Last()))
+		if (FBoolProperty* Property = FindFProperty<FBoolProperty>(StructType, Binding.PropertyBindingChain.Last()))
 		{
 			// const cast is here just because ResolvePropertyChain expects a const void*&
 			*Property->ContainerPtrToValuePtr<bool>(const_cast<void*>(Container)) = InValue;
diff --git a/Engine/Plugins/Experimental/Chooser/Source/Chooser/Private/EnumColumn.cpp b/Engine/Plugins/Experimental/Chooser/Source/Chooser/Private/EnumColumn.cpp
--- a/Engine/Plugins/Experimental/Chooser/Source/Chooser/Private/EnumColumn.cpp
+++ b/Engine/Plugins/Experimental/Chooser/Source/Chooser/Private/EnumColumn.cpp
@@ -11,15 +11,15 @@ bool FEnumContextProperty::GetValue(const UObject* ContextObject, uint8& OutResu
 	UStruct* StructType = ContextObject->GetClass();
 	const void* Container = ContextObject;
 
-	if (UE::Chooser::ResolvePropertyChain(Container, StructType, PropertyBindingChain))
+	if (UE::Chooser::ResolvePropertyChain(Container, StructType, Binding.PropertyBindingChain))
 	{
-		if (const FEnumProperty* EnumProperty = FindFProperty<FEnumProperty>(StructType, PropertyBindingChain.Last()))
+		if (const FEnumProperty* EnumProperty = FindFProperty<FEnumProperty>(StructType, Binding.PropertyBindingChain.Last()))
 		{
 			OutResult = *EnumProperty->ContainerPtrToValuePtr<uint8>(Container);
 			return true;
 		}
 
-		if (const FByteProperty* ByteProperty = FindFProperty<FByteProperty>(StructType, PropertyBindingChain.Last()))
+		if (const FByteProperty* ByteProperty = FindFProperty<FByteProperty>(StructType, Binding.PropertyBindingChain.Last()))
 		{
 			if (ByteProperty->IsEnum())
 			{
@@ -36,25 +36,19 @@ bool FEnumContextProperty::GetValue(const UObject* ContextObject, uint8& OutResu
 
 void FEnumContextProperty::SetBinding(const TArray<FBindingChainElement>& InBindingChain)
 {
-	const UEnum* PreviousEnum = Enum;
-	Enum = nullptr;
+	const UEnum* PreviousEnum = Binding.Enum;
+	Binding.Enum = nullptr;
 
-	UE::Chooser::CopyPropertyChain(InBindingChain, PropertyBindingChain);
+	UE::Chooser::CopyPropertyChain(InBindingChain, Binding.PropertyBindingChain);
 
 	const FField* Field = InBindingChain.Last().Field.ToField();
 	if (const FEnumProperty* EnumProperty = CastField<FEnumProperty>(Field))
 	{
-		Enum = EnumProperty->GetEnum();
+		Binding.Enum = EnumProperty->GetEnum();
 	}
 	else if (const FByteProperty* ByteProperty = CastField<FByteProperty>(Field))
 	{
-		Enum = ByteProperty->Enum;
-	}
-
-	if (Enum != PreviousEnum)
-	{
-		// Our enum type has changed! Need to refresh the UI to update enum value pickers.
-		OnEnumChanged.Broadcast();
+		Binding.Enum = ByteProperty->Enum;
 	}
 }
 
@@ -64,36 +58,13 @@ FEnumColumn::FEnumColumn()
 {
 #if WITH_EDITOR
 	InputValue.InitializeAs(FEnumContextProperty::StaticStruct());
-	InputChanged();
 #endif
 }
 
 bool FChooserEnumRowData::Evaluate(const uint8 LeftHandSide) const
 {
-	switch (Comparison)
-	{
-	case EChooserEnumComparison::Equal:
-		return LeftHandSide == Value;
-
-	case EChooserEnumComparison::NotEqual:
-		return LeftHandSide != Value;
-
-	case EChooserEnumComparison::GreaterThan:
-		return LeftHandSide > Value;
-
-	case EChooserEnumComparison::GreaterThanEqual:
-		return LeftHandSide >= Value;
-
-	case EChooserEnumComparison::LessThan:
-		return LeftHandSide < Value;
-
-	case EChooserEnumComparison::LessThanEqual:
-		return LeftHandSide <= Value;
-
-	default:
-		checkf(false, TEXT("Unsupported comparison type for enum comparison operation!"));
-		return false;
-	}
+	bool Equal = LeftHandSide == Value;
+	return Equal ^ CompareNotEqual;
 }
 
 void FEnumColumn::Filter(const UObject* ContextObject, const TArray<uint32>& IndexListIn, TArray<uint32>& IndexListOut) const
diff --git a/Engine/Plugins/Experimental/Chooser/Source/Chooser/Private/FloatRangeColumn.cpp b/Engine/Plugins/Experimental/Chooser/Source/Chooser/Private/FloatRangeColumn.cpp
--- a/Engine/Plugins/Experimental/Chooser/Source/Chooser/Private/FloatRangeColumn.cpp
+++ b/Engine/Plugins/Experimental/Chooser/Source/Chooser/Private/FloatRangeColumn.cpp
@@ -7,15 +7,15 @@ bool FFloatContextProperty::GetValue(const UObject* ContextObject, float& OutRes
 	UStruct* StructType = ContextObject->GetClass();
 	const void* Container = ContextObject;
 	
-	if (UE::Chooser::ResolvePropertyChain(Container, StructType, PropertyBindingChain))
+	if (UE::Chooser::ResolvePropertyChain(Container, StructType, Binding.PropertyBindingChain))
 	{
-		if (const FDoubleProperty* DoubleProperty = FindFProperty<FDoubleProperty>(StructType, PropertyBindingChain.Last()))
+		if (const FDoubleProperty* DoubleProperty = FindFProperty<FDoubleProperty>(StructType, Binding.PropertyBindingChain.Last()))
 		{
 			OutResult = *DoubleProperty->ContainerPtrToValuePtr<double>(Container);
 			return true;
 		}
 		
-		if (const FFloatProperty* FloatProperty = FindFProperty<FFloatProperty>(StructType, PropertyBindingChain.Last()))
+		if (const FFloatProperty* FloatProperty = FindFProperty<FFloatProperty>(StructType, Binding.PropertyBindingChain.Last()))
 		{
 			OutResult = *FloatProperty->ContainerPtrToValuePtr<float>(Container);
 			return true;
diff --git a/Engine/Plugins/Experimental/Chooser/Source/Chooser/Private/GameplayTagColumn.cpp b/Engine/Plugins/Experimental/Chooser/Source/Chooser/Private/GameplayTagColumn.cpp
--- a/Engine/Plugins/Experimental/Chooser/Source/Chooser/Private/GameplayTagColumn.cpp
+++ b/Engine/Plugins/Experimental/Chooser/Source/Chooser/Private/GameplayTagColumn.cpp
@@ -7,9 +7,9 @@ bool FGameplayTagContextProperty::GetValue(const UObject* ContextObject, const F
 	UStruct* StructType = ContextObject->GetClass();
 	const void* Container = ContextObject;
 	
-	if (UE::Chooser::ResolvePropertyChain(Container, StructType, PropertyBindingChain))
+	if (UE::Chooser::ResolvePropertyChain(Container, StructType, Binding.PropertyBindingChain))
 	{
-		if (const FStructProperty* Property = FindFProperty<FStructProperty>(StructType, PropertyBindingChain.Last()))
+		if (const FStructProperty* Property = FindFProperty<FStructProperty>(StructType, Binding.PropertyBindingChain.Last()))
 		{
 			OutResult = Property->ContainerPtrToValuePtr<FGameplayTagContainer>(Container);
 			return true;
diff --git a/Engine/Plugins/Experimental/Chooser/Source/Chooser/Public/ChooserPropertyAccess.h b/Engine/Plugins/Experimental/Chooser/Source/Chooser/Public/ChooserPropertyAccess.h
--- a/Engine/Plugins/Experimental/Chooser/Source/Chooser/Public/ChooserPropertyAccess.h
+++ b/Engine/Plugins/Experimental/Chooser/Source/Chooser/Public/ChooserPropertyAccess.h
@@ -33,6 +33,17 @@ struct FChooserPropertyBinding
 	TArray<FName> PropertyBindingChain;
 };
 
+USTRUCT()
+struct FChooserEnumPropertyBinding : public FChooserPropertyBinding
+{
+	GENERATED_BODY()
+	
+#if WITH_EDITORONLY_DATA
+	UPROPERTY()
+	TObjectPtr<const UEnum> Enum = nullptr;
+#endif
+};
+
 namespace UE::Chooser
 {
 	CHOOSER_API bool ResolvePropertyChain(const void*& Container, UStruct*& StructType, const TArray<FName>& PropertyBindingChain);
diff --git a/Engine/Plugins/Experimental/Chooser/Source/Chooser/Public/IChooserParameterBase.h b/Engine/Plugins/Experimental/Chooser/Source/Chooser/Public/IChooserParameterBase.h
--- a/Engine/Plugins/Experimental/Chooser/Source/Chooser/Public/IChooserParameterBase.h
+++ b/Engine/Plugins/Experimental/Chooser/Source/Chooser/Public/IChooserParameterBase.h
@@ -14,5 +14,7 @@ struct FChooserParameterBase
 
 	virtual void GetDisplayName(FText& OutName) const { }
 
+	virtual void PostLoad() {};
+
 	virtual ~FChooserParameterBase() {}
 };
\ No newline at end of file
diff --git a/Engine/Plugins/Experimental/Chooser/Source/Chooser/Public/IChooserParameterEnum.h b/Engine/Plugins/Experimental/Chooser/Source/Chooser/Public/IChooserParameterEnum.h
--- a/Engine/Plugins/Experimental/Chooser/Source/Chooser/Public/IChooserParameterEnum.h
+++ b/Engine/Plugins/Experimental/Chooser/Source/Chooser/Public/IChooserParameterEnum.h
@@ -31,6 +31,5 @@ struct FChooserParameterEnumBase : public FChooserParameterBase
 
 	#if WITH_EDITOR
     	virtual const UEnum* GetEnum() const { return nullptr; }
-    	FSimpleMulticastDelegate OnEnumChanged;
     #endif
 };
\ No newline at end of file
diff --git a/Engine/Plugins/Experimental/Chooser/Source/ChooserEditor/Private/ChooserEditorModule.cpp b/Engine/Plugins/Experimental/Chooser/Source/ChooserEditor/Private/ChooserEditorModule.cpp
--- a/Engine/Plugins/Experimental/Chooser/Source/ChooserEditor/Private/ChooserEditorModule.cpp
+++ b/Engine/Plugins/Experimental/Chooser/Source/ChooserEditor/Private/ChooserEditorModule.cpp
@@ -9,6 +9,7 @@
 #include "GameplayTagColumnEditor.h"
 #include "ChooserTableEditorCommands.h"
 #include "ChooserPropertyAccess.h"
+#include "PropertyEditorModule.h"
 
 #define LOCTEXT_NAMESPACE "ChooserEditorModule"
 
@@ -28,6 +29,7 @@ void FModule::StartupModule()
 	FPropertyEditorModule& PropertyModule = FModuleManager::LoadModuleChecked<FPropertyEditorModule>("PropertyEditor");
 	
 	PropertyModule.RegisterCustomPropertyTypeLayout(FChooserPropertyBinding::StaticStruct()->GetFName(), FOnGetPropertyTypeCustomizationInstance::CreateLambda([] { return MakeShared<FPropertyAccessChainCustomization>(); }));
+	PropertyModule.RegisterCustomPropertyTypeLayout(FChooserEnumPropertyBinding::StaticStruct()->GetFName(), FOnGetPropertyTypeCustomizationInstance::CreateLambda([] { return MakeShared<FPropertyAccessChainCustomization>(); }));
 
 }
 
diff --git a/Engine/Plugins/Experimental/Chooser/Source/ChooserEditor/Private/ChooserTableEditor.cpp b/Engine/Plugins/Experimental/Chooser/Source/ChooserEditor/Private/ChooserTableEditor.cpp
--- a/Engine/Plugins/Experimental/Chooser/Source/ChooserEditor/Private/ChooserTableEditor.cpp
+++ b/Engine/Plugins/Experimental/Chooser/Source/ChooserEditor/Private/ChooserTableEditor.cpp
@@ -194,7 +194,8 @@ void FChooserTableEditor::NotifyPostChange(const FPropertyChangedEvent& Property
 		// rebuild all result widgets
 		UpdateTableRows();
 	}
-	if (PropertyThatChanged->GetNameCPP() == GET_MEMBER_NAME_STRING_CHECKED(UChooserTable, ContextObjectType))
+	if (PropertyThatChanged->GetNameCPP() == GET_MEMBER_NAME_STRING_CHECKED(UChooserTable, ContextObjectType)     // refresh columns property selector if the Context Object changed
+		|| PropertyThatChanged->GetNameCPP() == "InputValue" || PropertyThatChanged->GetNameCPP() == "Binding")   // refresh column widgets if the Input value or its binding was changed from the details panel
 	{
 		// rebuild all column header widgets
 		UpdateTableColumns();
@@ -805,7 +806,7 @@ TSharedRef<SDockTab> FChooserTableEditor::SpawnTableTab( const FSpawnTabArgs& Ar
 				}
 				Chooser->ColumnsStructs.Insert(NewColumn, InsertIndex);
 			}
-			
+
 			UpdateTableColumns();
 			UpdateTableRows();
 
@@ -1215,17 +1216,17 @@ void FChooserColumnDetails::CustomizeDetails(IDetailLayoutBuilder& DetailBuilder
 	
 	if (Chooser->ColumnsStructs.Num() > Column->Column)
 	{
-		IDetailCategoryBuilder& PropertiesCategory = DetailBuilder.EditCategory("Row Properties");
+		IDetailCategoryBuilder& PropertiesCategory = DetailBuilder.EditCategory("Column Properties");
 
 		TSharedPtr<IPropertyHandle> ChooserProperty = DetailBuilder.GetProperty("Chooser", Column->StaticClass());
 		DetailBuilder.HideProperty(ChooserProperty);
 	
 		TSharedPtr<IPropertyHandle> ColumnsArrayProperty = ChooserProperty->GetChildHandle("ColumnsStructs");
 		TSharedPtr<IPropertyHandle> CurrentColumnProperty = ColumnsArrayProperty->AsArray()->GetElement(Column->Column);
-		IDetailPropertyRow& NewResultProperty = PropertiesCategory.AddProperty(CurrentColumnProperty);
-		NewResultProperty.DisplayName(LOCTEXT("Selected Column","Selected Column"));
-		NewResultProperty.ShowPropertyButtons(false); // hide array add button
-		NewResultProperty.ShouldAutoExpand(true);
+		IDetailPropertyRow& NewColumnProperty = PropertiesCategory.AddProperty(CurrentColumnProperty);
+		NewColumnProperty.DisplayName(LOCTEXT("Selected Column","Selected Column"));
+		NewColumnProperty.ShowPropertyButtons(false); // hide array add button
+		NewColumnProperty.ShouldAutoExpand(true);
 	}
 }
 
diff --git a/Engine/Plugins/Experimental/Chooser/Source/ChooserEditor/Private/EnumColumnEditor.cpp b/Engine/Plugins/Experimental/Chooser/Source/ChooserEditor/Private/EnumColumnEditor.cpp
--- a/Engine/Plugins/Experimental/Chooser/Source/ChooserEditor/Private/EnumColumnEditor.cpp
+++ b/Engine/Plugins/Experimental/Chooser/Source/ChooserEditor/Private/EnumColumnEditor.cpp
@@ -7,117 +7,149 @@
 #include "ChooserTableEditor.h"
 #include "GraphEditorSettings.h"
 #include "SEnumCombo.h"
+#include "TransactionCommon.h"
+#include "Widgets/Input/SButton.h"
 
 #define LOCTEXT_NAMESPACE "EnumColumnEditor"
 
 namespace UE::ChooserEditor
 {
-	// Wrapper widget for EnumComboBox which will reconstruct the combo box when the Enum has changed
-	class SEnumCell : public SCompoundWidget
-	{
-	public:
-		SLATE_BEGIN_ARGS(SEnumCell)
-			: _RowIndex(-1)
-		{}
+	
+// Wrapper widget for EnumComboBox which will reconstruct the combo box when the Enum has changed
+class SEnumCell : public SCompoundWidget
+{
+public:
+	SLATE_BEGIN_ARGS(SEnumCell)
+		: _RowIndex(-1)
+	{}
 
-		SLATE_ARGUMENT(FEnumColumn*, EnumColumn)
-		SLATE_ATTRIBUTE(int, RowIndex);
-                
-		SLATE_END_ARGS()
+	SLATE_ARGUMENT(UObject*, TransactionObject)
+	SLATE_ARGUMENT(FEnumColumn*, EnumColumn)
+	SLATE_ATTRIBUTE(int, RowIndex);
+            
+	SLATE_END_ARGS()
 
-		TSharedRef<SWidget> CreateEnumComboBox()
+	TSharedRef<SWidget> CreateEnumComboBox()
+	{
+		if (const FEnumColumn* EnumColumnPointer = EnumColumn)
 		{
-			if (const FEnumColumn* EnumColumnPointer = EnumColumn)
+			if (EnumColumnPointer->InputValue.IsValid())
 			{
-				if (EnumColumnPointer->InputValue.IsValid())
+				if (const UEnum* Enum = EnumColumnPointer->InputValue.Get<FChooserParameterEnumBase>().GetEnum())
 				{
-					if (const UEnum* Enum = EnumColumnPointer->InputValue.Get<FChooserParameterEnumBase>().GetEnum())
-					{
-						return SNew(SEnumComboBox, Enum)
-							.CurrentValue_Lambda([this]()
+					return SNew(SEnumComboBox, Enum)
+						.CurrentValue_Lambda([this]()
+						{
+							int Row = RowIndex.Get();
+							return EnumColumn->RowValues.IsValidIndex(Row) ? static_cast<int32>(EnumColumn->RowValues[Row].Value) : 0;
+						})
+						.OnEnumSelectionChanged_Lambda([this](int32 EnumValue, ESelectInfo::Type)
+						{
+							int Row = RowIndex.Get();
+							if (EnumColumn->RowValues.IsValidIndex(Row))
 							{
-								int Row = RowIndex.Get();
-								return EnumColumn->RowValues.IsValidIndex(Row) ? static_cast<int32>(EnumColumn->RowValues[Row].Value) : 0;
-							})
-							.OnEnumSelectionChanged_Lambda([this](int32 EnumValue, ESelectInfo::Type)
-							{
-								int Row = RowIndex.Get();
-								if (EnumColumn->RowValues.IsValidIndex(Row))
-								{
-									const FScopedTransaction Transaction(LOCTEXT("Edit RHS", "Edit Enum Value"));
-									// todo: need a reference to the UChooserTable to modify
-									// EnumColumn->Modify(true);
-									EnumColumn->RowValues[Row].Value = static_cast<uint8>(EnumValue);
-								}
-							});
-					}
+								const FScopedTransaction Transaction(LOCTEXT("Edit RHS", "Edit Enum Value"));
+								TransactionObject->Modify(true);
+								EnumColumn->RowValues[Row].Value = static_cast<uint8>(EnumValue);
+							}
+						});
 				}
 			}
-			
-			return SNullWidget::NullWidget;
-		}
-
-		void UpdateEnumComboBox()
-		{
-			ChildSlot[ CreateEnumComboBox()	];
 		}
+		
+		return SNullWidget::NullWidget;
+	}
 
-		void Construct( const FArguments& InArgs)
-		{
-			RowIndex = InArgs._RowIndex;
-			EnumColumn = InArgs._EnumColumn;
+	void UpdateEnumComboBox()
+	{
+		ChildSlot[ CreateEnumComboBox()	];
+	}
 
-			if (FEnumColumn* EnumColumnPointer = EnumColumn)
-			{
-				EnumChangedHandle = EnumColumnPointer->OnEnumChanged.AddSP(this, &SEnumCell::UpdateEnumComboBox);
-			}
-			
-			UpdateEnumComboBox();
-		}
+	void Construct( const FArguments& InArgs)
+	{
+		RowIndex = InArgs._RowIndex;
+		EnumColumn = InArgs._EnumColumn;
+		TransactionObject = InArgs._TransactionObject;
 
-		~SEnumCell()
+		if (const FEnumColumn* EnumColumnPointer = EnumColumn)
 		{
-			if (FEnumColumn* EnumColumnPointer = EnumColumn)
+			if (EnumColumnPointer->InputValue.IsValid())
 			{
-				EnumColumnPointer->OnEnumChanged.Remove(EnumChangedHandle);
+				EnumSource = EnumColumnPointer->InputValue.Get<FChooserParameterEnumBase>().GetEnum();
 			}
 		}
 
-	private:
-		FEnumColumn* EnumColumn = nullptr;
-		TAttribute<int> RowIndex;
-		FDelegateHandle EnumChangedHandle;
-	};
+		UpdateEnumComboBox();
 
-	TSharedRef<SWidget> CreateEnumColumnWidget(UChooserTable* Chooser, FChooserColumnBase* Column, int Row)
-	{
-		FEnumColumn* EnumColumn = static_cast<FEnumColumn*>(Column);
-		
-		return SNew(SHorizontalBox)
-		+ SHorizontalBox::Slot().MaxWidth(95.0f)
+		int Row = RowIndex.Get();
+
+		ChildSlot
 		[
-			SNew(SEnumComboBox, StaticEnum<EChooserEnumComparison>())
-			.CurrentValue_Lambda([EnumColumn, Row]()
-			{
-				return EnumColumn->RowValues.IsValidIndex(Row) ? static_cast<int32>(EnumColumn->RowValues[Row].Comparison) : 0;
-			})
-			.OnEnumSelectionChanged_Lambda([Chooser, EnumColumn, Row](int32 EnumValue, ESelectInfo::Type)
-			{
-				if (EnumColumn->RowValues.IsValidIndex(Row))
+			SNew(SHorizontalBox)
+			+ SHorizontalBox::Slot().MaxWidth(45.0f)
+			[
+				SNew(SButton).ButtonStyle(FAppStyle::Get(),"FlatButton").TextStyle(FAppStyle::Get(),"RichTextBlock.Bold").HAlign(HAlign_Center)
+				.Text_Lambda([this, Row]()
 				{
-					const FScopedTransaction Transaction(LOCTEXT("Edit Comparison", "Edit Comparison Operation"));
-					Chooser->Modify(true);
-					EnumColumn->RowValues[Row].Comparison = static_cast<EChooserEnumComparison>(EnumValue);
-				}
-			})
-		]
-		+ SHorizontalBox::Slot()
-		[
-			SNew(SEnumCell).EnumColumn(EnumColumn).RowIndex(Row)
+					// this is a bit of an odd thing to do, but as we are polling for the "=" / "!=" button state anyway,
+					// also poll to see if the currently bound Enum has changed, and if it has, recreate the enum combo box.
+					if (const FEnumColumn* EnumColumnPointer = EnumColumn)
+					{
+						const UEnum* CurrentEnumSource = nullptr;
+						if (EnumColumnPointer->InputValue.IsValid())
+						{
+							CurrentEnumSource = EnumColumnPointer->InputValue.Get<FChooserParameterEnumBase>().GetEnum(); 
+						}
+						if (EnumSource != CurrentEnumSource)
+						{
+							EnumComboBorder->SetContent(CreateEnumComboBox());
+							EnumSource = EnumColumnPointer->InputValue.Get<FChooserParameterEnumBase>().GetEnum();
+						}
+					}
+					
+					return (EnumColumn->RowValues.IsValidIndex(Row) && EnumColumn->RowValues[Row].CompareNotEqual ? LOCTEXT("Not Equal", "!=") : LOCTEXT("Equal", "="));
+				})
+				.OnClicked_Lambda([this, Row]()
+				{
+					if (EnumColumn->RowValues.IsValidIndex(Row))
+					{
+						const FScopedTransaction Transaction(LOCTEXT("Edit Comparison", "Edit Comparison Operation"));
+						TransactionObject->Modify(true);
+						EnumColumn->RowValues[Row].CompareNotEqual = !EnumColumn->RowValues[Row].CompareNotEqual;
+					}
+					return FReply::Handled();
+				})
+			]
+			+ SHorizontalBox::Slot().FillWidth(true)
+			[
+				SAssignNew(EnumComboBorder, SBorder).Padding(0).BorderBackgroundColor(FLinearColor(0,0,0,0))
+				[
+					CreateEnumComboBox()
+				]
+			]
 		];
+		
 	}
 
+	~SEnumCell()
+	{
+	}
+
+private:
+	UObject* TransactionObject = nullptr;
+	FEnumColumn* EnumColumn = nullptr;
+	const UEnum* EnumSource = nullptr;
+	TSharedPtr<SBorder> EnumComboBorder;
+	TAttribute<int> RowIndex;
+	FDelegateHandle EnumChangedHandle;
+};
+
+TSharedRef<SWidget> CreateEnumColumnWidget(UChooserTable* Chooser, FChooserColumnBase* Column, int Row)
+{
+	FEnumColumn* EnumColumn = static_cast<FEnumColumn*>(Column);
 	
+	return SNew(SEnumCell).TransactionObject(Chooser).EnumColumn(EnumColumn).RowIndex(Row);
+}
 
 TSharedRef<SWidget> CreateEnumPropertyWidget(UObject* TransactionObject, void* Value, UClass* ContextClass, UClass* ResultBaseClass)
 {
diff --git a/Engine/Plugins/Experimental/Chooser/Source/ChooserEditor/Private/PropertyAccessChainCustomization.cpp b/Engine/Plugins/Experimental/Chooser/Source/ChooserEditor/Private/PropertyAccessChainCustomization.cpp
--- a/Engine/Plugins/Experimental/Chooser/Source/ChooserEditor/Private/PropertyAccessChainCustomization.cpp
+++ b/Engine/Plugins/Experimental/Chooser/Source/ChooserEditor/Private/PropertyAccessChainCustomization.cpp
@@ -6,6 +6,7 @@
 #include "ChooserPropertyAccess.h"
 #include "PropertyHandle.h"
 #include "DetailWidgetRow.h"
+#include "GraphEditorSettings.h"
 #include "SClassViewer.h"
 #include "IDetailChildrenBuilder.h"
 #include "IPropertyAccessEditor.h"
@@ -23,8 +24,10 @@ void FPropertyAccessChainCustomization::CustomizeHeader(TSharedRef<IPropertyHand
 
 	UClass* ContextClass = nullptr;
 	static FName TypeMetaData = "BindingType";
+	static FName ColorMetaData = "BindingColor";
 
 	FString TypeFilter = PropertyHandle->GetMetaData(TypeMetaData);
+	FString BindingColor = PropertyHandle->GetMetaData(ColorMetaData);
 
 	TArray<UObject*> OuterObjects;
 	PropertyHandle->GetOuterObjects(OuterObjects);
@@ -48,24 +51,44 @@ void FPropertyAccessChainCustomization::CustomizeHeader(TSharedRef<IPropertyHand
 	
 	Args.OnCanBindProperty = FOnCanBindProperty::CreateLambda([TypeFilter](FProperty* Property)
 	{
-		if (TypeFilter == "")
+		if (TypeFilter == "" || Property == nullptr)
 		{
 			return true;
 		}
-		else
+		else if (TypeFilter == "enum")
 		{
-			return Property == nullptr || Property->GetCPPType() == TypeFilter;
+            if (const FEnumProperty* EnumProperty = CastField<const FEnumProperty>(Property))
+            {
+             	return true;
+			}
+            else if (const FByteProperty* ByteProperty = CastField<const FByteProperty>(Property))
+            {
+            	return ByteProperty->Enum != nullptr;
+            }
+			return false;
 		}
 		
+		return Property->GetCPPType() == TypeFilter;
 	});
 	
-	Args.OnCanBindToClass = FOnCanBindToClass::CreateLambda([ContextClass ](UClass* InClass)
+	Args.OnCanBindToClass = FOnCanBindToClass::CreateLambda([ContextClass](UClass* InClass)
 	{
 		return true;
 	});
 
-	Args.CurrentBindingColor = MakeAttributeLambda([]() {
-		return FLinearColor::Gray;
+
+	FLinearColor BindingColorValue = FLinearColor::Gray;
+	if (BindingColor != "")
+	{
+		const UGraphEditorSettings* GraphEditorSettings = GetDefault<UGraphEditorSettings>();
+		if (const FStructProperty* ColorProperty = FindFProperty<FStructProperty>(GraphEditorSettings->GetClass(), FName(BindingColor)))
+		{
+			BindingColorValue = *ColorProperty->ContainerPtrToValuePtr<FLinearColor>(GraphEditorSettings);
+		}
+	}
+
+	Args.CurrentBindingColor = MakeAttributeLambda([BindingColorValue]() {
+		return BindingColorValue;
 	});
 
 	Args.OnCanAcceptPropertyOrChildren = FOnCanBindProperty::CreateLambda([](FProperty* InProperty)
@@ -74,7 +97,7 @@ void FPropertyAccessChainCustomization::CustomizeHeader(TSharedRef<IPropertyHand
 			return InProperty->HasAnyPropertyFlags(CPF_BlueprintVisible);
 		});	
 
-	Args.OnAddBinding = FOnAddBinding::CreateLambda([PropertyHandle](FName InPropertyName, const TArray<FBindingChainElement>& InBindingChain)
+	Args.OnAddBinding = FOnAddBinding::CreateLambda([TypeFilter, PropertyHandle](FName InPropertyName, const TArray<FBindingChainElement>& InBindingChain)
 		{
 			TArray<UObject*> OuterObjects;
 			PropertyHandle->GetOuterObjects(OuterObjects);
@@ -90,9 +113,24 @@ void FPropertyAccessChainCustomization::CustomizeHeader(TSharedRef<IPropertyHand
 				{
 					PropertyHandle->NotifyPreChange();
 					
-					OuterObjects[i]->Modify(true); // todo
+					OuterObjects[i]->Modify(true);
 					Chooser::CopyPropertyChain(InBindingChain, PropertyValue->PropertyBindingChain);
 
+					if (TypeFilter == "enum")
+					{
+						FField* Property = InBindingChain.Last().Field.ToField();
+						FChooserEnumPropertyBinding* EnumPropertyValue = static_cast<FChooserEnumPropertyBinding*>(PropertyValue);
+						
+						if (const FEnumProperty* EnumProperty = CastField<const FEnumProperty>(Property))
+						{
+							EnumPropertyValue->Enum = EnumProperty->GetEnum();
+						}
+						else if (const FByteProperty* ByteProperty = CastField<const FByteProperty>(Property))
+						{
+							EnumPropertyValue->Enum = ByteProperty->Enum;
+						}
+					}
+
 					PropertyHandle->NotifyPostChange(EPropertyChangeType::ValueSet);
 				}
 			}
diff --git a/Engine/Plugins/Experimental/Chooser/Source/ChooserEditor/Private/PropertyAccessChainCustomization.h b/Engine/Plugins/Experimental/Chooser/Source/ChooserEditor/Private/PropertyAccessChainCustomization.h
--- a/Engine/Plugins/Experimental/Chooser/Source/ChooserEditor/Private/PropertyAccessChainCustomization.h
+++ b/Engine/Plugins/Experimental/Chooser/Source/ChooserEditor/Private/PropertyAccessChainCustomization.h
@@ -3,7 +3,6 @@
 #pragma once
 
 #include "IPropertyTypeCustomization.h"
-#include "PropertyEditorModule.h"
 
 namespace UE::ChooserEditor
 {
diff --git a/Engine/Plugins/Experimental/Chooser/Source/ChooserEditor/Public/ContextPropertyWidget.h b/Engine/Plugins/Experimental/Chooser/Source/ChooserEditor/Public/ContextPropertyWidget.h
--- a/Engine/Plugins/Experimental/Chooser/Source/ChooserEditor/Public/ContextPropertyWidget.h
+++ b/Engine/Plugins/Experimental/Chooser/Source/ChooserEditor/Public/ContextPropertyWidget.h
@@ -26,7 +26,7 @@ namespace UE::ChooserEditor
 		});
 		
 		Args.CurrentBindingColor = MakeAttributeLambda([ContextProperty, BindingColor]() {
-			 if (ContextProperty!= nullptr && ContextProperty->PropertyBindingChain.Num() > 0)
+			 if (ContextProperty!= nullptr && ContextProperty->Binding.PropertyBindingChain.Num() > 0)
 			 {
 				return BindingColor;
 			 }
@@ -44,7 +44,7 @@ namespace UE::ChooserEditor
 				if (ContextProperty != nullptr)
 				{
 					const FScopedTransaction Transaction(NSLOCTEXT("ContextPropertyWidget", "Change Property Binding", "Change Property Binding"));
-					TransactionObject->Modify(true); // todo
+					TransactionObject->Modify(true);
 					ContextProperty->SetBinding(InBindingChain);
 				}
 			});
@@ -54,12 +54,12 @@ namespace UE::ChooserEditor
 					const FText Bind = NSLOCTEXT("ContextPropertyWidget", "Bind", "Bind");
 					FText CurrentValue = Bind;
 	
-					if (ContextProperty != nullptr && ContextProperty->PropertyBindingChain.Num()>0)
+					if (ContextProperty != nullptr && ContextProperty->Binding.PropertyBindingChain.Num()>0)
 					{
 						TArray<FText> BindingChainText;
-						BindingChainText.Reserve(ContextProperty->PropertyBindingChain.Num());
+						BindingChainText.Reserve(ContextProperty->Binding.PropertyBindingChain.Num());
 					 
-						for (const FName& Name : ContextProperty->PropertyBindingChain)
+						for (const FName& Name : ContextProperty->Binding.PropertyBindingChain)
 						{
 							BindingChainText.Add(FText::FromName(Name));
 						}
@@ -77,21 +77,21 @@ namespace UE::ChooserEditor
 	
 					if (ContextProperty != nullptr)
 					{
-						int BindingChainLength = ContextProperty->PropertyBindingChain.Num();
+						int BindingChainLength = ContextProperty->Binding.PropertyBindingChain.Num();
 						if (BindingChainLength > 0)
 						{
 							if (BindingChainLength == 1)
 							{
 								// single property, just use the property name
-								CurrentValue = FText::FromName(ContextProperty->PropertyBindingChain.Last());
+								CurrentValue = FText::FromName(ContextProperty->Binding.PropertyBindingChain.Last());
 							}
 							else
 							{
 								// for longer chains always show the last struct/object name, and the final property name (full path in tooltip)
 								CurrentValue = FText::Join(NSLOCTEXT("ContextPropertyWidget", "PropertyPathSeparator","."),
 									TArray<FText>({
-										FText::FromName(ContextProperty->PropertyBindingChain[BindingChainLength-2]),
-										FText::FromName(ContextProperty->PropertyBindingChain[BindingChainLength-1])
+										FText::FromName(ContextProperty->Binding.PropertyBindingChain[BindingChainLength-2]),
+										FText::FromName(ContextProperty->Binding.PropertyBindingChain[BindingChainLength-1])
 									}));
 							}
 						}
