Static analysis fix for EnumColumnEditor

diff --git a/Engine/Plugins/Experimental/Chooser/Source/ChooserEditor/Private/EnumColumnEditor.cpp b/Engine/Plugins/Experimental/Chooser/Source/ChooserEditor/Private/EnumColumnEditor.cpp
--- a/Engine/Plugins/Experimental/Chooser/Source/ChooserEditor/Private/EnumColumnEditor.cpp
+++ b/Engine/Plugins/Experimental/Chooser/Source/ChooserEditor/Private/EnumColumnEditor.cpp
@@ -93,18 +93,15 @@ public:
 				{
 					// this is a bit of an odd thing to do, but as we are polling for the "=" / "!=" button state anyway,
 					// also poll to see if the currently bound Enum has changed, and if it has, recreate the enum combo box.
-					if (const FEnumColumn* EnumColumnPointer = EnumColumn)
+					const UEnum* CurrentEnumSource = nullptr;
+					if (EnumColumn->InputValue.IsValid())
 					{
-						const UEnum* CurrentEnumSource = nullptr;
-						if (EnumColumnPointer->InputValue.IsValid())
-						{
-							CurrentEnumSource = EnumColumnPointer->InputValue.Get<FChooserParameterEnumBase>().GetEnum(); 
-						}
-						if (EnumSource != CurrentEnumSource)
-						{
-							EnumComboBorder->SetContent(CreateEnumComboBox());
-							EnumSource = EnumColumnPointer->InputValue.Get<FChooserParameterEnumBase>().GetEnum();
-						}
+						CurrentEnumSource = EnumColumn->InputValue.Get<FChooserParameterEnumBase>().GetEnum(); 
+					}
+					if (EnumSource != CurrentEnumSource)
+					{
+						EnumComboBorder->SetContent(CreateEnumComboBox());
+						EnumSource = EnumColumn->InputValue.Get<FChooserParameterEnumBase>().GetEnum();
 					}
 					
 					return (EnumColumn->RowValues.IsValidIndex(Row) && EnumColumn->RowValues[Row].CompareNotEqual ? LOCTEXT("Not Equal", "!=") : LOCTEXT("Equal", "="));
