"Any" support in chooser bool column

diff --git a/Engine/Plugins/Experimental/Chooser/Source/Chooser/Internal/BoolColumn.h b/Engine/Plugins/Experimental/Chooser/Source/Chooser/Internal/BoolColumn.h
--- a/Engine/Plugins/Experimental/Chooser/Source/Chooser/Internal/BoolColumn.h
+++ b/Engine/Plugins/Experimental/Chooser/Source/Chooser/Internal/BoolColumn.h
@@ -8,6 +8,14 @@
 #include "InstancedStruct.h"
 #include "BoolColumn.generated.h"
 
+UENUM()
+enum class EBoolColumnCellValue
+{
+	MatchFalse = 0,
+	MatchTrue = 1,
+	MatchAny = 2,
+};
+
 USTRUCT(DisplayName = "Bool Property Binding")
 struct CHOOSER_API FBoolContextProperty :  public FChooserParameterBoolBase
 {
@@ -65,21 +73,39 @@ struct CHOOSER_API FBoolColumn : public FChooserColumnBase
 	
 	UPROPERTY(EditAnywhere, Meta = (ExcludeBaseStruct, BaseStruct = "/Script/Chooser.ChooserParameterBoolBase"), Category = "Data")
 	FInstancedStruct InputValue;
+
+#if WITH_EDITORONLY_DATA
+	UPROPERTY()
+	TArray<bool> RowValues_DEPRECATED;
+#endif
 	
 	UPROPERTY(EditAnywhere, Category= "Data");
-	TArray<bool> RowValues; 
+	TArray<EBoolColumnCellValue> RowValuesWithAny; 
 	
 	virtual void Filter(const UObject* ContextObject, const TArray<uint32>& IndexListIn, TArray<uint32>& IndexListOut) const override;
 
 	virtual void PostLoad() override
 	{
+#if WITH_EDITORONLY_DATA
+		if (RowValues_DEPRECATED.Num() > 0)
+		{
+			RowValuesWithAny.SetNum(0,false);
+			RowValuesWithAny.Reserve(RowValues_DEPRECATED.Num());
+			for(bool Value : RowValues_DEPRECATED)
+			{
+				RowValuesWithAny.Add(Value ? EBoolColumnCellValue::MatchTrue : EBoolColumnCellValue::MatchFalse);
+			}
+			RowValues_DEPRECATED.SetNum(0);
+		}
+#endif
+		
 		if (InputValue.IsValid())
 		{
 			InputValue.GetMutable<FChooserParameterBase>().PostLoad();
 		}
 	}
 
-	CHOOSER_COLUMN_BOILERPLATE(FChooserParameterBoolBase);
+	CHOOSER_COLUMN_BOILERPLATE2(FChooserParameterBoolBase, RowValuesWithAny);
 };
 
 // deprecated class versions for converting old data
@@ -127,6 +153,12 @@ public:
 		{
 			InputValueInterface->ConvertToInstancedStruct(Column.InputValue);
 		}
-		Column.RowValues = RowValues;
+
+		Column.RowValuesWithAny.SetNum(0,false);
+		Column.RowValuesWithAny.Reserve(RowValues.Num());
+		for(bool Value : RowValues)
+		{
+			Column.RowValuesWithAny.Add(Value ? EBoolColumnCellValue::MatchTrue : EBoolColumnCellValue::MatchFalse);
+		}
 	}
 };
\ No newline at end of file
diff --git a/Engine/Plugins/Experimental/Chooser/Source/Chooser/Internal/EnumColumn.h b/Engine/Plugins/Experimental/Chooser/Source/Chooser/Internal/EnumColumn.h
--- a/Engine/Plugins/Experimental/Chooser/Source/Chooser/Internal/EnumColumn.h
+++ b/Engine/Plugins/Experimental/Chooser/Source/Chooser/Internal/EnumColumn.h
@@ -105,14 +105,9 @@ public:
 
 	virtual void Filter(const UObject* ContextObject, const TArray<uint32>& IndexListIn, TArray<uint32>& IndexListOut) const override;
 	
-	CHOOSER_COLUMN_BOILERPLATE_NoSetInputType(FChooserParameterEnumBase);
+	CHOOSER_COLUMN_BOILERPLATE(FChooserParameterEnumBase);
 
 #if WITH_EDITOR
-	virtual void SetInputType(const UScriptStruct* Type) override
-	{
-		InputValue.InitializeAs(Type);
-	};
-
 	virtual void PostLoad() override
 	{
 		Super::PostLoad();
diff --git a/Engine/Plugins/Experimental/Chooser/Source/Chooser/Private/BoolColumn.cpp b/Engine/Plugins/Experimental/Chooser/Source/Chooser/Private/BoolColumn.cpp
--- a/Engine/Plugins/Experimental/Chooser/Source/Chooser/Private/BoolColumn.cpp
+++ b/Engine/Plugins/Experimental/Chooser/Source/Chooser/Private/BoolColumn.cpp
@@ -51,9 +51,10 @@ void FBoolColumn::Filter(const UObject* ContextObject, const TArray<uint32>& Ind
 		
 		for (uint32 Index : IndexListIn)
 		{
-			if (RowValues.Num() > (int)Index)
+			if (RowValuesWithAny.Num() > (int)Index)
 			{
-				if (Result == RowValues[Index])
+				
+				if (RowValuesWithAny[Index] == EBoolColumnCellValue::MatchAny || Result == static_cast<bool>(RowValuesWithAny[Index]))
 				{
 					IndexListOut.Push(Index);
 				}
diff --git a/Engine/Plugins/Experimental/Chooser/Source/Chooser/Public/IChooserColumn.h b/Engine/Plugins/Experimental/Chooser/Source/Chooser/Public/IChooserColumn.h
--- a/Engine/Plugins/Experimental/Chooser/Source/Chooser/Public/IChooserColumn.h
+++ b/Engine/Plugins/Experimental/Chooser/Source/Chooser/Public/IChooserColumn.h
@@ -38,6 +38,7 @@ public:
 	virtual void SetOutputs(UObject* ContextObject, int RowIndex) const { }
 
 #if WITH_EDITOR
+	virtual FName RowValuesPropertyName() { return FName(); }
 	virtual void SetNumRows(uint32 NumRows) {}
 	virtual void DeleteRows(const TArray<uint32> & RowIndices) {}
 	virtual void MoveRow(int SourceIndex, int TargetIndex) {}
@@ -50,33 +51,33 @@ public:
 };
 
 #if WITH_EDITOR
-#define CHOOSER_COLUMN_BOILERPLATE_NoSetInputType(ParameterType) \
-	virtual void SetNumRows(uint32 NumRows) override { RowValues.SetNum(NumRows); }\
+#define CHOOSER_COLUMN_BOILERPLATE2(ParameterType, RowValuesProperty) \
+	virtual FName RowValuesPropertyName() override { return #RowValuesProperty; }\
+	virtual void SetNumRows(uint32 NumRows) override { RowValuesProperty.SetNum(NumRows); }\
 	virtual void DeleteRows(const TArray<uint32> & RowIndices )\
 	{\
 		for(uint32 Index : RowIndices)\
 		{\
-			RowValues.RemoveAt(Index);\
+			RowValuesProperty.RemoveAt(Index);\
 		}\
 	}\
 	virtual void MoveRow(int SourceRowIndex, int TargetRowIndex)\
 	{\
-		auto RowData = RowValues[SourceRowIndex];\
-    	RowValues.RemoveAt(SourceRowIndex);\
+		auto RowData = RowValuesProperty[SourceRowIndex];\
+    	RowValuesProperty.RemoveAt(SourceRowIndex);\
     	if (SourceRowIndex < TargetRowIndex) { TargetRowIndex--; }\
-    	RowValues.Insert(RowData, TargetRowIndex);\
+    	RowValuesProperty.Insert(RowData, TargetRowIndex);\
 	}\
 	virtual UScriptStruct* GetInputBaseType() const override { return ParameterType::StaticStruct(); };\
 	virtual const UScriptStruct* GetInputType() const override { return InputValue.IsValid() ? InputValue.GetScriptStruct() : nullptr; };\
-	virtual FChooserParameterBase* GetInputValue() override { return InputValue.IsValid() ? &InputValue.GetMutable<FChooserParameterBase>() : nullptr; };
-	
-#define CHOOSER_COLUMN_BOILERPLATE(ParameterType) \
-	CHOOSER_COLUMN_BOILERPLATE_NoSetInputType(ParameterType); \
+	virtual FChooserParameterBase* GetInputValue() override { return InputValue.IsValid() ? &InputValue.GetMutable<FChooserParameterBase>() : nullptr; };\
 	virtual void SetInputType(const UScriptStruct* Type) override { InputValue.InitializeAs(Type); };
 
+#define CHOOSER_COLUMN_BOILERPLATE(ParameterType) CHOOSER_COLUMN_BOILERPLATE2(ParameterType, RowValues)
+
 #else
 
+#define CHOOSER_COLUMN_BOILERPLATE2(ParameterType, RowValuesProperty)
 #define CHOOSER_COLUMN_BOILERPLATE(ParameterType)
-#define CHOOSER_COLUMN_BOILERPLATE_NoSetInputType(ParameterType)
 
 #endif
diff --git a/Engine/Plugins/Experimental/Chooser/Source/ChooserEditor/Private/BoolColumnEditor.cpp b/Engine/Plugins/Experimental/Chooser/Source/ChooserEditor/Private/BoolColumnEditor.cpp
--- a/Engine/Plugins/Experimental/Chooser/Source/ChooserEditor/Private/BoolColumnEditor.cpp
+++ b/Engine/Plugins/Experimental/Chooser/Source/ChooserEditor/Private/BoolColumnEditor.cpp
@@ -6,6 +6,7 @@
 #include "ContextPropertyWidget.h"
 #include "ChooserTableEditor.h"
 #include "ObjectChooserWidgetFactories.h"
+#include "Widgets/Input/SButton.h"
 #include "Widgets/Input/SCheckBox.h"
 #include "GraphEditorSettings.h"
 
@@ -17,22 +18,39 @@ namespace UE::ChooserEditor
 TSharedRef<SWidget> CreateBoolColumnWidget(UChooserTable* Chooser, FChooserColumnBase* Column, int Row)
 {
 	FBoolColumn* BoolColumn = static_cast<FBoolColumn*>(Column);
-
-	return SNew (SCheckBox)
-	.OnCheckStateChanged_Lambda([Chooser, BoolColumn,Row](ECheckBoxState State)
-	{
-		if (Row < BoolColumn->RowValues.Num())
-		{
-			const FScopedTransaction Transaction(LOCTEXT("Change Bool Value", "Change Bool Value"));
-			Chooser->Modify(true);
-			BoolColumn->RowValues[Row] = (State == ECheckBoxState::Checked);
-		}
-	})
-	.IsChecked_Lambda([BoolColumn, Row]()
-	{
-		const bool value = (Row < BoolColumn->RowValues.Num()) ? BoolColumn->RowValues[Row] : false;
-		return value ? ECheckBoxState::Checked : ECheckBoxState::Unchecked;
-	});
+	
+	return
+	SNew(SHorizontalBox)
+		+ SHorizontalBox::Slot().FillWidth(1)
+		+ SHorizontalBox::Slot().MaxWidth(100).HAlign(HAlign_Center)
+		[
+			SNew(SButton).ButtonStyle(FAppStyle::Get(),"FlatButton").HAlign(HAlign_Center)
+    				.Text_Lambda([Row, BoolColumn]()
+    				{
+    					if (!BoolColumn->RowValuesWithAny.IsValidIndex(Row))
+    					{
+    						return FText();
+    					}
+    					switch (BoolColumn->RowValuesWithAny[Row])
+    					{
+							case EBoolColumnCellValue::MatchAny: return LOCTEXT("Any","Any");
+							case EBoolColumnCellValue::MatchTrue: return LOCTEXT("Any","True");
+							case EBoolColumnCellValue::MatchFalse: return LOCTEXT("Any","False");
+    					}
+    					return FText();
+    				})
+    				.OnClicked_Lambda([Chooser, Row, BoolColumn]()
+    				{
+    					if (BoolColumn->RowValuesWithAny.IsValidIndex(Row))
+    					{
+    						const FScopedTransaction Transaction(LOCTEXT("Edit Bool Cell Data", "Edit Bool Cell Data"));
+    						Chooser->Modify(true);
+    						BoolColumn->RowValuesWithAny[Row] = static_cast<EBoolColumnCellValue>((static_cast<int>(BoolColumn->RowValuesWithAny[Row]) + 1) % 3);
+    					}
+    					return FReply::Handled();
+    				})
+    	]
+		+ SHorizontalBox::Slot().FillWidth(1);
 }
 	
 TSharedRef<SWidget> CreateOutputBoolColumnWidget(UChooserTable* Chooser, FChooserColumnBase* Column, int Row)
diff --git a/Engine/Plugins/Experimental/Chooser/Source/ChooserEditor/Private/ChooserTableEditor.cpp b/Engine/Plugins/Experimental/Chooser/Source/ChooserEditor/Private/ChooserTableEditor.cpp
--- a/Engine/Plugins/Experimental/Chooser/Source/ChooserEditor/Private/ChooserTableEditor.cpp
+++ b/Engine/Plugins/Experimental/Chooser/Source/ChooserEditor/Private/ChooserTableEditor.cpp
@@ -1169,7 +1169,7 @@ void FChooserRowDetails::CustomizeDetails(IDetailLayoutBuilder& DetailBuilder)
 		{
 			FChooserColumnBase& Column = Chooser->ColumnsStructs[ColumnIndex].GetMutable<FChooserColumnBase>();
 			TSharedRef<FStructOnScope> StructOnScope = MakeShared<FStructOnScope>(Chooser->ColumnsStructs[ColumnIndex].GetScriptStruct(), reinterpret_cast<uint8*>(&Column));
-			TSharedPtr<IPropertyHandle> ColumnDataProperty = DetailBuilder.AddStructurePropertyData({StructOnScope}, "RowValues");
+			TSharedPtr<IPropertyHandle> ColumnDataProperty = DetailBuilder.AddStructurePropertyData({StructOnScope}, Column.RowValuesPropertyName());
 			uint32 NumElements = 0;
 			ColumnDataProperty->AsArray()->GetNumElements(NumElements);
 			if (Row->Row < (int)NumElements)
diff --git a/Engine/Plugins/Experimental/Chooser/Source/ProxyTable/ProxyTable.Build.cs b/Engine/Plugins/Experimental/Chooser/Source/ProxyTable/ProxyTable.Build.cs
--- a/Engine/Plugins/Experimental/Chooser/Source/ProxyTable/ProxyTable.Build.cs
+++ b/Engine/Plugins/Experimental/Chooser/Source/ProxyTable/ProxyTable.Build.cs
@@ -9,6 +9,7 @@ namespace UnrealBuildTool.Rules
 			PublicDependencyModuleNames.AddRange(
 				new string[]
 				{
+					"Chooser"
 					// ... add other public dependencies that you statically link with here ...
 				}
 			);
@@ -19,7 +20,6 @@ namespace UnrealBuildTool.Rules
 					"Core",
 					"CoreUObject",
 					"Engine",
-					"Chooser",
 					"StructUtils"
 					// ... add private dependencies that you statically link with here ...
 				}
