EvaluateChooser K2Node, and updates to Evaluate Chooser functions to not require casting the output.

diff --git a/Engine/Plugins/Experimental/Chooser/Chooser.uplugin b/Engine/Plugins/Experimental/Chooser/Chooser.uplugin
--- a/Engine/Plugins/Experimental/Chooser/Chooser.uplugin
+++ b/Engine/Plugins/Experimental/Chooser/Chooser.uplugin
@@ -32,6 +32,11 @@
 			"Type" : "Runtime",
 			"LoadingPhase" : "Default"
 		},
+		{
+			"Name": "ChooserUncooked",
+			"Type": "UncookedOnly",
+			"LoadingPhase": "Default"
+		},
 		{
 			"Name" : "ChooserEditor",
 			"Type" : "Editor",
@@ -42,10 +47,15 @@
 			"Type" : "Runtime",
 			"LoadingPhase" : "Default"
 		},
+		{
+			"Name": "ProxyTableUncooked",
+			"Type": "UncookedOnly",
+			"LoadingPhase": "Default"
+		},
 		{
 			"Name" : "ProxyTableEditor",
 			"Type" : "Editor",
 			"LoadingPhase" : "Default"
-		},
+		}
 	]
 }
diff --git a/Engine/Plugins/Experimental/Chooser/Source/Chooser/Private/Chooser.cpp b/Engine/Plugins/Experimental/Chooser/Source/Chooser/Private/Chooser.cpp
--- a/Engine/Plugins/Experimental/Chooser/Source/Chooser/Private/Chooser.cpp
+++ b/Engine/Plugins/Experimental/Chooser/Source/Chooser/Private/Chooser.cpp
@@ -134,26 +134,41 @@ UChooserFunctionLibrary::UChooserFunctionLibrary(const FObjectInitializer& Objec
 {
 }
 
-UObject* UChooserFunctionLibrary::EvaluateChooser(const UObject* ContextObject, const UChooserTable* Chooser)
+UObject* UChooserFunctionLibrary::EvaluateChooser(const UObject* ContextObject, const UChooserTable* Chooser, TSubclassOf<UObject> ObjectClass)
 {
 	UObject* Result = nullptr;
-	StaticEvaluateChooser(ContextObject, Chooser, FObjectChooserBase::FObjectChooserIteratorCallback::CreateLambda([&Result](UObject* InResult)
+	StaticEvaluateChooser(ContextObject, Chooser, FObjectChooserBase::FObjectChooserIteratorCallback::CreateLambda([&Result, ObjectClass](UObject* InResult)
 	{
-		Result = InResult;
-		return FObjectChooserBase::EIteratorStatus::Stop;
+		if (InResult && (ObjectClass == nullptr || InResult->IsA(ObjectClass)))
+		{
+			Result = InResult;
+			return FObjectChooserBase::EIteratorStatus::Stop;
+		}
+		return FObjectChooserBase::EIteratorStatus::Continue;
 	}));
 
 	return Result;
 }
 
-TArray<UObject*> UChooserFunctionLibrary::EvaluateChooserMulti(const UObject* ContextObject, const UChooserTable* Chooser)
+TArray<UObject*> UChooserFunctionLibrary::EvaluateChooserMulti(const UObject* ContextObject, const UChooserTable* Chooser, TSubclassOf<UObject> ObjectClass)
 {
 	TArray<UObject*> Result;
-	StaticEvaluateChooser(ContextObject, Chooser, FObjectChooserBase::FObjectChooserIteratorCallback::CreateLambda([&Result](UObject* InResult)
+	StaticEvaluateChooser(ContextObject, Chooser, FObjectChooserBase::FObjectChooserIteratorCallback::CreateLambda([&Result, ObjectClass](UObject* InResult)
 	{
-		Result.Add(InResult);
+		if (InResult && (ObjectClass == nullptr || InResult->IsA(ObjectClass)))
+		{
+			Result.Add(InResult);
+		}
 		return FObjectChooserBase::EIteratorStatus::Continue;
 	}));
 
+	for (int Index = 0; Index < Result.Num(); ++Index)
+	{
+		if (Result[Index] && !Result[Index]->IsA(ObjectClass))
+		{
+			Result[Index] = nullptr;
+		}
+	}
+
 	return Result;
 }
\ No newline at end of file
diff --git a/Engine/Plugins/Experimental/Chooser/Source/Chooser/Public/ChooserFunctionLibrary.h b/Engine/Plugins/Experimental/Chooser/Source/Chooser/Public/ChooserFunctionLibrary.h
--- a/Engine/Plugins/Experimental/Chooser/Source/Chooser/Public/ChooserFunctionLibrary.h
+++ b/Engine/Plugins/Experimental/Chooser/Source/Chooser/Public/ChooserFunctionLibrary.h
@@ -5,6 +5,7 @@
 #include "CoreMinimal.h"
 #include "Kismet/BlueprintFunctionLibrary.h"
 #include "Chooser.h"
+#include "Templates/SubclassOf.h"
 #include "ChooserFunctionLibrary.generated.h"
 
 /**
@@ -21,16 +22,18 @@ public:
 	*
 	* @param ContextObject			(in) An Object from which the parameters to the Chooser Table will be read
 	* @param ChooserTable			(in) The ChooserTable asset
+	* @param ObjectClass			(in) Expected type of result objects
 	*/
-	UFUNCTION(BlueprintPure, meta = (BlueprintThreadSafe), Category = "Animation")
-	static UObject* EvaluateChooser(const UObject* ContextObject, const UChooserTable* ChooserTable);
+	UFUNCTION(BlueprintPure, meta = (BlueprintThreadSafe, DeterminesOutputType = "ObjectClass"), Category = "Animation")
+	static UObject* EvaluateChooser(const UObject* ContextObject, const UChooserTable* ChooserTable, TSubclassOf<UObject> ObjectClass);
 
 	/**
     * Evaluate a chooser table and return the list of all selected UObjects
     *
     * @param ContextObject			(in) An Object from which the parameters to the Chooser Table will be read
     * @param ChooserTable			(in) The ChooserTable asset
+    * @param ObjectClass			(in) Expected type of result objects
     */
-    UFUNCTION(BlueprintPure, meta = (BlueprintThreadSafe), Category = "Animation")
-    static TArray<UObject*> EvaluateChooserMulti(const UObject* ContextObject, const UChooserTable* ChooserTable);
+    UFUNCTION(BlueprintPure, meta = (BlueprintThreadSafe, DeterminesOutputType = "ObjectClass"), Category = "Animation")
+    static TArray<UObject*> EvaluateChooserMulti(const UObject* ContextObject, const UChooserTable* ChooserTable, TSubclassOf<UObject> ObjectClass);
 };
diff --git a/Engine/Plugins/Experimental/Chooser/Source/ChooserEditor/ChooserEditor.Build.cs b/Engine/Plugins/Experimental/Chooser/Source/ChooserEditor/ChooserEditor.Build.cs
--- a/Engine/Plugins/Experimental/Chooser/Source/ChooserEditor/ChooserEditor.Build.cs
+++ b/Engine/Plugins/Experimental/Chooser/Source/ChooserEditor/ChooserEditor.Build.cs
@@ -6,6 +6,7 @@ namespace UnrealBuildTool.Rules
 	{
 		public ChooserEditor(ReadOnlyTargetRules Target) : base(Target)
 		{
+			PublicDependencyModuleNames.AddRange(new string[] {"Chooser"});
 			PrivateDependencyModuleNames.AddRange(
 				new string[]
 				{
@@ -29,7 +30,8 @@ namespace UnrealBuildTool.Rules
 					"GraphEditor",
 					"GameplayTags",
 					"GameplayTagsEditor",
-					"StructUtils"
+					"StructUtils",
+					"KismetCompiler"
 					// ... add private dependencies that you statically link with here ...
 				}
 			);
diff --git a/Engine/Plugins/Experimental/Chooser/Source/ChooserUncooked/ChooserUncooked.Build.cs b/Engine/Plugins/Experimental/Chooser/Source/ChooserUncooked/ChooserUncooked.Build.cs
new file mode 100644
--- /dev/null
+++ b/Engine/Plugins/Experimental/Chooser/Source/ChooserUncooked/ChooserUncooked.Build.cs
@@ -0,0 +1,29 @@
+// Copyright Epic Games, Inc. All Rights Reserved.
+
+namespace UnrealBuildTool.Rules
+{
+	public class ChooserUncooked : ModuleRules
+	{
+		public ChooserUncooked(ReadOnlyTargetRules Target) : base(Target)
+		{
+			PublicDependencyModuleNames.AddRange(new string[] {"Chooser"});
+			PrivateDependencyModuleNames.AddRange(
+				new string[]
+				{
+					"Core",
+					"CoreUObject",
+					"Engine",
+					"AssetTools",
+					"UnrealEd",
+					"AssetDefinition",
+					"Chooser",
+					"GameplayTags",
+					"StructUtils",
+					"BlueprintGraph",
+					"KismetCompiler"
+					// ... add private dependencies that you statically link with here ...
+				}
+			);
+		}
+	}
+}
\ No newline at end of file
diff --git a/Engine/Plugins/Experimental/Chooser/Source/ChooserUncooked/Private/ChooserUncookedModule.cpp b/Engine/Plugins/Experimental/Chooser/Source/ChooserUncooked/Private/ChooserUncookedModule.cpp
new file mode 100644
--- /dev/null
+++ b/Engine/Plugins/Experimental/Chooser/Source/ChooserUncooked/Private/ChooserUncookedModule.cpp
@@ -0,0 +1,18 @@
+// Copyright Epic Games, Inc. All Rights Reserved.
+
+#include "ChooserUncookedModule.h"
+
+namespace UE::Chooser
+{
+
+void FChooserUncookedModule::StartupModule()
+{
+}
+
+void FChooserUncookedModule::ShutdownModule()
+{
+}
+
+}
+
+IMPLEMENT_MODULE(UE::Chooser::FChooserUncookedModule, ChooserUncooked);
\ No newline at end of file
diff --git a/Engine/Plugins/Experimental/Chooser/Source/ChooserUncooked/Private/ChooserUncookedModule.h b/Engine/Plugins/Experimental/Chooser/Source/ChooserUncooked/Private/ChooserUncookedModule.h
new file mode 100644
--- /dev/null
+++ b/Engine/Plugins/Experimental/Chooser/Source/ChooserUncooked/Private/ChooserUncookedModule.h
@@ -0,0 +1,18 @@
+// Copyright Epic Games, Inc. All Rights Reserved.
+
+#pragma once
+
+#include "CoreMinimal.h"
+#include "Modules/ModuleInterface.h"
+
+namespace UE::Chooser
+{
+
+class FChooserUncookedModule : public IModuleInterface
+{
+private:
+	virtual void StartupModule() override;
+	virtual void ShutdownModule() override;
+};
+
+}
\ No newline at end of file
diff --git a/Engine/Plugins/Experimental/Chooser/Source/ChooserUncooked/Private/EvaluateChooserNode.cpp b/Engine/Plugins/Experimental/Chooser/Source/ChooserUncooked/Private/EvaluateChooserNode.cpp
new file mode 100644
--- /dev/null
+++ b/Engine/Plugins/Experimental/Chooser/Source/ChooserUncooked/Private/EvaluateChooserNode.cpp
@@ -0,0 +1,265 @@
+// Copyright Epic Games, Inc. All Rights Reserved.
+
+
+#include "EvaluateChooserNode.h"
+
+#include "BlueprintActionDatabaseRegistrar.h"
+#include "BlueprintNodeSpawner.h"
+#include "ChooserFunctionLibrary.h"
+#include "Containers/UnrealString.h"
+#include "EdGraph/EdGraph.h"
+#include "EdGraph/EdGraphSchema.h"
+#include "EdGraphSchema_K2.h"
+#include "EdGraphSchema_K2_Actions.h"
+#include "EditorCategoryUtils.h"
+#include "Internationalization/Internationalization.h"
+#include "K2Node_CallFunction.h"
+#include "K2Node_MakeArray.h"
+#include "Kismet2/BlueprintEditorUtils.h"
+#include "Kismet2/CompilerResultsLog.h"
+#include "KismetCompiler.h"
+#include "Misc/AssertionMacros.h"
+#include "Misc/CString.h"
+#include "Templates/Casts.h"
+#include "UObject/Class.h"
+#include "UObject/ObjectPtr.h"
+#include "UObject/UObjectGlobals.h"
+#include "UObject/UnrealType.h"
+#include "UObject/WeakObjectPtrTemplates.h"
+#include "K2Node_Self.h"
+
+
+#define LOCTEXT_NAMESPACE "EvaluateChooserNode"
+
+/////////////////////////////////////////////////////
+// UK2Node_EvaluateChooser
+
+UK2Node_EvaluateChooser::UK2Node_EvaluateChooser(const FObjectInitializer& ObjectInitializer)
+	: Super(ObjectInitializer)
+{
+	NodeTooltip = LOCTEXT("NodeTooltip", "Evaluates an Chooser Table, and returns the resulting Object or Objects.");
+
+	// need to unregister this somewhere - also this is pretty excessive
+	FCoreUObjectDelegates::OnObjectPropertyChanged.AddLambda([this](UObject* Object, struct FPropertyChangedEvent& Event)
+	{
+		if (Object == Chooser)
+		{
+			AllocateDefaultPins();
+		}
+	});
+}
+
+void UK2Node_EvaluateChooser::PreloadRequiredAssets()
+{
+	if (Chooser)
+	{
+		if (FLinkerLoad* ObjLinker = GetLinker())
+		{
+			ObjLinker->Preload(Chooser);
+		}
+	}
+    		
+	Super::PreloadRequiredAssets();
+}
+
+void UK2Node_EvaluateChooser::AllocateDefaultPins()
+{
+	Super::AllocateDefaultPins();
+	
+	UClass* ChooserResultType = UObject::StaticClass();
+
+   	if (Chooser && Chooser->OutputObjectType)
+   	{
+   		ChooserResultType = Chooser->OutputObjectType;
+   	}
+
+	if (UEdGraphPin* ResultPin = FindPin(TEXT("Result"), EGPD_Output))
+	{
+		ResultPin->PinType.PinSubCategoryObject = ChooserResultType;
+		ResultPin->PinType.ContainerType = (Mode == EEvaluateChooserMode::AllResults) ? EPinContainerType::Array : EPinContainerType::None;
+	}
+	else
+	{
+		UEdGraphNode::FCreatePinParams PinParams;
+		PinParams.ContainerType = (Mode == EEvaluateChooserMode::AllResults) ? EPinContainerType::Array : EPinContainerType::None;
+		PinParams.ValueTerminalType.TerminalCategory = UEdGraphSchema_K2::PC_Object;
+				
+		CreatePin(EGPD_Output, UEdGraphSchema_K2::PC_Object, ChooserResultType, TEXT("Result"), PinParams);
+	}
+}
+
+FText UK2Node_EvaluateChooser::GetNodeTitle(ENodeTitleType::Type TitleType) const
+{
+	if (Chooser)
+	{
+		return FText::FromString(Chooser->GetName());
+	}
+	else
+	{
+		return LOCTEXT("EvaluateChooser_Title", "Evaluate Chooser");
+	}
+}
+
+FText UK2Node_EvaluateChooser::GetPinDisplayName(const UEdGraphPin* Pin) const
+{
+	return FText::FromName(Pin->PinName);
+}
+
+void UK2Node_EvaluateChooser::PostEditChangeProperty(struct FPropertyChangedEvent& PropertyChangedEvent)
+{
+	if (PropertyChangedEvent.Property->GetName() == "Chooser" || PropertyChangedEvent.Property->GetName() == "Mode")
+	{
+		AllocateDefaultPins();
+	}
+	Super::PostEditChangeProperty(PropertyChangedEvent);
+}
+
+void UK2Node_EvaluateChooser::PinConnectionListChanged(UEdGraphPin* Pin)
+{
+
+	Modify();
+
+	FBlueprintEditorUtils::MarkBlueprintAsStructurallyModified(GetBlueprint());
+}
+
+void UK2Node_EvaluateChooser::PinDefaultValueChanged(UEdGraphPin* Pin)
+{
+	Super::PinDefaultValueChanged(Pin);
+}
+
+void UK2Node_EvaluateChooser::PinTypeChanged(UEdGraphPin* Pin)
+{
+
+	Super::PinTypeChanged(Pin);
+}
+
+FText UK2Node_EvaluateChooser::GetTooltipText() const
+{
+	return NodeTooltip;
+}
+
+void UK2Node_EvaluateChooser::PostReconstructNode()
+{
+	Super::PostReconstructNode();
+}
+
+void UK2Node_EvaluateChooser::ExpandNode(class FKismetCompilerContext& CompilerContext, UEdGraph* SourceGraph)
+{
+	Super::ExpandNode(CompilerContext, SourceGraph);
+
+	UEdGraphPin* ResultPin = FindPinChecked(TEXT("Result"));
+	if (ResultPin->HasAnyConnections())
+	{
+		UK2Node_CallFunction* CallFunction = CompilerContext.SpawnIntermediateNode<UK2Node_CallFunction>(this, SourceGraph);
+		CompilerContext.MessageLog.NotifyIntermediateObjectCreation(CallFunction, this);
+
+		if (Mode == EEvaluateChooserMode::AllResults)
+		{
+			CallFunction->SetFromFunction(UChooserFunctionLibrary::StaticClass()->FindFunctionByName(GET_MEMBER_NAME_CHECKED(UChooserFunctionLibrary, EvaluateChooserMulti)));
+		}
+		else
+		{
+			CallFunction->SetFromFunction(UChooserFunctionLibrary::StaticClass()->FindFunctionByName(GET_MEMBER_NAME_CHECKED(UChooserFunctionLibrary, EvaluateChooser)));
+		}
+		CallFunction->AllocateDefaultPins();
+
+		UK2Node_Self* SelfNode = CompilerContext.SpawnIntermediateNode<UK2Node_Self>(this, SourceGraph);
+		CompilerContext.MessageLog.NotifyIntermediateObjectCreation(SelfNode, this);
+		SelfNode->AllocateDefaultPins();
+
+		SelfNode->FindPin(TEXT("self"))->MakeLinkTo(CallFunction->FindPin(TEXT("ContextObject"))); // add null check + error
+		
+		UEdGraphPin* ChooserTablePin = CallFunction->FindPin(TEXT("ChooserTable"));
+		CallFunction->GetSchema()->TrySetDefaultObject(*ChooserTablePin, Chooser);
+
+		UEdGraphPin* OutputPin = CallFunction->GetReturnValuePin();
+
+		if (Chooser->OutputObjectType)
+		{
+			UEdGraphPin* OutputClassPin = CallFunction->FindPin(TEXT("ObjectClass"));
+			CallFunction->GetSchema()->TrySetDefaultObject(*OutputClassPin, Chooser->OutputObjectType);
+		}
+
+		CompilerContext.MovePinLinksToIntermediate(*ResultPin, *OutputPin);
+	}
+	
+	BreakAllNodeLinks();
+}
+
+UK2Node::ERedirectType UK2Node_EvaluateChooser::DoPinsMatchForReconstruction(const UEdGraphPin* NewPin, int32 NewPinIndex, const UEdGraphPin* OldPin, int32 OldPinIndex) const
+{
+	ERedirectType RedirectType = ERedirectType_None;
+
+	// if the pin names do match
+	if (NewPin->PinName.ToString().Equals(OldPin->PinName.ToString(), ESearchCase::CaseSensitive))
+	{
+		// Make sure we're not dealing with a menu node
+		UEdGraph* OuterGraph = GetGraph();
+		if( OuterGraph && OuterGraph->Schema )
+		{
+			const UEdGraphSchema_K2* K2Schema = Cast<const UEdGraphSchema_K2>(GetSchema());
+			if( !K2Schema || K2Schema->IsSelfPin(*NewPin) || K2Schema->ArePinTypesCompatible(OldPin->PinType, NewPin->PinType) )
+			{
+				RedirectType = ERedirectType_Name;
+			}
+			else
+			{
+				RedirectType = ERedirectType_None;
+			}
+		}
+	}
+	else
+	{
+		// try looking for a redirect if it's a K2 node
+		if (UK2Node* Node = Cast<UK2Node>(NewPin->GetOwningNode()))
+		{	
+			// if you don't have matching pin, now check if there is any redirect param set
+			TArray<FString> OldPinNames;
+			GetRedirectPinNames(*OldPin, OldPinNames);
+
+			FName NewPinName;
+			RedirectType = ShouldRedirectParam(OldPinNames, /*out*/ NewPinName, Node);
+
+			// make sure they match
+			if ((RedirectType != ERedirectType_None) && (!NewPin->PinName.ToString().Equals(NewPinName.ToString(), ESearchCase::CaseSensitive)))
+			{
+				RedirectType = ERedirectType_None;
+			}
+		}
+	}
+
+	return RedirectType;
+}
+
+bool UK2Node_EvaluateChooser::IsConnectionDisallowed(const UEdGraphPin* MyPin, const UEdGraphPin* OtherPin, FString& OutReason) const
+{
+	return Super::IsConnectionDisallowed(MyPin, OtherPin, OutReason);
+}
+
+void UK2Node_EvaluateChooser::GetMenuActions(FBlueprintActionDatabaseRegistrar& ActionRegistrar) const
+{
+	// actions get registered under specific object-keys; the idea is that 
+	// actions might have to be updated (or deleted) if their object-key is  
+	// mutated (or removed)... here we use the node's class (so if the node 
+	// type disappears, then the action should go with it)
+	UClass* ActionKey = GetClass();
+	// to keep from needlessly instantiating a UBlueprintNodeSpawner, first   
+	// check to make sure that the registrar is looking for actions of this type
+	// (could be regenerating actions for a specific asset, and therefore the 
+	// registrar would only accept actions corresponding to that asset)
+	if (ActionRegistrar.IsOpenForRegistration(ActionKey))
+	{
+		UBlueprintNodeSpawner* NodeSpawner = UBlueprintNodeSpawner::Create(GetClass());
+		check(NodeSpawner != nullptr);
+
+		ActionRegistrar.AddBlueprintAction(ActionKey, NodeSpawner);
+	}
+}
+
+FText UK2Node_EvaluateChooser::GetMenuCategory() const
+{
+	return FEditorCategoryUtils::GetCommonCategory(FCommonEditorCategory::Animation);
+}
+
+
+#undef LOCTEXT_NAMESPACE
diff --git a/Engine/Plugins/Experimental/Chooser/Source/ChooserUncooked/Private/EvaluateChooserNode.h b/Engine/Plugins/Experimental/Chooser/Source/ChooserUncooked/Private/EvaluateChooserNode.h
new file mode 100644
--- /dev/null
+++ b/Engine/Plugins/Experimental/Chooser/Source/ChooserUncooked/Private/EvaluateChooserNode.h
@@ -0,0 +1,68 @@
+// Copyright Epic Games, Inc. All Rights Reserved.
+
+
+#pragma once
+
+#include "CoreMinimal.h"
+#include "Chooser.h"
+#include "EdGraph/EdGraphNode.h"
+#include "HAL/Platform.h"
+#include "Internationalization/Text.h"
+#include "K2Node.h"
+#include "UObject/ObjectMacros.h"
+#include "UObject/UObjectGlobals.h"
+
+#include "EvaluateChooserNode.generated.h"
+
+UENUM()
+enum class EEvaluateChooserMode
+{
+	FirstResult,
+	AllResults
+};
+
+UCLASS(MinimalAPI)
+class UK2Node_EvaluateChooser : public UK2Node
+{
+	GENERATED_UCLASS_BODY()
+
+	//~ Begin UObject Interface
+	virtual void PostEditChangeProperty(struct FPropertyChangedEvent& PropertyChangedEvent) override;
+	//~ End UObject Interface
+
+	//~ Begin UEdGraphNode Interface.
+	virtual void AllocateDefaultPins() override;
+	virtual FText GetNodeTitle(ENodeTitleType::Type TitleType) const override;
+	virtual bool ShouldShowNodeProperties() const override { return true; }
+	virtual void PinConnectionListChanged(UEdGraphPin* Pin) override;
+	virtual void PinDefaultValueChanged(UEdGraphPin* Pin) override;
+	virtual void PinTypeChanged(UEdGraphPin* Pin) override;
+	virtual FText GetTooltipText() const override;
+	virtual FText GetPinDisplayName(const UEdGraphPin* Pin) const override;
+	//~ End UEdGraphNode Interface.
+
+	//~ Begin UK2Node Interface.
+	virtual bool IsNodePure() const override { return true; }
+	virtual bool NodeCausesStructuralBlueprintChange() const override { return true; }
+	virtual void PostReconstructNode() override;
+	virtual void ExpandNode(class FKismetCompilerContext& CompilerContext, UEdGraph* SourceGraph) override;
+	virtual ERedirectType DoPinsMatchForReconstruction(const UEdGraphPin* NewPin, int32 NewPinIndex, const UEdGraphPin* OldPin, int32 OldPinIndex) const override;
+	virtual bool IsConnectionDisallowed(const UEdGraphPin* MyPin, const UEdGraphPin* OtherPin, FString& OutReason) const override;
+	virtual void GetMenuActions(FBlueprintActionDatabaseRegistrar& ActionRegistrar) const override;
+	virtual FText GetMenuCategory() const override;
+	virtual int32 GetNodeRefreshPriority() const override { return EBaseNodeRefreshPriority::Low_UsesDependentWildcard; }
+	virtual void PreloadRequiredAssets() override;
+    
+	//~ End UK2Node Interface.
+
+private:
+	UPROPERTY(EditAnywhere, Category = "Chooser")
+	TObjectPtr<UChooserTable> Chooser;
+	
+	UPROPERTY(EditAnywhere, Category = "Chooser")
+	EEvaluateChooserMode Mode;
+
+	/** Tooltip text for this node. */
+	FText NodeTooltip;
+};
+
diff --git a/Engine/Plugins/Experimental/Chooser/Source/ProxyTable/Private/ProxyTable.cpp b/Engine/Plugins/Experimental/Chooser/Source/ProxyTable/Private/ProxyTable.cpp
--- a/Engine/Plugins/Experimental/Chooser/Source/ProxyTable/Private/ProxyTable.cpp
+++ b/Engine/Plugins/Experimental/Chooser/Source/ProxyTable/Private/ProxyTable.cpp
@@ -83,9 +83,15 @@ bool FProxyTableContextProperty::GetValue(const UObject* ContextObject, const UP
 	return false;
 }
 
-UObject* UProxyTableFunctionLibrary::EvaluateProxyAsset(const UObject* ContextObject, const UProxyAsset* Proxy)
+UObject* UProxyTableFunctionLibrary::EvaluateProxyAsset(const UObject* ContextObject, const UProxyAsset* Proxy, TSubclassOf<UObject> ObjectClass)
 {
-	return FindProxyObject(Proxy, ContextObject);
+	UObject* Result = FindProxyObject(Proxy, ContextObject);
+	if (ObjectClass && Result && !Result->IsA(ObjectClass))
+	{
+		return nullptr;
+	}
+	return Result;
+	
 }
 
 // fallback for FName based Keys:
diff --git a/Engine/Plugins/Experimental/Chooser/Source/ProxyTable/Public/ProxyTableFunctionLibrary.h b/Engine/Plugins/Experimental/Chooser/Source/ProxyTable/Public/ProxyTableFunctionLibrary.h
--- a/Engine/Plugins/Experimental/Chooser/Source/ProxyTable/Public/ProxyTableFunctionLibrary.h
+++ b/Engine/Plugins/Experimental/Chooser/Source/ProxyTable/Public/ProxyTableFunctionLibrary.h
@@ -4,6 +4,7 @@
 
 #include "Kismet/BlueprintFunctionLibrary.h"
 #include "ProxyTable.h"
+#include "Templates/SubclassOf.h"
 #include "ProxyTableFunctionLibrary.generated.h"
 
 /**
@@ -30,8 +31,9 @@ public:
 	*
 	* @param ContextObject			(in) An Object from which the Proxy Table will be read, and parameters to any nested Chooser Tables that need to evaluate
 	* @param ProxyAsset				(in) The ProxyAsset asset
+	* @param ObjectClass			(in) Expected type of result objects
 	*/
-	UFUNCTION(BlueprintPure, meta = (BlueprintThreadSafe), Category = "Animation")
-	static UObject* EvaluateProxyAsset(const UObject* ContextObject, const UProxyAsset* Proxy);
+	UFUNCTION(BlueprintPure, meta = (BlueprintThreadSafe, DeterminesOutputType = "ObjectClass"), Category = "Animation")
+	static UObject* EvaluateProxyAsset(const UObject* ContextObject, const UProxyAsset* Proxy, TSubclassOf<UObject> ObjectClass);
 
 };
\ No newline at end of file
diff --git a/Engine/Plugins/Experimental/Chooser/Source/ProxyTableUncooked/Private/EvaluateProxyNode.cpp b/Engine/Plugins/Experimental/Chooser/Source/ProxyTableUncooked/Private/EvaluateProxyNode.cpp
new file mode 100644
--- /dev/null
+++ b/Engine/Plugins/Experimental/Chooser/Source/ProxyTableUncooked/Private/EvaluateProxyNode.cpp
@@ -0,0 +1,252 @@
+// Copyright Epic Games, Inc. All Rights Reserved.
+
+
+#include "EvaluateProxyNode.h"
+
+#include "BlueprintActionDatabaseRegistrar.h"
+#include "BlueprintNodeSpawner.h"
+#include "ProxyTableFunctionLibrary.h"
+#include "Containers/UnrealString.h"
+#include "EdGraph/EdGraph.h"
+#include "EdGraph/EdGraphSchema.h"
+#include "EdGraphSchema_K2.h"
+#include "EdGraphSchema_K2_Actions.h"
+#include "EditorCategoryUtils.h"
+#include "Internationalization/Internationalization.h"
+#include "K2Node_CallFunction.h"
+#include "K2Node_MakeArray.h"
+#include "Kismet2/BlueprintEditorUtils.h"
+#include "Kismet2/CompilerResultsLog.h"
+#include "KismetCompiler.h"
+#include "Misc/AssertionMacros.h"
+#include "Misc/CString.h"
+#include "Templates/Casts.h"
+#include "UObject/Class.h"
+#include "UObject/ObjectPtr.h"
+#include "UObject/UObjectGlobals.h"
+#include "UObject/UnrealType.h"
+#include "UObject/WeakObjectPtrTemplates.h"
+#include "K2Node_Self.h"
+
+
+#define LOCTEXT_NAMESPACE "EvaluateProxyNode"
+
+/////////////////////////////////////////////////////
+// UK2Node_EvaluateProxy
+
+UK2Node_EvaluateProxy::UK2Node_EvaluateProxy(const FObjectInitializer& ObjectInitializer)
+	: Super(ObjectInitializer)
+{
+	NodeTooltip = LOCTEXT("NodeTooltip", "Evaluates an Proxy Table, and returns the resulting Object or Objects.");
+
+	// need to unregister this somewhere - also this is pretty excessive
+	FCoreUObjectDelegates::OnObjectPropertyChanged.AddLambda([this](UObject* Object, struct FPropertyChangedEvent& Event)
+	{
+		if (Object == Proxy)
+		{
+			AllocateDefaultPins();
+		}
+	});
+}
+
+void UK2Node_EvaluateProxy::PreloadRequiredAssets()
+{
+	if (Proxy)
+	{
+		if (FLinkerLoad* ObjLinker = GetLinker())
+		{
+			ObjLinker->Preload(Proxy);
+		}
+	}
+    		
+	Super::PreloadRequiredAssets();
+}
+
+void UK2Node_EvaluateProxy::AllocateDefaultPins()
+{
+	Super::AllocateDefaultPins();
+	
+	UClass* ProxyResultType = UObject::StaticClass();
+
+   	if (Proxy && Proxy->Type)
+   	{
+   		ProxyResultType = Proxy->Type;
+   	}
+
+	if (UEdGraphPin* ResultPin = FindPin(TEXT("Result"), EGPD_Output))
+	{
+		ResultPin->PinType.PinSubCategoryObject = ProxyResultType;
+	}
+	else
+	{
+		CreatePin(EGPD_Output, UEdGraphSchema_K2::PC_Object, ProxyResultType, TEXT("Result"));
+	}
+}
+
+FText UK2Node_EvaluateProxy::GetNodeTitle(ENodeTitleType::Type TitleType) const
+{
+	if (Proxy)
+	{
+		return FText::FromString(Proxy->GetName());
+	}
+	else
+	{
+		return LOCTEXT("EvaluateProxy_Title", "Evaluate Proxy");
+	}
+}
+
+FText UK2Node_EvaluateProxy::GetPinDisplayName(const UEdGraphPin* Pin) const
+{
+	return FText::FromName(Pin->PinName);
+}
+
+void UK2Node_EvaluateProxy::PostEditChangeProperty(struct FPropertyChangedEvent& PropertyChangedEvent)
+{
+	if (PropertyChangedEvent.Property->GetName() == "Proxy")
+	{
+		AllocateDefaultPins();
+	}
+	Super::PostEditChangeProperty(PropertyChangedEvent);
+}
+
+void UK2Node_EvaluateProxy::PinConnectionListChanged(UEdGraphPin* Pin)
+{
+	Modify();
+
+	FBlueprintEditorUtils::MarkBlueprintAsStructurallyModified(GetBlueprint());
+}
+
+void UK2Node_EvaluateProxy::PinDefaultValueChanged(UEdGraphPin* Pin)
+{
+	Super::PinDefaultValueChanged(Pin);
+}
+
+void UK2Node_EvaluateProxy::PinTypeChanged(UEdGraphPin* Pin)
+{
+
+	Super::PinTypeChanged(Pin);
+}
+
+FText UK2Node_EvaluateProxy::GetTooltipText() const
+{
+	return NodeTooltip;
+}
+
+void UK2Node_EvaluateProxy::PostReconstructNode()
+{
+	Super::PostReconstructNode();
+}
+
+void UK2Node_EvaluateProxy::ExpandNode(class FKismetCompilerContext& CompilerContext, UEdGraph* SourceGraph)
+{
+	Super::ExpandNode(CompilerContext, SourceGraph);
+
+	UEdGraphPin* ResultPin = FindPinChecked(TEXT("Result"));
+	if (ResultPin->HasAnyConnections())
+	{
+		UK2Node_CallFunction* CallFunction = CompilerContext.SpawnIntermediateNode<UK2Node_CallFunction>(this, SourceGraph);
+		CompilerContext.MessageLog.NotifyIntermediateObjectCreation(CallFunction, this);
+
+		CallFunction->SetFromFunction(UProxyTableFunctionLibrary::StaticClass()->FindFunctionByName(GET_MEMBER_NAME_CHECKED(UProxyTableFunctionLibrary, EvaluateProxyAsset)));
+		CallFunction->AllocateDefaultPins();
+
+		UK2Node_Self* SelfNode = CompilerContext.SpawnIntermediateNode<UK2Node_Self>(this, SourceGraph);
+		CompilerContext.MessageLog.NotifyIntermediateObjectCreation(SelfNode, this);
+		SelfNode->AllocateDefaultPins();
+
+		SelfNode->FindPin(TEXT("self"))->MakeLinkTo(CallFunction->FindPin(TEXT("ContextObject"))); // add null check + error
+		
+		UEdGraphPin* ProxyAssetPin = CallFunction->FindPin(TEXT("Proxy"));
+		CallFunction->GetSchema()->TrySetDefaultObject(*ProxyAssetPin, Proxy);
+
+		UEdGraphPin* OutputPin = CallFunction->GetReturnValuePin();
+
+		if (Proxy->Type)
+		{
+			UEdGraphPin* OutputClassPin = CallFunction->FindPin(TEXT("ObjectClass"));
+			CallFunction->GetSchema()->TrySetDefaultObject(*OutputClassPin, Proxy->Type);
+		}
+
+		CompilerContext.MovePinLinksToIntermediate(*ResultPin, *OutputPin);
+	}
+	
+	BreakAllNodeLinks();
+}
+
+UK2Node::ERedirectType UK2Node_EvaluateProxy::DoPinsMatchForReconstruction(const UEdGraphPin* NewPin, int32 NewPinIndex, const UEdGraphPin* OldPin, int32 OldPinIndex) const
+{
+	ERedirectType RedirectType = ERedirectType_None;
+
+	// if the pin names do match
+	if (NewPin->PinName.ToString().Equals(OldPin->PinName.ToString(), ESearchCase::CaseSensitive))
+	{
+		// Make sure we're not dealing with a menu node
+		UEdGraph* OuterGraph = GetGraph();
+		if( OuterGraph && OuterGraph->Schema )
+		{
+			const UEdGraphSchema_K2* K2Schema = Cast<const UEdGraphSchema_K2>(GetSchema());
+			if( !K2Schema || K2Schema->IsSelfPin(*NewPin) || K2Schema->ArePinTypesCompatible(OldPin->PinType, NewPin->PinType) )
+			{
+				RedirectType = ERedirectType_Name;
+			}
+			else
+			{
+				RedirectType = ERedirectType_None;
+			}
+		}
+	}
+	else
+	{
+		// try looking for a redirect if it's a K2 node
+		if (UK2Node* Node = Cast<UK2Node>(NewPin->GetOwningNode()))
+		{	
+			// if you don't have matching pin, now check if there is any redirect param set
+			TArray<FString> OldPinNames;
+			GetRedirectPinNames(*OldPin, OldPinNames);
+
+			FName NewPinName;
+			RedirectType = ShouldRedirectParam(OldPinNames, /*out*/ NewPinName, Node);
+
+			// make sure they match
+			if ((RedirectType != ERedirectType_None) && (!NewPin->PinName.ToString().Equals(NewPinName.ToString(), ESearchCase::CaseSensitive)))
+			{
+				RedirectType = ERedirectType_None;
+			}
+		}
+	}
+
+	return RedirectType;
+}
+
+bool UK2Node_EvaluateProxy::IsConnectionDisallowed(const UEdGraphPin* MyPin, const UEdGraphPin* OtherPin, FString& OutReason) const
+{
+	return Super::IsConnectionDisallowed(MyPin, OtherPin, OutReason);
+}
+
+void UK2Node_EvaluateProxy::GetMenuActions(FBlueprintActionDatabaseRegistrar& ActionRegistrar) const
+{
+	// actions get registered under specific object-keys; the idea is that 
+	// actions might have to be updated (or deleted) if their object-key is  
+	// mutated (or removed)... here we use the node's class (so if the node 
+	// type disappears, then the action should go with it)
+	UClass* ActionKey = GetClass();
+	// to keep from needlessly instantiating a UBlueprintNodeSpawner, first   
+	// check to make sure that the registrar is looking for actions of this type
+	// (could be regenerating actions for a specific asset, and therefore the 
+	// registrar would only accept actions corresponding to that asset)
+	if (ActionRegistrar.IsOpenForRegistration(ActionKey))
+	{
+		UBlueprintNodeSpawner* NodeSpawner = UBlueprintNodeSpawner::Create(GetClass());
+		check(NodeSpawner != nullptr);
+
+		ActionRegistrar.AddBlueprintAction(ActionKey, NodeSpawner);
+	}
+}
+
+FText UK2Node_EvaluateProxy::GetMenuCategory() const
+{
+	return FEditorCategoryUtils::GetCommonCategory(FCommonEditorCategory::Animation);
+}
+
+
+#undef LOCTEXT_NAMESPACE
diff --git a/Engine/Plugins/Experimental/Chooser/Source/ProxyTableUncooked/Private/EvaluateProxyNode.h b/Engine/Plugins/Experimental/Chooser/Source/ProxyTableUncooked/Private/EvaluateProxyNode.h
new file mode 100644
--- /dev/null
+++ b/Engine/Plugins/Experimental/Chooser/Source/ProxyTableUncooked/Private/EvaluateProxyNode.h
@@ -0,0 +1,58 @@
+// Copyright Epic Games, Inc. All Rights Reserved.
+
+
+#pragma once
+
+#include "CoreMinimal.h"
+#include "ProxyTable.h"
+#include "EdGraph/EdGraphNode.h"
+#include "HAL/Platform.h"
+#include "Internationalization/Text.h"
+#include "K2Node.h"
+#include "UObject/ObjectMacros.h"
+#include "UObject/UObjectGlobals.h"
+
+#include "EvaluateProxyNode.generated.h"
+
+UCLASS(MinimalAPI)
+class UK2Node_EvaluateProxy : public UK2Node
+{
+	GENERATED_UCLASS_BODY()
+
+	//~ Begin UObject Interface
+	virtual void PostEditChangeProperty(struct FPropertyChangedEvent& PropertyChangedEvent) override;
+	//~ End UObject Interface
+
+	//~ Begin UEdGraphNode Interface.
+	virtual void AllocateDefaultPins() override;
+	virtual FText GetNodeTitle(ENodeTitleType::Type TitleType) const override;
+	virtual bool ShouldShowNodeProperties() const override { return true; }
+	virtual void PinConnectionListChanged(UEdGraphPin* Pin) override;
+	virtual void PinDefaultValueChanged(UEdGraphPin* Pin) override;
+	virtual void PinTypeChanged(UEdGraphPin* Pin) override;
+	virtual FText GetTooltipText() const override;
+	virtual FText GetPinDisplayName(const UEdGraphPin* Pin) const override;
+	//~ End UEdGraphNode Interface.
+
+	//~ Begin UK2Node Interface.
+	virtual bool IsNodePure() const override { return true; }
+	virtual bool NodeCausesStructuralBlueprintChange() const override { return true; }
+	virtual void PostReconstructNode() override;
+	virtual void ExpandNode(class FKismetCompilerContext& CompilerContext, UEdGraph* SourceGraph) override;
+	virtual ERedirectType DoPinsMatchForReconstruction(const UEdGraphPin* NewPin, int32 NewPinIndex, const UEdGraphPin* OldPin, int32 OldPinIndex) const override;
+	virtual bool IsConnectionDisallowed(const UEdGraphPin* MyPin, const UEdGraphPin* OtherPin, FString& OutReason) const override;
+	virtual void GetMenuActions(FBlueprintActionDatabaseRegistrar& ActionRegistrar) const override;
+	virtual FText GetMenuCategory() const override;
+	virtual int32 GetNodeRefreshPriority() const override { return EBaseNodeRefreshPriority::Low_UsesDependentWildcard; }
+	virtual void PreloadRequiredAssets() override;
+    
+	//~ End UK2Node Interface.
+
+private:
+	UPROPERTY(EditAnywhere, Category = "Proxy")
+	TObjectPtr<UProxyAsset> Proxy;
+
+	/** Tooltip text for this node. */
+	FText NodeTooltip;
+};
+
diff --git a/Engine/Plugins/Experimental/Chooser/Source/ProxyTableUncooked/Private/ProxyTableUncookedModule.cpp b/Engine/Plugins/Experimental/Chooser/Source/ProxyTableUncooked/Private/ProxyTableUncookedModule.cpp
new file mode 100644
--- /dev/null
+++ b/Engine/Plugins/Experimental/Chooser/Source/ProxyTableUncooked/Private/ProxyTableUncookedModule.cpp
@@ -0,0 +1,18 @@
+// Copyright Epic Games, Inc. All Rights Reserved.
+
+#include "ProxyTableUncookedModule.h"
+
+namespace UE::ProxyTable
+{
+
+void FProxyTableUncookedModule::StartupModule()
+{
+}
+
+void FProxyTableUncookedModule::ShutdownModule()
+{
+}
+
+}
+
+IMPLEMENT_MODULE(UE::ProxyTable::FProxyTableUncookedModule, ProxyTableUncooked);
\ No newline at end of file
diff --git a/Engine/Plugins/Experimental/Chooser/Source/ProxyTableUncooked/Private/ProxyTableUncookedModule.h b/Engine/Plugins/Experimental/Chooser/Source/ProxyTableUncooked/Private/ProxyTableUncookedModule.h
new file mode 100644
--- /dev/null
+++ b/Engine/Plugins/Experimental/Chooser/Source/ProxyTableUncooked/Private/ProxyTableUncookedModule.h
@@ -0,0 +1,18 @@
+// Copyright Epic Games, Inc. All Rights Reserved.
+
+#pragma once
+
+#include "CoreMinimal.h"
+#include "Modules/ModuleInterface.h"
+
+namespace UE::ProxyTable
+{
+
+class FProxyTableUncookedModule : public IModuleInterface
+{
+private:
+	virtual void StartupModule() override;
+	virtual void ShutdownModule() override;
+};
+
+}
\ No newline at end of file
diff --git a/Engine/Plugins/Experimental/Chooser/Source/ProxyTableUncooked/ProxyTableUncooked.Build.cs b/Engine/Plugins/Experimental/Chooser/Source/ProxyTableUncooked/ProxyTableUncooked.Build.cs
new file mode 100644
--- /dev/null
+++ b/Engine/Plugins/Experimental/Chooser/Source/ProxyTableUncooked/ProxyTableUncooked.Build.cs
@@ -0,0 +1,29 @@
+// Copyright Epic Games, Inc. All Rights Reserved.
+
+namespace UnrealBuildTool.Rules
+{
+	public class ProxyTableUncooked : ModuleRules
+	{
+		public ProxyTableUncooked(ReadOnlyTargetRules Target) : base(Target)
+		{
+			PublicDependencyModuleNames.AddRange(new string[] {"ProxyTable"});
+			PrivateDependencyModuleNames.AddRange(
+				new string[]
+				{
+					"Core",
+					"CoreUObject",
+					"Engine",
+					"AssetTools",
+					"UnrealEd",
+					"AssetDefinition",
+					"ProxyTable",
+					"GameplayTags",
+					"StructUtils",
+					"BlueprintGraph",
+					"KismetCompiler"
+					// ... add private dependencies that you statically link with here ...
+				}
+			);
+		}
+	}
+}
\ No newline at end of file
