non-unity build fix

diff --git a/Engine/Plugins/Experimental/Chooser/Source/ChooserUncooked/Private/ChooserUncookedModule.cpp b/Engine/Plugins/Experimental/Chooser/Source/ChooserUncooked/Private/ChooserUncookedModule.cpp
--- a/Engine/Plugins/Experimental/Chooser/Source/ChooserUncooked/Private/ChooserUncookedModule.cpp
+++ b/Engine/Plugins/Experimental/Chooser/Source/ChooserUncooked/Private/ChooserUncookedModule.cpp
@@ -1,6 +1,7 @@
 // Copyright Epic Games, Inc. All Rights Reserved.
 
 #include "ChooserUncookedModule.h"
+#include "Modules/ModuleManager.h"
 
 namespace UE::Chooser
 {
diff --git a/Engine/Plugins/Experimental/Chooser/Source/ProxyTableUncooked/Private/ProxyTableUncookedModule.cpp b/Engine/Plugins/Experimental/Chooser/Source/ProxyTableUncooked/Private/ProxyTableUncookedModule.cpp
--- a/Engine/Plugins/Experimental/Chooser/Source/ProxyTableUncooked/Private/ProxyTableUncookedModule.cpp
+++ b/Engine/Plugins/Experimental/Chooser/Source/ProxyTableUncooked/Private/ProxyTableUncookedModule.cpp
@@ -1,6 +1,7 @@
 // Copyright Epic Games, Inc. All Rights Reserved.
 
 #include "ProxyTableUncookedModule.h"
+#include "Modules/ModuleManager.h"
 
 namespace UE::ProxyTable
 {
