Fix LOCTEXT error in BoolColumnEditor

diff --git a/Engine/Plugins/Experimental/Chooser/Source/ChooserEditor/Private/BoolColumnEditor.cpp b/Engine/Plugins/Experimental/Chooser/Source/ChooserEditor/Private/BoolColumnEditor.cpp
--- a/Engine/Plugins/Experimental/Chooser/Source/ChooserEditor/Private/BoolColumnEditor.cpp
+++ b/Engine/Plugins/Experimental/Chooser/Source/ChooserEditor/Private/BoolColumnEditor.cpp
@@ -34,8 +34,8 @@ TSharedRef<SWidget> CreateBoolColumnWidget(UChooserTable* Chooser, FChooserColum
     					switch (BoolColumn->RowValuesWithAny[Row])
     					{
 							case EBoolColumnCellValue::MatchAny: return LOCTEXT("Any","Any");
-							case EBoolColumnCellValue::MatchTrue: return LOCTEXT("Any","True");
-							case EBoolColumnCellValue::MatchFalse: return LOCTEXT("Any","False");
+							case EBoolColumnCellValue::MatchTrue: return LOCTEXT("True","True");
+							case EBoolColumnCellValue::MatchFalse: return LOCTEXT("False","False");
     					}
     					return FText();
     				})
