Proxy Table editor support for displaying entries imported from parent tables

diff --git a/Engine/Plugins/Experimental/Chooser/Source/ChooserEditor/Private/BoolColumnEditor.cpp b/Engine/Plugins/Experimental/Chooser/Source/ChooserEditor/Private/BoolColumnEditor.cpp
--- a/Engine/Plugins/Experimental/Chooser/Source/ChooserEditor/Private/BoolColumnEditor.cpp
+++ b/Engine/Plugins/Experimental/Chooser/Source/ChooserEditor/Private/BoolColumnEditor.cpp
@@ -74,9 +74,9 @@ TSharedRef<SWidget> CreateOutputBoolColumnWidget(UChooserTable* Chooser, FChoose
 	});
 }
 	
-TSharedRef<SWidget> CreateBoolPropertyWidget(UObject* TransactionObject, void* Value, UClass* ContextClass, UClass* ResultBaseClass)
+TSharedRef<SWidget> CreateBoolPropertyWidget(bool bReadOnly, UObject* TransactionObject, void* Value, UClass* ContextClass, UClass* ResultBaseClass)
 {
-	return CreatePropertyWidget<FBoolContextProperty>(TransactionObject, Value, ContextClass, GetDefault<UGraphEditorSettings>()->BooleanPinTypeColor);
+	return CreatePropertyWidget<FBoolContextProperty>(bReadOnly, TransactionObject, Value, ContextClass, GetDefault<UGraphEditorSettings>()->BooleanPinTypeColor);
 }
 	
 void RegisterBoolWidgets()
diff --git a/Engine/Plugins/Experimental/Chooser/Source/ChooserEditor/Private/ChooserTableEditor.cpp b/Engine/Plugins/Experimental/Chooser/Source/ChooserEditor/Private/ChooserTableEditor.cpp
--- a/Engine/Plugins/Experimental/Chooser/Source/ChooserEditor/Private/ChooserTableEditor.cpp
+++ b/Engine/Plugins/Experimental/Chooser/Source/ChooserEditor/Private/ChooserTableEditor.cpp
@@ -188,18 +188,10 @@ void FChooserTableEditor::NotifyPreChange(FProperty* PropertyAboutToChange)
 
 void FChooserTableEditor::NotifyPostChange(const FPropertyChangedEvent& PropertyChangedEvent, FProperty* PropertyThatChanged)
 {
-	if (PropertyThatChanged->GetNameCPP() == GET_MEMBER_NAME_STRING_CHECKED(UChooserTable, OutputObjectType)              // if you change the Output type, we need to update all the Results asset pickers
-	   || PropertyChangedEvent.Property->GetNameCPP() == GET_MEMBER_NAME_STRING_CHECKED(UChooserTable, ResultsStructs) )   // if you change the type of a Result, we need to update that widget
-	{
-		// rebuild all result widgets
-		UpdateTableRows();
-	}
-	if (PropertyThatChanged->GetNameCPP() == GET_MEMBER_NAME_STRING_CHECKED(UChooserTable, ContextObjectType)     // refresh columns property selector if the Context Object changed
-		|| PropertyThatChanged->GetNameCPP() == "InputValue" || PropertyThatChanged->GetNameCPP() == "Binding")   // refresh column widgets if the Input value or its binding was changed from the details panel
-	{
-		// rebuild all column header widgets
-		UpdateTableColumns();
-	}
+	// was previously only refreshing when known properties changed, however, changing the contents of an FInstancedStruct in the Details Panel (for example DefaultRowValue on a Column)
+	// will invalidate previously cached pointers to that FInstancedStruct, and so we need to refresh the columns every time anything is edited in the Details Panel
+	UpdateTableColumns();
+	UpdateTableRows();
 }
 
 	
@@ -440,13 +432,13 @@ public:
 			}
 			else if (ColumnName == Result) 
 			{
-				TSharedPtr<SWidget> ResultWidget = FObjectChooserWidgetFactories::CreateWidget(Chooser, FObjectChooserBase::StaticStruct(), Chooser->ResultsStructs[RowIndex->RowIndex].GetMutableMemory(), Chooser->ResultsStructs[RowIndex->RowIndex].GetScriptStruct(),Chooser->ContextObjectType, Chooser->OutputObjectType,
+				TSharedPtr<SWidget> ResultWidget = FObjectChooserWidgetFactories::CreateWidget(false, Chooser, FObjectChooserBase::StaticStruct(), Chooser->ResultsStructs[RowIndex->RowIndex].GetMutableMemory(), Chooser->ResultsStructs[RowIndex->RowIndex].GetScriptStruct(),Chooser->ContextObjectType, Chooser->OutputObjectType,
 				FOnStructPicked::CreateLambda([this, RowIndex=RowIndex->RowIndex](const UScriptStruct* ChosenStruct)
 				{
 					const FScopedTransaction Transaction(LOCTEXT("Change Row Result Type", "Change Row Result Type"));
-					Chooser->ResultsStructs[RowIndex].InitializeAs(ChosenStruct);
 					Chooser->Modify(true);
-					FObjectChooserWidgetFactories::CreateWidget(Chooser, FObjectChooserBase::StaticStruct(), Chooser->ResultsStructs[RowIndex].GetMutableMemory(), ChosenStruct, Chooser->ContextObjectType, Chooser->OutputObjectType, FOnStructPicked(), &CacheBorder);
+					Chooser->ResultsStructs[RowIndex].InitializeAs(ChosenStruct);
+					FObjectChooserWidgetFactories::CreateWidget(false, Chooser, FObjectChooserBase::StaticStruct(), Chooser->ResultsStructs[RowIndex].GetMutableMemory(), ChosenStruct, Chooser->ContextObjectType, Chooser->OutputObjectType, FOnStructPicked(), &CacheBorder);
 				}),
 				&CacheBorder
 				);
@@ -616,7 +608,7 @@ void FChooserTableEditor::UpdateTableColumns()
 		TSharedPtr<SWidget> HeaderWidget = nullptr;
 		if (FChooserParameterBase* InputValue = Column.GetInputValue())
 		{
-			HeaderWidget = FObjectChooserWidgetFactories::CreateWidget(Chooser, InputValue, Column.GetInputType(), Chooser->ContextObjectType, Chooser->OutputObjectType);
+			HeaderWidget = FObjectChooserWidgetFactories::CreateWidget(false, Chooser, InputValue, Column.GetInputType(), Chooser->ContextObjectType, Chooser->OutputObjectType);
 		}
 
 		HeaderRow->AddColumn(SHeaderRow::FColumn::FArguments()
@@ -1026,13 +1018,14 @@ void FChooserTableEditor::DeleteColumn(int Index)
 	}
 }
 
-TSharedRef<SWidget> CreateAssetWidget(UObject* TransactionObject, void* Value, UClass* ContextClass, UClass* ResultBaseClass)
+TSharedRef<SWidget> CreateAssetWidget(bool bReadOnly, UObject* TransactionObject, void* Value, UClass* ContextClass, UClass* ResultBaseClass)
 {
 	FAssetChooser* DIAsset = static_cast<FAssetChooser*>(Value);
 
 	UObject* Asset = DIAsset->Asset;
 	
 	return SNew(SObjectPropertyEntryBox)
+		.IsEnabled(!bReadOnly)
 		.AllowedClass(ResultBaseClass ? ResultBaseClass : UObject::StaticClass())
 		.ObjectPath_Lambda([DIAsset](){ return DIAsset->Asset ? DIAsset->Asset.GetPath() : "";})
 		.OnObjectChanged_Lambda([TransactionObject, DIAsset](const FAssetData& AssetData)
@@ -1043,13 +1036,14 @@ TSharedRef<SWidget> CreateAssetWidget(UObject* TransactionObject, void* Value, U
 		});
 }
 	
-TSharedRef<SWidget> CreateClassWidget(UObject* TransactionObject, void* Value, UClass* ContextClas, UClass* ResultBaseClass)
+TSharedRef<SWidget> CreateClassWidget(bool bReadOnly, UObject* TransactionObject, void* Value, UClass* ContextClas, UClass* ResultBaseClass)
 {
 	FClassChooser* ClassChooser = static_cast<FClassChooser*>(Value);
 
 	UClass* Class = ClassChooser->Class;
 
 	return SNew(SClassPropertyEntryBox)
+			.IsEnabled(!bReadOnly)
 			.MetaClass(ResultBaseClass ? ResultBaseClass : UObject::StaticClass())
 			.SelectedClass_Lambda([ClassChooser]()
 			{
@@ -1063,11 +1057,12 @@ TSharedRef<SWidget> CreateClassWidget(UObject* TransactionObject, void* Value, U
 			});
 }
 
-TSharedRef<SWidget> CreateEvaluateChooserWidget(UObject* TransactionObject, void* Value, UClass* ContextObject, UClass* ResultBaseClass)
+TSharedRef<SWidget> CreateEvaluateChooserWidget(bool bReadOnly, UObject* TransactionObject, void* Value, UClass* ContextObject, UClass* ResultBaseClass)
 {
 	FEvaluateChooser* EvaluateChooser = static_cast<FEvaluateChooser*>(Value);
 	
 	return SNew(SObjectPropertyEntryBox)
+		.IsEnabled(!bReadOnly)
 		.AllowedClass(UChooserTable::StaticClass())
 		.ObjectPath_Lambda([EvaluateChooser](){ return EvaluateChooser->Chooser ? EvaluateChooser->Chooser.GetPath() : "";})
 		.OnShouldFilterAsset_Lambda([ResultBaseClass](const FAssetData& InAssetData)
diff --git a/Engine/Plugins/Experimental/Chooser/Source/ChooserEditor/Private/EnumColumnEditor.cpp b/Engine/Plugins/Experimental/Chooser/Source/ChooserEditor/Private/EnumColumnEditor.cpp
--- a/Engine/Plugins/Experimental/Chooser/Source/ChooserEditor/Private/EnumColumnEditor.cpp
+++ b/Engine/Plugins/Experimental/Chooser/Source/ChooserEditor/Private/EnumColumnEditor.cpp
@@ -148,9 +148,9 @@ TSharedRef<SWidget> CreateEnumColumnWidget(UChooserTable* Chooser, FChooserColum
 	return SNew(SEnumCell).TransactionObject(Chooser).EnumColumn(EnumColumn).RowIndex(Row);
 }
 
-TSharedRef<SWidget> CreateEnumPropertyWidget(UObject* TransactionObject, void* Value, UClass* ContextClass, UClass* ResultBaseClass)
+TSharedRef<SWidget> CreateEnumPropertyWidget(bool bReadOnly, UObject* TransactionObject, void* Value, UClass* ContextClass, UClass* ResultBaseClass)
 {
-	return CreatePropertyWidget<FEnumContextProperty>(TransactionObject, Value, ContextClass, GetDefault<UGraphEditorSettings>()->BytePinTypeColor);
+	return CreatePropertyWidget<FEnumContextProperty>(bReadOnly, TransactionObject, Value, ContextClass, GetDefault<UGraphEditorSettings>()->BytePinTypeColor);
 }
 	
 void RegisterEnumWidgets()
diff --git a/Engine/Plugins/Experimental/Chooser/Source/ChooserEditor/Private/FloatRangeColumnEditor.cpp b/Engine/Plugins/Experimental/Chooser/Source/ChooserEditor/Private/FloatRangeColumnEditor.cpp
--- a/Engine/Plugins/Experimental/Chooser/Source/ChooserEditor/Private/FloatRangeColumnEditor.cpp
+++ b/Engine/Plugins/Experimental/Chooser/Source/ChooserEditor/Private/FloatRangeColumnEditor.cpp
@@ -12,9 +12,9 @@
 
 namespace UE::ChooserEditor
 {
-TSharedRef<SWidget> CreateFloatPropertyWidget(UObject* TransactionObject, void* Value, UClass* ContextClass, UClass* ResultBaseClass)
+TSharedRef<SWidget> CreateFloatPropertyWidget(bool bReadOnly, UObject* TransactionObject, void* Value, UClass* ContextClass, UClass* ResultBaseClass)
 {
-	return CreatePropertyWidget<FFloatContextProperty>(TransactionObject, Value, ContextClass, GetDefault<UGraphEditorSettings>()->FloatPinTypeColor);
+	return CreatePropertyWidget<FFloatContextProperty>(bReadOnly, TransactionObject, Value, ContextClass, GetDefault<UGraphEditorSettings>()->FloatPinTypeColor);
 }
 	
 TSharedRef<SWidget> CreateFloatRangeColumnWidget(UChooserTable* Chooser, FChooserColumnBase* Column, int Row)
diff --git a/Engine/Plugins/Experimental/Chooser/Source/ChooserEditor/Private/GameplayTagColumnEditor.cpp b/Engine/Plugins/Experimental/Chooser/Source/ChooserEditor/Private/GameplayTagColumnEditor.cpp
--- a/Engine/Plugins/Experimental/Chooser/Source/ChooserEditor/Private/GameplayTagColumnEditor.cpp
+++ b/Engine/Plugins/Experimental/Chooser/Source/ChooserEditor/Private/GameplayTagColumnEditor.cpp
@@ -49,9 +49,9 @@ TSharedRef<SWidget> CreateGameplayTagColumnWidget(UChooserTable* Chooser, FChoos
 	);
 }
 
-TSharedRef<SWidget> CreateGameplayTagPropertyWidget(UObject* TransactionObject, void* Value, UClass* ContextClass, UClass* ResultBaseClass)
+TSharedRef<SWidget> CreateGameplayTagPropertyWidget(bool bReadOnly, UObject* TransactionObject, void* Value, UClass* ContextClass, UClass* ResultBaseClass)
 {
-	return CreatePropertyWidget<FGameplayTagContextProperty>(TransactionObject, Value, ContextClass, GetDefault<UGraphEditorSettings>()->StructPinTypeColor);
+	return CreatePropertyWidget<FGameplayTagContextProperty>(bReadOnly, TransactionObject, Value, ContextClass, GetDefault<UGraphEditorSettings>()->StructPinTypeColor);
 }
 	
 void RegisterGameplayTagWidgets()
diff --git a/Engine/Plugins/Experimental/Chooser/Source/ChooserEditor/Private/ObjectChooserWidgetFactories.cpp b/Engine/Plugins/Experimental/Chooser/Source/ChooserEditor/Private/ObjectChooserWidgetFactories.cpp
--- a/Engine/Plugins/Experimental/Chooser/Source/ChooserEditor/Private/ObjectChooserWidgetFactories.cpp
+++ b/Engine/Plugins/Experimental/Chooser/Source/ChooserEditor/Private/ObjectChooserWidgetFactories.cpp
@@ -46,7 +46,7 @@ TSharedPtr<SWidget> FObjectChooserWidgetFactories::CreateColumnWidget(FChooserCo
 }
 	
 	
-TSharedPtr<SWidget> FObjectChooserWidgetFactories::CreateWidget(UObject* TransactionObject, void* Value, const UStruct* ValueType, UClass* ContextClass, UClass* ResultBaseClass)
+TSharedPtr<SWidget> FObjectChooserWidgetFactories::CreateWidget(bool bReadOnly, UObject* TransactionObject, void* Value, const UStruct* ValueType, UClass* ContextClass, UClass* ResultBaseClass)
 {
 	if (Value)
 	{
@@ -54,7 +54,7 @@ TSharedPtr<SWidget> FObjectChooserWidgetFactories::CreateWidget(UObject* Transac
 		{
 			if (FChooserWidgetCreator* Creator = ChooserWidgetCreators.Find(ValueType))
 			{
-				return (*Creator)(TransactionObject, Value, ContextClass, ResultBaseClass);
+				return (*Creator)(bReadOnly, TransactionObject, Value, ContextClass, ResultBaseClass);
 			}
 			ValueType = ValueType->GetSuperStruct();
 		}
@@ -63,9 +63,14 @@ TSharedPtr<SWidget> FObjectChooserWidgetFactories::CreateWidget(UObject* Transac
 	return nullptr;
 }
 
-TSharedPtr<SWidget> FObjectChooserWidgetFactories::CreateWidget(UObject* TransactionObject, const UScriptStruct* BaseType, void* Value, const UStruct* ValueType, UClass* ContextClass, UClass* ResultBaseClass, const FOnStructPicked& CreateClassCallback, TSharedPtr<SBorder>* InnerWidget)
+TSharedPtr<SWidget> FObjectChooserWidgetFactories::CreateWidget(bool bReadOnly, UObject* TransactionObject, const UScriptStruct* BaseType, void* Value, const UStruct* ValueType, UClass* ContextClass, UClass* ResultBaseClass, const FOnStructPicked& CreateClassCallback, TSharedPtr<SBorder>* InnerWidget)
 {
-	TSharedPtr<SWidget> LeftWidget = CreateWidget(TransactionObject, Value, ValueType, ContextClass, ResultBaseClass);
+	TSharedPtr<SWidget> LeftWidget = CreateWidget(bReadOnly, TransactionObject, Value, ValueType, ContextClass, ResultBaseClass);
+	if (bReadOnly)
+	{
+		// don't need the type selector dropdown when read only
+		return LeftWidget;
+	}
 	
 	if (!LeftWidget.IsValid())
 	{
diff --git a/Engine/Plugins/Experimental/Chooser/Source/ChooserEditor/Public/ContextPropertyWidget.h b/Engine/Plugins/Experimental/Chooser/Source/ChooserEditor/Public/ContextPropertyWidget.h
--- a/Engine/Plugins/Experimental/Chooser/Source/ChooserEditor/Public/ContextPropertyWidget.h
+++ b/Engine/Plugins/Experimental/Chooser/Source/ChooserEditor/Public/ContextPropertyWidget.h
@@ -9,12 +9,13 @@
 namespace UE::ChooserEditor
 {
 	template <typename PropertyType>
-	TSharedRef<SWidget> CreatePropertyWidget(UObject* TransactionObject, void* Value, UClass* ContextClass, const FLinearColor& BindingColor)
+	TSharedRef<SWidget> CreatePropertyWidget(bool bReadOnly, UObject* TransactionObject, void* Value, UClass* ContextClass, const FLinearColor& BindingColor)
 	{
 		PropertyType* ContextProperty = static_cast<PropertyType*>(Value);
 		
 		FPropertyBindingWidgetArgs Args;
 		Args.bAllowPropertyBindings = true;
+		// todo: there doesn't seem to be an "Enabled" Argument to assign !bReadOnly to
 		
 		Args.OnCanBindProperty = FOnCanBindProperty::CreateLambda([](FProperty* Property)
 		{
diff --git a/Engine/Plugins/Experimental/Chooser/Source/ChooserEditor/Public/ObjectChooserWidgetFactories.h b/Engine/Plugins/Experimental/Chooser/Source/ChooserEditor/Public/ObjectChooserWidgetFactories.h
--- a/Engine/Plugins/Experimental/Chooser/Source/ChooserEditor/Public/ObjectChooserWidgetFactories.h
+++ b/Engine/Plugins/Experimental/Chooser/Source/ChooserEditor/Public/ObjectChooserWidgetFactories.h
@@ -13,15 +13,15 @@ struct FChooserColumnBase;
 
 namespace UE::ChooserEditor
 {
-	typedef TFunction<TSharedRef<SWidget>(UObject* TransactionObject, void* Value, UClass* ContextClass, UClass* ResultBaseClass)> FChooserWidgetCreator;
+	typedef TFunction<TSharedRef<SWidget>(bool bReadOnly, UObject* TransactionObject, void* Value, UClass* ContextClass, UClass* ResultBaseClass)> FChooserWidgetCreator;
 	typedef TFunction<TSharedRef<SWidget>(UChooserTable* Chooser, FChooserColumnBase* Column, int Row)> FColumnWidgetCreator;
 
 	class CHOOSEREDITOR_API FObjectChooserWidgetFactories
 	{
 	public:
 
-		static TSharedPtr<SWidget> CreateWidget(UObject* TransactionObject, void* Value, const UStruct* ValueType, UClass* ContextClass, UClass* ResultBaseClass);
-		static TSharedPtr<SWidget> CreateWidget(UObject* TransactionObject, const UScriptStruct* BaseType, void* Value, const UStruct* ValueType, UClass* ContextClass, UClass* ResultBaseClass, const FOnStructPicked& CreateClassCallback, TSharedPtr<SBorder>* InnerWidget = nullptr);
+		static TSharedPtr<SWidget> CreateWidget(bool ReadOnly, UObject* TransactionObject, void* Value, const UStruct* ValueType, UClass* ContextClass, UClass* ResultBaseClass);
+		static TSharedPtr<SWidget> CreateWidget(bool ReadOnly, UObject* TransactionObject, const UScriptStruct* BaseType, void* Value, const UStruct* ValueType, UClass* ContextClass, UClass* ResultBaseClass, const FOnStructPicked& CreateClassCallback, TSharedPtr<SBorder>* InnerWidget = nullptr);
 		
 		static TSharedPtr<SWidget> CreateColumnWidget(FChooserColumnBase* Column, const UStruct* ColumnType, UChooserTable* Chooser, int Row);
 		
diff --git a/Engine/Plugins/Experimental/Chooser/Source/ProxyTableEditor/Private/ProxyTableEditor.cpp b/Engine/Plugins/Experimental/Chooser/Source/ProxyTableEditor/Private/ProxyTableEditor.cpp
--- a/Engine/Plugins/Experimental/Chooser/Source/ProxyTableEditor/Private/ProxyTableEditor.cpp
+++ b/Engine/Plugins/Experimental/Chooser/Source/ProxyTableEditor/Private/ProxyTableEditor.cpp
@@ -2,6 +2,7 @@
 
 #include "ProxyTableEditor.h"
 
+#include "AssetViewUtils.h"
 #include "Framework/Docking/TabManager.h"
 #include "Modules/ModuleManager.h"
 #include "Styling/AppStyle.h"
@@ -74,6 +75,7 @@ FProxyTableEditor::~FProxyTableEditor()
 {
 	GEditor->GetEditorSubsystem<UImportSubsystem>()->OnAssetPostImport.RemoveAll(this);
 	FCoreUObjectDelegates::OnObjectsReplaced.RemoveAll(this);
+	FCoreUObjectDelegates::OnObjectTransacted.RemoveAll(this);
 	
 	DetailsView.Reset();
 }
@@ -119,6 +121,7 @@ void FProxyTableEditor::InitEditor( const EToolkitMode::Type Mode, const TShared
 {
 	EditingObjects = ObjectsToEdit;
 	FCoreUObjectDelegates::OnObjectsReplaced.AddSP(this, &FProxyTableEditor::OnObjectsReplaced);
+	FCoreUObjectDelegates::OnObjectTransacted.AddSP(this, &FProxyTableEditor::OnObjectTransacted);
 
 	FPropertyEditorModule& PropertyEditorModule = FModuleManager::GetModuleChecked<FPropertyEditorModule>( "PropertyEditor" );
 	FDetailsViewArgs DetailsViewArgs;
@@ -326,15 +329,13 @@ class FProxyRowDragDropOp : public FDecoratedDragDropOp
 public:
 	DRAG_DROP_OPERATOR_TYPE(FWidgetTemplateDragDropOp, FDecoratedDragDropOp)
 
-	FProxyTableEditor* Editor;
-	uint32 RowIndex;
+	TSharedPtr<FProxyTableEditor::FProxyTableRow> Row;
 
 	/** Constructs the drag drop operation */
-	static TSharedRef<FProxyRowDragDropOp> New(FProxyTableEditor* InEditor, uint32 InRowIndex)
+	static TSharedRef<FProxyRowDragDropOp> New(TSharedPtr<FProxyTableEditor::FProxyTableRow> InRow)
 	{
 		TSharedRef<FProxyRowDragDropOp> Operation = MakeShareable(new FProxyRowDragDropOp());
-		Operation->Editor = InEditor;
-		Operation->RowIndex = InRowIndex;
+		Operation->Row = InRow;
 		Operation->DefaultHoverText = LOCTEXT("Proxy Row", "Proxy Row");
 		Operation->CurrentHoverText = Operation->DefaultHoverText;
 			
@@ -350,14 +351,12 @@ public:
 	SLATE_BEGIN_ARGS(SProxyRowHandle)
 	{}
 	SLATE_DEFAULT_SLOT(FArguments, Content)
-	SLATE_ARGUMENT(FProxyTableEditor*, ProxyEditor)
-	SLATE_ARGUMENT(uint32, RowIndex)
+	SLATE_ARGUMENT(TSharedPtr<FProxyTableEditor::FProxyTableRow>, Row)
 	SLATE_END_ARGS()
 
 	void Construct(const FArguments& InArgs)
 	{
-		ProxyEditor = InArgs._ProxyEditor;
-		RowIndex = InArgs._RowIndex;
+		Row = InArgs._Row;
 
 		ChildSlot
 		[
@@ -376,13 +375,13 @@ public:
 
 	FReply OnDragDetected(const FGeometry& MyGeometry, const FPointerEvent& MouseEvent) override
 	{
-		TSharedRef<FProxyRowDragDropOp> DragDropOp = FProxyRowDragDropOp::New(ProxyEditor, RowIndex);
+		TSharedRef<FProxyRowDragDropOp> DragDropOp = FProxyRowDragDropOp::New(Row);
 		return FReply::Handled().BeginDragDrop(DragDropOp);
 	}
 
 private:
 	FProxyTableEditor* ProxyEditor = nullptr;
-	uint32 RowIndex;
+	TSharedPtr<FProxyTableEditor::FProxyTableRow> Row;
 };
 
 
@@ -392,15 +391,16 @@ public:
 	SLATE_BEGIN_ARGS(SProxyTableRow) {}
 		/** The list item for this row */
 		SLATE_ARGUMENT(TSharedPtr<FProxyTableEditor::FProxyTableRow>, Entry)
-		SLATE_ARGUMENT(UProxyTable*, ProxyTable)
 		SLATE_ARGUMENT(FProxyTableEditor*, Editor)
 	SLATE_END_ARGS()
 
+	static constexpr int SpecialIndex_AddRow = -1;
+	static constexpr int SpecialIndex_InheritedFrom = -2;
+	static constexpr int SpecialIndex_NothingInheritedFrom = -3;
 
 	void Construct(const FArguments& Args, const TSharedRef<STableViewBase>& OwnerTableView)
 	{
-		RowIndex = Args._Entry;
-		ProxyTable = Args._ProxyTable;
+		Row = Args._Entry;
 		Editor = Args._Editor;
 
 		SMultiColumnTableRow<TSharedPtr<FProxyTableEditor::FProxyTableRow>>::Construct(
@@ -415,38 +415,41 @@ public:
 		static FName Handles = "Handles";
 		static FName Key = "Key";
 		static FName Value = "Value";
-		
-		if (RowIndex->RowIndex < ProxyTable->Entries.Num())
+
+		if (Row->RowIndex >=0)
 		{
 			if (ColumnName == Handles)
 			{
 				// row drag handle
-				return SNew(SProxyRowHandle).ProxyEditor(Editor).RowIndex(RowIndex->RowIndex);
+				return SNew(SProxyRowHandle).Row(Row);
 			}
 			else if (ColumnName == Value) 
 			{
 				UClass* ContextClass = nullptr;
 				UClass* ObjectType = nullptr;
-				if (ProxyTable->Entries[RowIndex->RowIndex].Proxy)
+				if (Row->ProxyTable->Entries[Row->RowIndex].Proxy)
 				{
-					ContextClass = ProxyTable->Entries[RowIndex->RowIndex].Proxy->ContextClass;
-					ObjectType = ProxyTable->Entries[RowIndex->RowIndex].Proxy->Type;
+					ContextClass = Row->ProxyTable->Entries[Row->RowIndex].Proxy->ContextClass;
+					ObjectType = Row->ProxyTable->Entries[Row->RowIndex].Proxy->Type;
 				}
-				
-				TSharedPtr<SWidget> ResultWidget = ChooserEditor::FObjectChooserWidgetFactories::CreateWidget(ProxyTable, FObjectChooserBase::StaticStruct(),
-					ProxyTable->Entries[RowIndex->RowIndex].ValueStruct.GetMutableMemory(),
-					ProxyTable->Entries[RowIndex->RowIndex].ValueStruct.GetScriptStruct(),
+				bool bReadOnly = Row->ProxyTable != Editor->GetProxyTable();
+				TSharedPtr<SWidget> ResultWidget = ChooserEditor::FObjectChooserWidgetFactories::CreateWidget( bReadOnly,
+					Row->ProxyTable, FObjectChooserBase::StaticStruct(),
+					Row->ProxyTable->Entries[Row->RowIndex].ValueStruct.GetMutableMemory(),
+					Row->ProxyTable->Entries[Row->RowIndex].ValueStruct.GetScriptStruct(),
 					ContextClass, ObjectType,
-					FOnStructPicked::CreateLambda([this, RowIndex=RowIndex->RowIndex](const UScriptStruct* ChosenStruct)
+					FOnStructPicked::CreateLambda([this](const UScriptStruct* ChosenStruct)
 					{
-						const FScopedTransaction Transaction(LOCTEXT("Change Value Type", "Change Value Type"));
-						ProxyTable->Entries[RowIndex].ValueStruct.InitializeAs(ChosenStruct);
-						ProxyTable->Modify(true);
-						ChooserEditor::FObjectChooserWidgetFactories::CreateWidget(ProxyTable, FObjectChooserBase::StaticStruct(),
-								ProxyTable->Entries[RowIndex].ValueStruct.GetMutableMemory(),
-								ProxyTable->Entries[RowIndex].ValueStruct.GetScriptStruct(),
-								ProxyTable->Entries[RowIndex].Proxy->ContextClass,
-								ProxyTable->Entries[RowIndex].Proxy->Type,
+						{
+							const FScopedTransaction Transaction(LOCTEXT("Change Value Type", "Change Value Type"));
+							Row->ProxyTable->Modify(true);
+							Row->ProxyTable->Entries[Row->RowIndex].ValueStruct.InitializeAs(ChosenStruct);
+						}
+						ChooserEditor::FObjectChooserWidgetFactories::CreateWidget(false, Row->ProxyTable, FObjectChooserBase::StaticStruct(),
+								Row->ProxyTable->Entries[Row->RowIndex].ValueStruct.GetMutableMemory(),
+								Row->ProxyTable->Entries[Row->RowIndex].ValueStruct.GetScriptStruct(),
+								Row->ProxyTable->Entries[Row->RowIndex].Proxy->ContextClass,
+								Row->ProxyTable->Entries[Row->RowIndex].Proxy->Type,
 								FOnStructPicked(), &CacheBorder);
 					}),
 					&CacheBorder
@@ -470,19 +473,20 @@ public:
 			}
 			else if (ColumnName == Key)
 			{
-				return SNew(SObjectPropertyEntryBox)
+				bool bReadOnly = Row->ProxyTable != Editor->GetProxyTable();
+				return SNew(SObjectPropertyEntryBox).IsEnabled(!bReadOnly)
 					.AllowedClass(UProxyAsset::StaticClass())
 					.ObjectPath_Lambda([this]()
 					{
-						return ProxyTable->Entries.Num() > RowIndex->RowIndex ?  ProxyTable->Entries[RowIndex->RowIndex].Proxy.GetPath() : FString();
+						return Row->ProxyTable->Entries.Num() > Row->RowIndex ?  Row->ProxyTable->Entries[Row->RowIndex].Proxy.GetPath() : FString();
 					})
 					.OnObjectChanged_Lambda([this](const FAssetData& AssetData)
 					{
-						if (ProxyTable->Entries.Num() > RowIndex->RowIndex)
+						if (Row->ProxyTable->Entries.Num() > Row->RowIndex)
 						{
 							const FScopedTransaction Transaction(LOCTEXT("Edit Proxy Asset", "Edit Proxy Asset"));
-							ProxyTable->Modify(true);
-							ProxyTable->Entries[RowIndex->RowIndex].Proxy = Cast<UProxyAsset>(AssetData.GetAsset());
+							Row->ProxyTable->Modify(true);
+							Row->ProxyTable->Entries[Row->RowIndex].Proxy = Cast<UProxyAsset>(AssetData.GetAsset());
 							
 							// ideally just need to rebuild the widget for the "Value" to update the UObject type filtering.
 							// For now just trigger a full refresh
@@ -491,12 +495,40 @@ public:
 					});
 			}
 		}
-		else if (RowIndex->RowIndex == ProxyTable->Entries.Num())
-        {
-			// on the row past the end, show an Add button in the result column
-			if (ColumnName == Key)
+		else
+		{
+			// special case rows
+			if (Row->RowIndex == SpecialIndex_AddRow)
 			{
-				return Editor->GetCreateRowComboButton().ToSharedRef();
+				// on the row past the end, show an Add button in the result column
+				if (ColumnName == Value)
+				{
+					return Editor->GetCreateRowComboButton().ToSharedRef();
+				}
+
+			}
+			else if (Row->RowIndex == SpecialIndex_InheritedFrom || Row->RowIndex == SpecialIndex_NothingInheritedFrom)
+			{
+				if (ColumnName == Key)
+				{
+					return SNew(SHorizontalBox)
+					+ SHorizontalBox::Slot().AutoWidth()
+					[
+						SNew(STextBlock)
+							.Text(Row->RowIndex == SpecialIndex_InheritedFrom ?  LOCTEXT("Inherited from ", "Inherited from ")
+																				: LOCTEXT("No rows inherited from ", "No rows inherited from ") )
+					]
+					+ SHorizontalBox::Slot().AutoWidth()
+					[
+						SNew(SHyperlink)
+							.Text(FText::FromString(Row->ProxyTable->GetName()))
+							.OnNavigate_Lambda([this]()
+							{
+								AssetViewUtils::OpenEditorForAsset(Row->ProxyTable);
+							})
+					]
+					+ SHorizontalBox::Slot().FillWidth(1);
+				}
 			}
 		}
 		return SNullWidget::NullWidget;
@@ -507,9 +539,12 @@ public:
 	{
 		if (TSharedPtr<FProxyRowDragDropOp> Operation = DragDropEvent.GetOperationAs<FProxyRowDragDropOp>())
 		{
-			bDragActive = true;
-			float Center = MyGeometry.Position.Y + MyGeometry.Size.Y;
-			bDropAbove = DragDropEvent.GetScreenSpacePosition().Y < Center;
+			if (Row->ProxyTable == Editor->GetProxyTable())
+			{
+				bDragActive = true;
+				float Center = MyGeometry.Position.Y + MyGeometry.Size.Y;
+				bDropAbove = DragDropEvent.GetScreenSpacePosition().Y < Center;
+			}
 		}
 	}
 	virtual void OnDragLeave(const FDragDropEvent& DragDropEvent) override
@@ -519,11 +554,14 @@ public:
 	
 	virtual FReply OnDragOver(const FGeometry& MyGeometry, const FDragDropEvent& DragDropEvent) override
 	{
-		if (TSharedPtr<FProxyRowDragDropOp> Operation = DragDropEvent.GetOperationAs<FProxyRowDragDropOp>())
+		if (Row->ProxyTable == Editor->GetProxyTable())
 		{
-			float Center = MyGeometry.AbsolutePosition.Y + MyGeometry.Size.Y/2;
-			bDropAbove = DragDropEvent.GetScreenSpacePosition().Y < Center;
-			return FReply::Handled();
+			if (TSharedPtr<FProxyRowDragDropOp> Operation = DragDropEvent.GetOperationAs<FProxyRowDragDropOp>())
+			{
+				float Center = MyGeometry.AbsolutePosition.Y + MyGeometry.Size.Y/2;
+				bDropAbove = DragDropEvent.GetScreenSpacePosition().Y < Center;
+				return FReply::Handled();
+			}
 		}
 		return FReply::Unhandled();
 	}
@@ -532,25 +570,26 @@ public:
 	{
 		if (TSharedPtr<FProxyRowDragDropOp> Operation = DragDropEvent.GetOperationAs<FProxyRowDragDropOp>())
 		{
-			if (ProxyTable == Operation->Editor->GetProxyTable())
+			if (Row->ProxyTable == Editor->GetProxyTable()) // only allow dropping on rows that are part of this actual table (not inherited entries)
 			{
-				if (bDropAbove)
+				if (Row->ProxyTable == Operation->Row->ProxyTable)
 				{
-					Editor->MoveRow(Operation->RowIndex, RowIndex->RowIndex);
+					// move row within a proxy table
+					Editor->MoveRow(Operation->Row->RowIndex, bDropAbove? Row->RowIndex : Row->RowIndex + 1);
+					return FReply::Handled();		
 				}
 				else
 				{
-					Editor->MoveRow(Operation->RowIndex, RowIndex->RowIndex+1);
+					Editor->InsertEntry(Operation->Row->ProxyTable->Entries[Operation->Row->RowIndex], bDropAbove? Row->RowIndex : Row->RowIndex + 1);
+					return FReply::Handled();
 				}
-				return FReply::Handled();		
 			}
 		}
 		return FReply::Unhandled();
 	}	
 
 private:
-	TSharedPtr<FProxyTableEditor::FProxyTableRow> RowIndex;
-	UProxyTable* ProxyTable;
+	TSharedPtr<FProxyTableEditor::FProxyTableRow> Row;
 	FProxyTableEditor* Editor;
 	TSharedPtr<SBorder> CacheBorder;
 	bool bDragActive = false;
@@ -563,7 +602,7 @@ TSharedRef<ITableRow> FProxyTableEditor::GenerateTableRow(TSharedPtr<FProxyTable
 	UProxyTable* ProxyTable = Cast<UProxyTable>(EditingObjects[0]);
 	
 	return SNew(SProxyTableRow, OwnerTable)
-		.Entry(InItem).ProxyTable(ProxyTable).Editor(this);
+		.Entry(InItem).Editor(this);
 }
 
 void FProxyTableEditor::SelectRootProperties()
@@ -574,6 +613,20 @@ void FProxyTableEditor::SelectRootProperties()
 		DetailsView->SetObjects( EditingObjects );
 	}
 }
+	
+void FProxyTableEditor::InsertEntry(FProxyEntry& Entry, int RowIndex)
+{
+	UProxyTable* Table = Cast<UProxyTable>(EditingObjects[0]);
+	RowIndex = FMath::Min(RowIndex,Table->Entries.Num());
+	
+	const FScopedTransaction Transaction(LOCTEXT("Move Row", "Move Row"));
+	
+	Table->Modify(true);
+
+	Table->Entries.Insert(Entry, RowIndex);
+
+	UpdateTableRows();
+}
 
 void FProxyTableEditor::MoveRow(int SourceRowIndex, int TargetRowIndex)
 {
@@ -712,20 +765,64 @@ TSharedRef<SDockTab> FProxyTableEditor::SpawnTableTab( const FSpawnTabArgs& Args
 		];
 }
 
+void FProxyTableEditor::AddInheritedRows(UProxyTable* ProxyTable)
+{
+	// add "Inherited from" header above inherited rows
+	TableRows.Add(MakeShared<FProxyTableRow>(SProxyTableRow::SpecialIndex_InheritedFrom, ProxyTable));
+	int NumRows = TableRows.Num();
+	for(int i =0; i<ProxyTable->Entries.Num(); i++)
+	{
+		// check if there's already an entry in TableRows for the same ProxyAsset
+		if (!TableRows.FindByPredicate(
+			[Proxy = ProxyTable->Entries[i].Proxy](const TSharedPtr<FProxyTableRow>& Row)
+			{
+				return Row->ProxyTable != nullptr
+					&& Row->RowIndex >= 0
+					&& Row->RowIndex < Row->ProxyTable->Entries.Num()
+				    && Proxy == Row->ProxyTable->Entries[Row->RowIndex].Proxy;
+			}))
+		{
+			TableRows.Add(MakeShared<FProxyTableRow>(i, ProxyTable));
+		}
+	}
+	if (NumRows == TableRows.Num())
+	{
+		// if none of the rows were actually inherited, change header to "No rows inherited from"
+		TableRows.Last()->RowIndex = SProxyTableRow::SpecialIndex_NothingInheritedFrom;
+	}
+
+	// recursively add inherited entries
+	for(UProxyTable* InheritedTable : ProxyTable->InheritEntriesFrom)
+	{
+		AddInheritedRows(InheritedTable);
+	}
+}
+
 void FProxyTableEditor::UpdateTableRows()
 {
 	UProxyTable* ProxyTable = Cast<UProxyTable>(EditingObjects[0]);
-	int32 OldNum = TableRows.Num();
-	int32 NewNum = ProxyTable->Entries.Num();
 
-	// Sync the TableRows array which drives the ui table to match the number of results.
+	TableRows.SetNum(0);
+
+	// add rows from this table
+	for(int i =0; i<ProxyTable->Entries.Num(); i++)
+	{
+		TableRows.Add(MakeShared<FProxyTableRow>(i, ProxyTable));
+	}
+
 	// Add 1 at the end, for the "Add Row" control
-	TableRows.SetNum(NewNum + 1);
+	TableRows.Add(MakeShared<FProxyTableRow>(SProxyTableRow::SpecialIndex_AddRow, ProxyTable));
 	
-	for(int32 i = OldNum; i < NewNum+1; i++)
+	if (ProxyTable->InheritEntriesFrom.Num()>0)
 	{
-		TableRows[i] = MakeShared<FProxyTableRow>(i);
+
+		// add imported rows in to the table
+		for(UProxyTable* InheritedTable : ProxyTable->InheritEntriesFrom)
+		{
+			AddInheritedRows(InheritedTable);
+		}
 	}
+	
 
 	if (TableView.IsValid())
 	{
@@ -733,6 +830,38 @@ void FProxyTableEditor::UpdateTableRows()
 	}
 }
 
+static bool IsDependency(UProxyTable* QueryTable, UProxyTable* Table)
+{
+	if (QueryTable == nullptr || Table == nullptr)
+	{
+		return false;
+	}
+	
+	for (UProxyTable* InheritFromTable : Table->InheritEntriesFrom)
+	{
+		if (InheritFromTable == QueryTable)
+		{
+			return true;
+		}
+		if (IsDependency(QueryTable, InheritFromTable))
+		{
+			return true;
+		}
+	}
+	return false;
+}
+	
+void FProxyTableEditor::OnObjectTransacted(UObject* InObject, const FTransactionObjectEvent& InTransactionObjectEvent)
+{
+	if (UProxyTable* ModifiedProxyTable = Cast<UProxyTable>(InObject))
+	{
+		if (IsDependency(ModifiedProxyTable, GetProxyTable()))
+		{
+			UpdateTableRows();
+		}
+	}	
+}
+
 void FProxyTableEditor::OnObjectsReplaced(const TMap<UObject*, UObject*>& ReplacementMap)
 {
 	bool bChangedAny = false;
@@ -781,11 +910,12 @@ TSharedRef<FProxyTableEditor> FProxyTableEditor::CreateEditor( const EToolkitMod
 
 /// Result widgets
 
-TSharedRef<SWidget> CreateLookupProxyWidget(UObject* TransactionObject, void* Value, UClass* ContextObject, UClass* ResultBaseClass)
+TSharedRef<SWidget> CreateLookupProxyWidget(bool bReadOnly, UObject* TransactionObject, void* Value, UClass* ContextObject, UClass* ResultBaseClass)
 {
 	FLookupProxy* LookupProxy = static_cast<FLookupProxy*>(Value);
 	
 	return SNew(SObjectPropertyEntryBox)
+		.IsEnabled(!bReadOnly)
 		.AllowedClass(UProxyAsset::StaticClass())
 		.ObjectPath_Lambda([LookupProxy](){ return LookupProxy->Proxy.GetPath();})
 		.OnShouldFilterAsset_Lambda([ResultBaseClass](const FAssetData& AssetData)
diff --git a/Engine/Plugins/Experimental/Chooser/Source/ProxyTableEditor/Private/ProxyTableEditor.h b/Engine/Plugins/Experimental/Chooser/Source/ProxyTableEditor/Private/ProxyTableEditor.h
--- a/Engine/Plugins/Experimental/Chooser/Source/ProxyTableEditor/Private/ProxyTableEditor.h
+++ b/Engine/Plugins/Experimental/Chooser/Source/ProxyTableEditor/Private/ProxyTableEditor.h
@@ -80,15 +80,22 @@ namespace UE::ProxyTableEditor
 	
 		struct FProxyTableRow
 		{
-			FProxyTableRow(int32 i) { RowIndex = i; }
+			FProxyTableRow(int32 InIndex, UProxyTable* InProxyTable) :
+				ProxyTable(InProxyTable),
+				RowIndex(InIndex)
+			{
+			}
+			
+			UProxyTable* ProxyTable;
 			int32 RowIndex;
 		};
 
 		void UpdateTableColumns();
 		void UpdateTableRows();
 		void MoveRow(int SourceRowIndex, int TargetIndex);
+		void InsertEntry(FProxyEntry& Entry, int RowIndex);
 	private:
-
+		void AddInheritedRows(UProxyTable* ProxyTable);
 		void SelectRootProperties();
 		void RegisterToolbar();
 		void BindCommands();
@@ -102,6 +109,7 @@ namespace UE::ProxyTableEditor
 
 		/** Called when objects need to be swapped out for new versions, like after a blueprint recompile. */
 		void OnObjectsReplaced(const TMap<UObject*, UObject*>& ReplacementMap);
+		void OnObjectTransacted(UObject* InObject, const FTransactionObjectEvent& InTransactionObjectEvent);
 
 		/** Details view */
 		TSharedPtr< class IDetailsView > DetailsView;
@@ -118,6 +126,7 @@ namespace UE::ProxyTableEditor
 
 		TArray<TObjectPtr<UProxyRowDetails>> SelectedRows;
 
+		TArray<TSharedPtr<FProxyTableRow>> MainTableRows;
 		TArray<TSharedPtr<FProxyTableRow>> TableRows;
 	
 		TSharedPtr<SComboButton> CreateRowComboButton;
