Chooser Column default row values

diff --git a/Engine/Plugins/Experimental/Chooser/Source/Chooser/Internal/BoolColumn.h b/Engine/Plugins/Experimental/Chooser/Source/Chooser/Internal/BoolColumn.h
--- a/Engine/Plugins/Experimental/Chooser/Source/Chooser/Internal/BoolColumn.h
+++ b/Engine/Plugins/Experimental/Chooser/Source/Chooser/Internal/BoolColumn.h
@@ -77,9 +77,12 @@ struct CHOOSER_API FBoolColumn : public FChooserColumnBase
 #if WITH_EDITORONLY_DATA
 	UPROPERTY()
 	TArray<bool> RowValues_DEPRECATED;
+
+	UPROPERTY(EditAnywhere, Category= "Data", DisplayName="DefaultRowValue");
+	EBoolColumnCellValue DefaultRowValue = EBoolColumnCellValue::MatchAny;
 #endif
 	
-	UPROPERTY(EditAnywhere, Category= "Data");
+	UPROPERTY(EditAnywhere, Category= "Data", DisplayName="RowValues");
 	TArray<EBoolColumnCellValue> RowValuesWithAny; 
 	
 	virtual void Filter(const UObject* ContextObject, const TArray<uint32>& IndexListIn, TArray<uint32>& IndexListOut) const override;
diff --git a/Engine/Plugins/Experimental/Chooser/Source/Chooser/Internal/EnumColumn.h b/Engine/Plugins/Experimental/Chooser/Source/Chooser/Internal/EnumColumn.h
--- a/Engine/Plugins/Experimental/Chooser/Source/Chooser/Internal/EnumColumn.h
+++ b/Engine/Plugins/Experimental/Chooser/Source/Chooser/Internal/EnumColumn.h
@@ -98,6 +98,11 @@ public:
 	UPROPERTY(EditAnywhere, Meta = (ExcludeBaseStruct, BaseStruct = "/Script/Chooser.ChooserParameterEnumBase"), Category = "Data")
 	FInstancedStruct InputValue;
 
+#if WITH_EDITORONLY_DATA
+	UPROPERTY(EditAnywhere, Category = "Data")
+	FChooserEnumRowData DefaultRowValue;
+#endif
+	
 	UPROPERTY(EditAnywhere, Category = "Data")
 	// array of results (cells for this column for each row in the table)
 	// should match the length of the Results array
diff --git a/Engine/Plugins/Experimental/Chooser/Source/Chooser/Internal/FloatRangeColumn.h b/Engine/Plugins/Experimental/Chooser/Source/Chooser/Internal/FloatRangeColumn.h
--- a/Engine/Plugins/Experimental/Chooser/Source/Chooser/Internal/FloatRangeColumn.h
+++ b/Engine/Plugins/Experimental/Chooser/Source/Chooser/Internal/FloatRangeColumn.h
@@ -19,7 +19,6 @@ struct CHOOSER_API FFloatContextProperty :  public FChooserParameterFloatBase
 	TArray<FName> PropertyBindingChain_DEPRECATED;
 	
 	UPROPERTY(EditAnywhere, Meta = (BindingType = "double", BindingColor = "FloatPinTypeColor"), Category = "Binding")
-	
 	FChooserPropertyBinding Binding;
 
 	virtual void PostLoad() override
@@ -78,6 +77,11 @@ struct CHOOSER_API FFloatRangeColumn : public FChooserColumnBase
 	UPROPERTY(EditAnywhere, Meta = (ExcludeBaseStruct, BaseStruct = "/Script/Chooser.ChooserParameterFloatBase"), Category = "Data")
 	FInstancedStruct InputValue;
 	
+#if WITH_EDITORONLY_DATA
+	UPROPERTY(EditAnywhere, Category="Data")
+	FChooserFloatRangeRowData DefaultRowValue;
+#endif
+	
 	UPROPERTY(EditAnywhere, Category="Data")
 	// array of results (cells for this column for each row in the table)
 	// should match the length of the Results array 
diff --git a/Engine/Plugins/Experimental/Chooser/Source/Chooser/Internal/GameplayTagColumn.h b/Engine/Plugins/Experimental/Chooser/Source/Chooser/Internal/GameplayTagColumn.h
--- a/Engine/Plugins/Experimental/Chooser/Source/Chooser/Internal/GameplayTagColumn.h
+++ b/Engine/Plugins/Experimental/Chooser/Source/Chooser/Internal/GameplayTagColumn.h
@@ -68,6 +68,11 @@ struct CHOOSER_API FGameplayTagColumn : public FChooserColumnBase
 	UPROPERTY(EditAnywhere, Category="Data")
 	EGameplayContainerMatchType	TagMatchType = EGameplayContainerMatchType::Any;
 
+#if WITH_EDITORONLY_DATA
+	UPROPERTY(EditAnywhere, Category=Runtime)
+	FGameplayTagContainer DefaultRowValue;
+#endif
+	
 	UPROPERTY(EditAnywhere, Category=Runtime)
 	// array of results (cells for this column for each row in the table)
 	// should match the length of the Results array 
diff --git a/Engine/Plugins/Experimental/Chooser/Source/Chooser/Internal/OutputBoolColumn.h b/Engine/Plugins/Experimental/Chooser/Source/Chooser/Internal/OutputBoolColumn.h
--- a/Engine/Plugins/Experimental/Chooser/Source/Chooser/Internal/OutputBoolColumn.h
+++ b/Engine/Plugins/Experimental/Chooser/Source/Chooser/Internal/OutputBoolColumn.h
@@ -21,6 +21,11 @@ struct CHOOSER_API FOutputBoolColumn : public FChooserColumnBase
 	UPROPERTY(EditAnywhere, Meta = (ExcludeBaseStruct, BaseStruct = "/Script/Chooser.ChooserParameterBoolBase"), Category = "Hidden")
 	FInstancedStruct InputValue;
 	
+#if WITH_EDITORONLY_DATA
+	UPROPERTY(EditAnywhere, Category=Runtime);
+	bool DefaultRowValue = 0;
+#endif
+	
 	UPROPERTY(EditAnywhere, Category=Runtime);
 	TArray<bool> RowValues; 
 
diff --git a/Engine/Plugins/Experimental/Chooser/Source/Chooser/Public/IChooserColumn.h b/Engine/Plugins/Experimental/Chooser/Source/Chooser/Public/IChooserColumn.h
--- a/Engine/Plugins/Experimental/Chooser/Source/Chooser/Public/IChooserColumn.h
+++ b/Engine/Plugins/Experimental/Chooser/Source/Chooser/Public/IChooserColumn.h
@@ -39,7 +39,7 @@ public:
 
 #if WITH_EDITOR
 	virtual FName RowValuesPropertyName() { return FName(); }
-	virtual void SetNumRows(uint32 NumRows) {}
+	virtual void SetNumRows(int32 NumRows) {}
 	virtual void DeleteRows(const TArray<uint32> & RowIndices) {}
 	virtual void MoveRow(int SourceIndex, int TargetIndex) {}
 	
@@ -53,7 +53,18 @@ public:
 #if WITH_EDITOR
 #define CHOOSER_COLUMN_BOILERPLATE2(ParameterType, RowValuesProperty) \
 	virtual FName RowValuesPropertyName() override { return #RowValuesProperty; }\
-	virtual void SetNumRows(uint32 NumRows) override { RowValuesProperty.SetNum(NumRows); }\
+	virtual void SetNumRows(int32 NumRows) override\
+	{\
+		if (NumRows <= RowValuesProperty.Num())\
+		{\
+			RowValuesProperty.SetNum(NumRows);\
+		}\
+		else\
+		{\
+			while(RowValuesProperty.Num() < NumRows)\
+			RowValuesProperty.Add(DefaultRowValue);\
+		}\
+	}\
 	virtual void DeleteRows(const TArray<uint32> & RowIndices )\
 	{\
 		for(uint32 Index : RowIndices)\
