Proxy table editor: make imported entries show in a sub tree that can be collapsed

diff --git a/Engine/Plugins/Experimental/Chooser/Source/ChooserEditor/Private/ChooserTableEditor.cpp b/Engine/Plugins/Experimental/Chooser/Source/ChooserEditor/Private/ChooserTableEditor.cpp
--- a/Engine/Plugins/Experimental/Chooser/Source/ChooserEditor/Private/ChooserTableEditor.cpp
+++ b/Engine/Plugins/Experimental/Chooser/Source/ChooserEditor/Private/ChooserTableEditor.cpp
@@ -35,6 +35,7 @@
 #include "ScopedTransaction.h"
 #include "Widgets/Layout/SSeparator.h"
 #include "ChooserTableEditorCommands.h"
+#include "Widgets/Layout/SScrollBox.h"
 
 #define LOCTEXT_NAMESPACE "ChooserEditor"
 
@@ -915,7 +916,12 @@ TSharedRef<SDockTab> FChooserTableEditor::SpawnTableTab( const FSpawnTabArgs& Ar
 		.TabColorScale( GetTabColorScale() )
 		.OnCanCloseTab_Lambda([]() { return false; })
 		[
-			TableView.ToSharedRef()
+			
+			SNew(SScrollBox).Orientation(Orient_Horizontal)
+			+ SScrollBox::Slot()
+			[
+				TableView.ToSharedRef()
+			]
 		];
 }
 
diff --git a/Engine/Plugins/Experimental/Chooser/Source/ProxyTableEditor/Private/ProxyTableEditor.cpp b/Engine/Plugins/Experimental/Chooser/Source/ProxyTableEditor/Private/ProxyTableEditor.cpp
--- a/Engine/Plugins/Experimental/Chooser/Source/ProxyTableEditor/Private/ProxyTableEditor.cpp
+++ b/Engine/Plugins/Experimental/Chooser/Source/ProxyTableEditor/Private/ProxyTableEditor.cpp
@@ -15,7 +15,6 @@
 #include "Widgets/Input/SNumericEntryBox.h"
 #include "Widgets/Views/SListView.h"
 #include "Widgets/Input/SComboButton.h"
-#include "Widgets/Input/SEditableTextBox.h"
 #include "SAssetDropTarget.h"
 #include "SClassViewer.h"
 #include "SourceCodeNavigation.h"
@@ -29,11 +28,11 @@
 #include "ObjectChooserClassFilter.h"
 #include "ScopedTransaction.h"
 #include "ObjectChooserWidgetFactories.h"
-#include "ContextPropertyWidget.h"
 #include "IObjectChooser.h"
 #include "Widgets/Layout/SSeparator.h"
 #include "PropertyCustomizationHelpers.h"
 #include "ProxyTableEditorCommands.h"
+#include "Widgets/Layout/SScrollBox.h"
 
 #define LOCTEXT_NAMESPACE "ProxyTableEditor"
 
@@ -188,6 +187,8 @@ void FProxyTableEditor::NotifyPreChange(FProperty* PropertyAboutToChange)
 
 void FProxyTableEditor::NotifyPostChange(const FPropertyChangedEvent& PropertyChangedEvent, FProperty* PropertyThatChanged)
 {
+	// for all details panel changes, just refresh the table
+	UpdateTableRows();
 }
 
 	
@@ -350,12 +351,13 @@ class SProxyRowHandle : public SCompoundWidget
 public:
 	SLATE_BEGIN_ARGS(SProxyRowHandle)
 	{}
-	SLATE_DEFAULT_SLOT(FArguments, Content)
+	SLATE_ARGUMENT(FProxyTableEditor*, Editor)
 	SLATE_ARGUMENT(TSharedPtr<FProxyTableEditor::FProxyTableRow>, Row)
 	SLATE_END_ARGS()
 
 	void Construct(const FArguments& InArgs)
 	{
+		ProxyEditor = InArgs._Editor;
 		Row = InArgs._Row;
 
 		ChildSlot
@@ -370,6 +372,7 @@ public:
 
 	FReply OnMouseButtonDown(const FGeometry& MyGeometry, const FPointerEvent& MouseEvent) override
 	{
+		ProxyEditor->SelectRow(Row);
 		return FReply::Handled().DetectDrag(SharedThis(this), EKeys::LeftMouseButton);
 	};
 
@@ -396,7 +399,6 @@ public:
 
 	static constexpr int SpecialIndex_AddRow = -1;
 	static constexpr int SpecialIndex_InheritedFrom = -2;
-	static constexpr int SpecialIndex_NothingInheritedFrom = -3;
 
 	void Construct(const FArguments& Args, const TSharedRef<STableViewBase>& OwnerTableView)
 	{
@@ -421,7 +423,7 @@ public:
 			if (ColumnName == Handles)
 			{
 				// row drag handle
-				return SNew(SProxyRowHandle).Row(Row);
+				return SNew(SProxyRowHandle).Row(Row).Editor(Editor);
 			}
 			else if (ColumnName == Value) 
 			{
@@ -497,37 +499,22 @@ public:
 		}
 		else
 		{
-			// special case rows
+			// special case row for "Add Row" button
 			if (Row->RowIndex == SpecialIndex_AddRow)
 			{
 				// on the row past the end, show an Add button in the result column
 				if (ColumnName == Value)
 				{
-					return Editor->GetCreateRowComboButton().ToSharedRef();
-				}
-
-			}
-			else if (Row->RowIndex == SpecialIndex_InheritedFrom || Row->RowIndex == SpecialIndex_NothingInheritedFrom)
-			{
-				if (ColumnName == Key)
-				{
-					return SNew(SHorizontalBox)
-					+ SHorizontalBox::Slot().AutoWidth()
-					[
-						SNew(STextBlock)
-							.Text(Row->RowIndex == SpecialIndex_InheritedFrom ?  LOCTEXT("Inherited from ", "Inherited from ")
-																				: LOCTEXT("No rows inherited from ", "No rows inherited from ") )
-					]
-					+ SHorizontalBox::Slot().AutoWidth()
-					[
-						SNew(SHyperlink)
-							.Text(FText::FromString(Row->ProxyTable->GetName()))
-							.OnNavigate_Lambda([this]()
-							{
-								AssetViewUtils::OpenEditorForAsset(Row->ProxyTable);
-							})
-					]
-					+ SHorizontalBox::Slot().FillWidth(1);
+					return SNew(SOverlay)
+						+ SOverlay::Slot()
+						[
+							Editor->GetCreateRowComboButton().ToSharedRef()
+						]
+						+ SOverlay::Slot().VAlign(VAlign_Top)
+						[
+							SNew(SSeparator).SeparatorImage(FCoreStyle::Get().GetBrush("FocusRectangle"))
+							.Visibility_Lambda([this]() { return bDragActive ? EVisibility::Visible : EVisibility::Hidden; })
+						];
 				}
 			}
 		}
@@ -565,22 +552,35 @@ public:
 		}
 		return FReply::Unhandled();
 	}
-	
+
 	virtual FReply OnDrop(const FGeometry& MyGeometry, const FDragDropEvent& DragDropEvent) override
 	{
 		if (TSharedPtr<FProxyRowDragDropOp> Operation = DragDropEvent.GetOperationAs<FProxyRowDragDropOp>())
 		{
 			if (Row->ProxyTable == Editor->GetProxyTable()) // only allow dropping on rows that are part of this actual table (not inherited entries)
 			{
+				int InsertIndex = Row->RowIndex;
+				if (InsertIndex < 0)
+				{
+					InsertIndex = Editor->GetProxyTable()->Entries.Num();
+				}
+				else
+				{
+					if (!bDropAbove)
+					{
+						InsertIndex++;
+					}
+				}
+				
 				if (Row->ProxyTable == Operation->Row->ProxyTable)
 				{
 					// move row within a proxy table
-					Editor->MoveRow(Operation->Row->RowIndex, bDropAbove? Row->RowIndex : Row->RowIndex + 1);
+					Editor->MoveRow(Operation->Row->RowIndex, InsertIndex);
 					return FReply::Handled();		
 				}
 				else
 				{
-					Editor->InsertEntry(Operation->Row->ProxyTable->Entries[Operation->Row->RowIndex], bDropAbove? Row->RowIndex : Row->RowIndex + 1);
+					Editor->InsertEntry(Operation->Row->ProxyTable->Entries[Operation->Row->RowIndex], InsertIndex);
 					return FReply::Handled();
 				}
 			}
@@ -600,9 +600,34 @@ private:
 TSharedRef<ITableRow> FProxyTableEditor::GenerateTableRow(TSharedPtr<FProxyTableRow> InItem, const TSharedRef<STableViewBase>& OwnerTable)
 {
 	UProxyTable* ProxyTable = Cast<UProxyTable>(EditingObjects[0]);
-	
-	return SNew(SProxyTableRow, OwnerTable)
-		.Entry(InItem).Editor(this);
+
+	if (InItem->RowIndex == SProxyTableRow::SpecialIndex_InheritedFrom)
+	{
+		return SNew(STableRow<TSharedRef<FProxyTableRow>>, OwnerTable)
+		[
+			SNew(SHorizontalBox)
+			+ SHorizontalBox::Slot().AutoWidth()
+			[
+				SNew(STextBlock)
+					.Text(InItem->Children.Num() > 0 ?  LOCTEXT("Inherited from ", "Inherited from ") : LOCTEXT("No rows inherited from ", "No rows inherited from ") )
+			]
+			+ SHorizontalBox::Slot().AutoWidth()
+			[
+				SNew(SHyperlink)
+					.Text(FText::FromString(InItem->ProxyTable->GetName()))
+					.OnNavigate_Lambda([InItem]()
+					{
+						AssetViewUtils::OpenEditorForAsset(InItem->ProxyTable);
+					})
+			]
+		];
+	}
+	else
+	{
+		return SNew(SProxyTableRow, OwnerTable)
+
+			.Entry(InItem).Editor(this);
+	}
 }
 
 void FProxyTableEditor::SelectRootProperties()
@@ -613,7 +638,30 @@ void FProxyTableEditor::SelectRootProperties()
 		DetailsView->SetObjects( EditingObjects );
 	}
 }
+
+void FProxyTableEditor::DeleteSelectedRows()
+{
+	UProxyTable* ProxyTable = Cast<UProxyTable>(EditingObjects[0]);
+	
+	const FScopedTransaction Transaction(LOCTEXT("Delete Row Transaction", "Delete Row"));
+	ProxyTable->Modify(true);
 	
+	// delete selected rows.
+	TArray<uint32> RowsToDelete;
+	for(auto& SelectedRow:SelectedRows)
+	{
+		RowsToDelete.Add(SelectedRow->Row);
+	}
+	// sort indices in reverse
+	RowsToDelete.Sort([](int32 A, int32 B){ return A>B; });
+	for(uint32 RowIndex : RowsToDelete)
+	{
+		ProxyTable->Entries.RemoveAt(RowIndex);
+	}
+	
+	UpdateTableRows();
+}
+
 void FProxyTableEditor::InsertEntry(FProxyEntry& Entry, int RowIndex)
 {
 	UProxyTable* Table = Cast<UProxyTable>(EditingObjects[0]);
@@ -623,6 +671,8 @@ void FProxyTableEditor::InsertEntry(FProxyEntry& Entry, int RowIndex)
 	
 	Table->Modify(true);
 
+	RowIndex = FMath::Clamp(RowIndex, 0, Table->Entries.Num());
+
 	Table->Entries.Insert(Entry, RowIndex);
 
 	UpdateTableRows();
@@ -646,9 +696,17 @@ void FProxyTableEditor::MoveRow(int SourceRowIndex, int TargetRowIndex)
 	Table->Entries.Insert(Entry, TargetRowIndex);
 
 	UpdateTableRows();
-	
 }
 
+void FProxyTableEditor::TreeViewExpansionChanged(TSharedPtr<FProxyTableEditor::FProxyTableRow> InItem, bool bShouldBeExpanded)
+{
+	if (InItem->RowIndex == SProxyTableRow::SpecialIndex_InheritedFrom)
+	{
+		ImportedTablesExpansionState.Add(InItem->ProxyTable, bShouldBeExpanded);
+	}
+}
+
+
 void FProxyTableEditor::UpdateTableColumns()
 {
 	HeaderRow->ClearColumns();
@@ -700,39 +758,24 @@ TSharedRef<SDockTab> FProxyTableEditor::SpawnTableTab( const FSpawnTabArgs& Args
 		});
 
 	HeaderRow = SNew(SHeaderRow);
-
-	UpdateTableColumns();
-	UpdateTableRows();
-
-	TableView = SNew(SListView<TSharedPtr<FProxyTableRow>>)
-    			.ListItemsSource(&TableRows)
-				.OnKeyDownHandler_Lambda([this](const FGeometry&, const FKeyEvent& Event)
+	
+	TableView = SNew(STreeView<TSharedPtr<FProxyTableRow>>)
+    			.TreeItemsSource(&TableRows)
+			.OnExpansionChanged(this, &FProxyTableEditor::TreeViewExpansionChanged)
+				.OnGetChildren_Lambda( [] (TSharedPtr<FProxyTableRow> Row, TArray<TSharedPtr<FProxyTableRow>>& OutChildren)
 				{
-					if (Event.GetKey() == EKeys::Delete)
+					OutChildren = Row->Children;
+				})
+				.OnKeyDownHandler_Lambda([this](const FGeometry&, const FKeyEvent& Event)
+				{	
+					if ( Event.GetKey() == EKeys::Delete)
 					{
-						const FScopedTransaction Transaction(LOCTEXT("Delete Row Transaction", "Delete Row"));
-						UProxyTable* ProxyTable = Cast<UProxyTable>(EditingObjects[0]);
-						ProxyTable->Modify(true);
-						// delete selected rows.
-						TArray<uint32> RowsToDelete;
-						for(auto& SelectedRow:SelectedRows)
-						{
-							RowsToDelete.Add(SelectedRow->Row);
-						}
-						// sort indices in reverse
-						RowsToDelete.Sort([](int32 A, int32 B){ return A>B; });
-						for(uint32 RowIndex : RowsToDelete)
-						{
-							ProxyTable->Entries.RemoveAt(RowIndex);
-						}
-
-						UpdateTableRows();
-						
+						UProxyTable* ProxyTable = GetProxyTable();
+						DeleteSelectedRows();
 						return FReply::Handled();
 					}
 					return FReply::Unhandled();
-				}
-				)
+				})
 				.OnSelectionChanged_Lambda([this](TSharedPtr<FProxyTableRow> SelectedItem,  ESelectInfo::Type SelectInfo)
 				{
 					if (SelectedItem)
@@ -755,41 +798,52 @@ TSharedRef<SDockTab> FProxyTableEditor::SpawnTableTab( const FSpawnTabArgs& Args
 				})
     			.OnGenerateRow_Raw(this, &FProxyTableEditor::GenerateTableRow)
 				.HeaderRow(HeaderRow);
+				
+
+	UpdateTableColumns();
+	UpdateTableRows();
 	
 	return SNew(SDockTab)
 		.Label( LOCTEXT("ProxtTableTitle", "Proxy Table") )
 		.TabColorScale( GetTabColorScale() )
 		.OnCanCloseTab_Lambda([]() { return false; })
 		[
-			TableView.ToSharedRef()
+			SNew(SScrollBox).Orientation(Orient_Horizontal)
+			+ SScrollBox::Slot()
+			[
+				TableView.ToSharedRef()
+			]
 		];
 }
 
 void FProxyTableEditor::AddInheritedRows(UProxyTable* ProxyTable)
 {
+	if (ProxyTable == nullptr || ReferencedProxyTables.Find(ProxyTable))
+	{
+		return;
+	}
+
+	// prevent infinite inheritance loops
+	ReferencedProxyTables.Add(ProxyTable);
+	
 	// add "Inherited from" header above inherited rows
-	TableRows.Add(MakeShared<FProxyTableRow>(SProxyTableRow::SpecialIndex_InheritedFrom, ProxyTable));
+	TSharedPtr<FProxyTableRow> ParentTableRow = MakeShared<FProxyTableRow>(SProxyTableRow::SpecialIndex_InheritedFrom, ProxyTable);
+	TableRows.Add(ParentTableRow);
+
+	bool* CachedExpansionState = ImportedTablesExpansionState.Find(ProxyTable);
+	bool Expansion = CachedExpansionState ? *CachedExpansionState : true;
+	TableView->SetItemExpansion(ParentTableRow, Expansion);
+
 	int NumRows = TableRows.Num();
 	for(int i =0; i<ProxyTable->Entries.Num(); i++)
 	{
 		// check if there's already an entry in TableRows for the same ProxyAsset
-		if (!TableRows.FindByPredicate(
-			[Proxy = ProxyTable->Entries[i].Proxy](const TSharedPtr<FProxyTableRow>& Row)
-			{
-				return Row->ProxyTable != nullptr
-					&& Row->RowIndex >= 0
-					&& Row->RowIndex < Row->ProxyTable->Entries.Num()
-				    && Proxy == Row->ProxyTable->Entries[Row->RowIndex].Proxy;
-			}))
+		if (!ReferencedProxyAssets.Find(ProxyTable->Entries[i].Proxy))
 		{
-			TableRows.Add(MakeShared<FProxyTableRow>(i, ProxyTable));
+			ParentTableRow->Children.Add(MakeShared<FProxyTableRow>(i, ProxyTable));
+			ReferencedProxyAssets.Add(ProxyTable->Entries[i].Proxy);
 		}
 	}
-	if (NumRows == TableRows.Num())
-	{
-		// if none of the rows were actually inherited, change header to "No rows inherited from"
-		TableRows.Last()->RowIndex = SProxyTableRow::SpecialIndex_NothingInheritedFrom;
-	}
 
 	// recursively add inherited entries
 	for(UProxyTable* InheritedTable : ProxyTable->InheritEntriesFrom)
@@ -803,11 +857,14 @@ void FProxyTableEditor::UpdateTableRows()
 	UProxyTable* ProxyTable = Cast<UProxyTable>(EditingObjects[0]);
 
 	TableRows.SetNum(0);
+	ReferencedProxyAssets.Empty(ReferencedProxyAssets.Num());
+	ReferencedProxyTables.Empty(ReferencedProxyTables.Num());
 
 	// add rows from this table
 	for(int i =0; i<ProxyTable->Entries.Num(); i++)
 	{
 		TableRows.Add(MakeShared<FProxyTableRow>(i, ProxyTable));
+		ReferencedProxyAssets.Add(ProxyTable->Entries[i].Proxy);
 	}
 
 	// Add 1 at the end, for the "Add Row" control
@@ -815,14 +872,12 @@ void FProxyTableEditor::UpdateTableRows()
 	
 	if (ProxyTable->InheritEntriesFrom.Num()>0)
 	{
-
 		// add imported rows in to the table
 		for(UProxyTable* InheritedTable : ProxyTable->InheritEntriesFrom)
 		{
 			AddInheritedRows(InheritedTable);
 		}
 	}
-	
 
 	if (TableView.IsValid())
 	{
@@ -830,32 +885,11 @@ void FProxyTableEditor::UpdateTableRows()
 	}
 }
 
-static bool IsDependency(UProxyTable* QueryTable, UProxyTable* Table)
-{
-	if (QueryTable == nullptr || Table == nullptr)
-	{
-		return false;
-	}
-	
-	for (UProxyTable* InheritFromTable : Table->InheritEntriesFrom)
-	{
-		if (InheritFromTable == QueryTable)
-		{
-			return true;
-		}
-		if (IsDependency(QueryTable, InheritFromTable))
-		{
-			return true;
-		}
-	}
-	return false;
-}
-	
 void FProxyTableEditor::OnObjectTransacted(UObject* InObject, const FTransactionObjectEvent& InTransactionObjectEvent)
 {
 	if (UProxyTable* ModifiedProxyTable = Cast<UProxyTable>(InObject))
 	{
-		if (IsDependency(ModifiedProxyTable, GetProxyTable()))
+		if (ReferencedProxyTables.Contains(ModifiedProxyTable))
 		{
 			UpdateTableRows();
 		}
diff --git a/Engine/Plugins/Experimental/Chooser/Source/ProxyTableEditor/Private/ProxyTableEditor.h b/Engine/Plugins/Experimental/Chooser/Source/ProxyTableEditor/Private/ProxyTableEditor.h
--- a/Engine/Plugins/Experimental/Chooser/Source/ProxyTableEditor/Private/ProxyTableEditor.h
+++ b/Engine/Plugins/Experimental/Chooser/Source/ProxyTableEditor/Private/ProxyTableEditor.h
@@ -7,6 +7,7 @@
 #include "Toolkits/AssetEditorToolkit.h"
 #include "Editor/PropertyEditor/Public/PropertyEditorDelegates.h"
 #include "Widgets/Views/SListView.h"
+#include "Widgets/Views/STreeView.h"
 #include "Widgets/Views/STableViewBase.h"
 #include "ProxyTable.h"
 #include "EditorUndoClient.h"
@@ -88,12 +89,14 @@ namespace UE::ProxyTableEditor
 			
 			UProxyTable* ProxyTable;
 			int32 RowIndex;
+			TArray<TSharedPtr<FProxyTableRow>> Children;			
 		};
 
 		void UpdateTableColumns();
 		void UpdateTableRows();
 		void MoveRow(int SourceRowIndex, int TargetIndex);
 		void InsertEntry(FProxyEntry& Entry, int RowIndex);
+		void DeleteSelectedRows();
 	private:
 		void AddInheritedRows(UProxyTable* ProxyTable);
 		void SelectRootProperties();
@@ -106,6 +109,7 @@ namespace UE::ProxyTableEditor
 		TSharedRef<SDockTab> SpawnTableTab( const FSpawnTabArgs& Args );
 	
 		TSharedRef<ITableRow> GenerateTableRow(TSharedPtr<FProxyTableRow> InItem, const TSharedRef<STableViewBase>& OwnerTable);
+		void TreeViewExpansionChanged(TSharedPtr<FProxyTableEditor::FProxyTableRow> InItem, bool bShouldBeExpanded);
 
 		/** Called when objects need to be swapped out for new versions, like after a blueprint recompile. */
 		void OnObjectsReplaced(const TMap<UObject*, UObject*>& ReplacementMap);
@@ -132,9 +136,21 @@ namespace UE::ProxyTableEditor
 		TSharedPtr<SComboButton> CreateRowComboButton;
 		
 		TSharedPtr<SHeaderRow> HeaderRow;
-		TSharedPtr<SListView<TSharedPtr<FProxyTableRow>>> TableView;
+		TSharedPtr<STreeView<TSharedPtr<FProxyTableRow>>> TableView;
 		
 	public:
+		TMap<UProxyTable*, bool> ImportedTablesExpansionState;
+		TSet<UProxyAsset*> ReferencedProxyAssets;
+		TSet<UProxyTable*> ReferencedProxyTables;
+		
+		void SelectRow(TSharedPtr<FProxyTableRow> Row)
+		{
+			if (!TableView->IsItemSelected(Row))
+			{
+				TableView->ClearSelection();
+				TableView->SetItemSelection(Row, true, ESelectInfo::OnMouseClick);
+			}
+		}
 		TSharedPtr<SComboButton>& GetCreateRowComboButton() { return CreateRowComboButton; };
 
 		/** The name given to all instances of this type of editor */
