#Chooser DisableOnObjectPropertyChanged delegate handles in Chooser plugin, for now

diff --git a/Engine/Plugins/Experimental/Chooser/Source/ChooserUncooked/Private/EvaluateChooserNode.cpp b/Engine/Plugins/Experimental/Chooser/Source/ChooserUncooked/Private/EvaluateChooserNode.cpp
--- a/Engine/Plugins/Experimental/Chooser/Source/ChooserUncooked/Private/EvaluateChooserNode.cpp
+++ b/Engine/Plugins/Experimental/Chooser/Source/ChooserUncooked/Private/EvaluateChooserNode.cpp
@@ -40,6 +40,7 @@ UK2Node_EvaluateChooser::UK2Node_EvaluateChooser(const FObjectInitializer& Objec
 	NodeTooltip = LOCTEXT("NodeTooltip", "Evaluates an Chooser Table, and returns the resulting Object or Objects.");
 
 	// need to unregister this somewhere - also this is pretty excessive
+	/*
 	FCoreUObjectDelegates::OnObjectPropertyChanged.AddLambda([this](UObject* Object, struct FPropertyChangedEvent& Event)
 	{
 		if (Object == Chooser)
@@ -47,6 +48,7 @@ UK2Node_EvaluateChooser::UK2Node_EvaluateChooser(const FObjectInitializer& Objec
 			AllocateDefaultPins();
 		}
 	});
+	*/
 }
 
 void UK2Node_EvaluateChooser::PreloadRequiredAssets()
diff --git a/Engine/Plugins/Experimental/Chooser/Source/ProxyTableUncooked/Private/EvaluateProxyNode.cpp b/Engine/Plugins/Experimental/Chooser/Source/ProxyTableUncooked/Private/EvaluateProxyNode.cpp
--- a/Engine/Plugins/Experimental/Chooser/Source/ProxyTableUncooked/Private/EvaluateProxyNode.cpp
+++ b/Engine/Plugins/Experimental/Chooser/Source/ProxyTableUncooked/Private/EvaluateProxyNode.cpp
@@ -40,6 +40,7 @@ UK2Node_EvaluateProxy::UK2Node_EvaluateProxy(const FObjectInitializer& ObjectIni
 	NodeTooltip = LOCTEXT("NodeTooltip", "Evaluates an Proxy Table, and returns the resulting Object or Objects.");
 
 	// need to unregister this somewhere - also this is pretty excessive
+	/*
 	FCoreUObjectDelegates::OnObjectPropertyChanged.AddLambda([this](UObject* Object, struct FPropertyChangedEvent& Event)
 	{
 		if (Object == Proxy)
@@ -47,6 +48,7 @@ UK2Node_EvaluateProxy::UK2Node_EvaluateProxy(const FObjectInitializer& ObjectIni
 			AllocateDefaultPins();
 		}
 	});
+	*/
 }
 
 void UK2Node_EvaluateProxy::PreloadRequiredAssets()
