Reimplement pin type updating based on Chooser/Proxy asset output object type

diff --git a/Engine/Plugins/Experimental/Chooser/Source/Chooser/Internal/Chooser.h b/Engine/Plugins/Experimental/Chooser/Source/Chooser/Internal/Chooser.h
--- a/Engine/Plugins/Experimental/Chooser/Source/Chooser/Internal/Chooser.h
+++ b/Engine/Plugins/Experimental/Chooser/Source/Chooser/Internal/Chooser.h
@@ -10,6 +10,7 @@
 
 #include "Chooser.generated.h"
 
+DECLARE_MULTICAST_DELEGATE_OneParam(FChooserOutputObjectTypeChanged, const UClass* OutputObjectType);
 
 UCLASS(MinimalAPI, BlueprintType)
 class UChooserTable : public UObject, public IHasContextClass
@@ -17,7 +18,18 @@ class UChooserTable : public UObject, public IHasContextClass
 	GENERATED_UCLASS_BODY()
 public:
 	UChooserTable() {}
+
+#if WITH_EDITOR
+	FChooserOutputObjectTypeChanged OnOutputObjectTypeChanged;
+	virtual void PostEditUndo() override;
+	virtual void PostEditChangeProperty(FPropertyChangedEvent& PropertyChangedEvent) override;
+	virtual void PostLoad() override;
+
+	// caching the OutputObjectType so that on Undo, we can tell if we should fire the changed delegate
+	UClass* CachedPreviousOutputObjectType;
+#endif
 	
+#if WITH_EDITORONLY_DATA
 	// deprecated UObject Results
 	UPROPERTY()
 	TArray<TScriptInterface<IObjectChooser>> Results_DEPRECATED;
@@ -25,6 +37,7 @@ public:
 	// deprecated UObject Columns
 	UPROPERTY()
 	TArray<TScriptInterface<IChooserColumn>> Columns_DEPRECATED;
+#endif
 
 	// Each possible Result (Rows of chooser table)
 	UPROPERTY(EditAnywhere, DisplayName = "Results", Meta = (ExcludeBaseStruct, BaseStruct = "/Script/Chooser.ObjectChooserBase"), Category = "Hidden")
@@ -40,10 +53,6 @@ public:
 	UPROPERTY(EditAnywhere, Category="Output", Meta = (AllowAbstract=true))
 	TObjectPtr<UClass> OutputObjectType;
 
-#if WITH_EDITOR
-	virtual void PostLoad() override;
-#endif
-
 	virtual UClass* GetContextClass() override { return ContextObjectType; }
 };
 
diff --git a/Engine/Plugins/Experimental/Chooser/Source/Chooser/Private/Chooser.cpp b/Engine/Plugins/Experimental/Chooser/Source/Chooser/Private/Chooser.cpp
--- a/Engine/Plugins/Experimental/Chooser/Source/Chooser/Private/Chooser.cpp
+++ b/Engine/Plugins/Experimental/Chooser/Source/Chooser/Private/Chooser.cpp
@@ -12,10 +12,38 @@ UChooserTable::UChooserTable(const FObjectInitializer& Initializer)
 }
 
 #if WITH_EDITOR
+void UChooserTable::PostEditUndo()
+{
+	UObject::PostEditUndo();
+
+	if (CachedPreviousOutputObjectType != OutputObjectType)
+	{
+		OnOutputObjectTypeChanged.Broadcast(OutputObjectType);
+		CachedPreviousOutputObjectType = OutputObjectType;
+	}
+}
+
+void UChooserTable::PostEditChangeProperty(FPropertyChangedEvent& PropertyChangedEvent)
+{
+	UObject::PostEditChangeProperty(PropertyChangedEvent);
+	
+	static FName OutputObjectTypeName = "OutputObjectType";
+	if (PropertyChangedEvent.Property->GetName() == OutputObjectTypeName)
+	{
+		if (CachedPreviousOutputObjectType != OutputObjectType)
+		{
+			OnOutputObjectTypeChanged.Broadcast(OutputObjectType);
+		}
+		CachedPreviousOutputObjectType = OutputObjectType;
+	}
+}
+
 void UChooserTable::PostLoad()
 {
 	Super::PostLoad();
 
+	CachedPreviousOutputObjectType = OutputObjectType;
+
 	// convert old data if it exists
 
 	if (Results_DEPRECATED.Num() > 0 || Columns_DEPRECATED.Num() > 0)
diff --git a/Engine/Plugins/Experimental/Chooser/Source/ChooserUncooked/Private/EvaluateChooserNode.cpp b/Engine/Plugins/Experimental/Chooser/Source/ChooserUncooked/Private/EvaluateChooserNode.cpp
--- a/Engine/Plugins/Experimental/Chooser/Source/ChooserUncooked/Private/EvaluateChooserNode.cpp
+++ b/Engine/Plugins/Experimental/Chooser/Source/ChooserUncooked/Private/EvaluateChooserNode.cpp
@@ -38,17 +38,55 @@ UK2Node_EvaluateChooser::UK2Node_EvaluateChooser(const FObjectInitializer& Objec
 	: Super(ObjectInitializer)
 {
 	NodeTooltip = LOCTEXT("NodeTooltip", "Evaluates an Chooser Table, and returns the resulting Object or Objects.");
+}
 
-	// need to unregister this somewhere - also this is pretty excessive
-	/*
-	FCoreUObjectDelegates::OnObjectPropertyChanged.AddLambda([this](UObject* Object, struct FPropertyChangedEvent& Event)
+void UK2Node_EvaluateChooser::UnregisterChooserCallback()
+{
+	if(CurrentCallbackChooser)
 	{
-		if (Object == Chooser)
+		CurrentCallbackChooser->OnOutputObjectTypeChanged.RemoveAll(this);
+		CurrentCallbackChooser = nullptr;
+	}
+}
+
+UK2Node_EvaluateChooser::~UK2Node_EvaluateChooser()
+{
+	UnregisterChooserCallback();
+}
+
+void UK2Node_EvaluateChooser::PostEditUndo()
+{
+	Super::PostEditUndo();
+	ChooserChanged();
+}
+
+void UK2Node_EvaluateChooser::DestroyNode()
+{
+	UnregisterChooserCallback();
+	Super::DestroyNode();
+}
+
+
+void UK2Node_EvaluateChooser::ChooserChanged()
+{
+	if (Chooser != CurrentCallbackChooser)
+	{
+		UnregisterChooserCallback();
+	
+		if (Chooser)
 		{
-			AllocateDefaultPins();
+			Chooser->OnOutputObjectTypeChanged.AddUObject(this, &UK2Node_EvaluateChooser::ResultTypeChanged);
 		}
-	});
-	*/
+	
+		CurrentCallbackChooser = Chooser;
+
+		AllocateDefaultPins();
+	}
+}
+
+void UK2Node_EvaluateChooser::ResultTypeChanged(const UClass*)
+{
+	AllocateDefaultPins();
 }
 
 void UK2Node_EvaluateChooser::PreloadRequiredAssets()
@@ -109,13 +147,23 @@ FText UK2Node_EvaluateChooser::GetPinDisplayName(const UEdGraphPin* Pin) const
 
 void UK2Node_EvaluateChooser::PostEditChangeProperty(struct FPropertyChangedEvent& PropertyChangedEvent)
 {
-	if (PropertyChangedEvent.Property->GetName() == "Chooser" || PropertyChangedEvent.Property->GetName() == "Mode")
+	if (PropertyChangedEvent.Property->GetName() == "Chooser")
+	{
+		ChooserChanged();
+	}
+	if (PropertyChangedEvent.Property->GetName() == "Mode")
 	{
 		AllocateDefaultPins();
 	}
 	Super::PostEditChangeProperty(PropertyChangedEvent);
 }
 
+void UK2Node_EvaluateChooser::PostLoad()
+{
+	Super::PostLoad();
+	ChooserChanged();
+}
+
 void UK2Node_EvaluateChooser::PinConnectionListChanged(UEdGraphPin* Pin)
 {
 
@@ -131,7 +179,6 @@ void UK2Node_EvaluateChooser::PinDefaultValueChanged(UEdGraphPin* Pin)
 
 void UK2Node_EvaluateChooser::PinTypeChanged(UEdGraphPin* Pin)
 {
-
 	Super::PinTypeChanged(Pin);
 }
 
diff --git a/Engine/Plugins/Experimental/Chooser/Source/ChooserUncooked/Private/EvaluateChooserNode.h b/Engine/Plugins/Experimental/Chooser/Source/ChooserUncooked/Private/EvaluateChooserNode.h
--- a/Engine/Plugins/Experimental/Chooser/Source/ChooserUncooked/Private/EvaluateChooserNode.h
+++ b/Engine/Plugins/Experimental/Chooser/Source/ChooserUncooked/Private/EvaluateChooserNode.h
@@ -25,9 +25,12 @@ UCLASS(MinimalAPI)
 class UK2Node_EvaluateChooser : public UK2Node
 {
 	GENERATED_UCLASS_BODY()
-
+	~UK2Node_EvaluateChooser();
+	
 	//~ Begin UObject Interface
+	virtual void PostEditUndo() override;
 	virtual void PostEditChangeProperty(struct FPropertyChangedEvent& PropertyChangedEvent) override;
+	virtual void PostLoad() override;
 	//~ End UObject Interface
 
 	//~ Begin UEdGraphNode Interface.
@@ -52,10 +55,17 @@ class UK2Node_EvaluateChooser : public UK2Node
 	virtual FText GetMenuCategory() const override;
 	virtual int32 GetNodeRefreshPriority() const override { return EBaseNodeRefreshPriority::Low_UsesDependentWildcard; }
 	virtual void PreloadRequiredAssets() override;
+
+	virtual void DestroyNode() override;
     
 	//~ End UK2Node Interface.
 
 private:
+	void ChooserChanged();
+	void ResultTypeChanged(const UClass*);
+	void UnregisterChooserCallback();
+	UChooserTable* CurrentCallbackChooser = nullptr;
+	
 	UPROPERTY(EditAnywhere, Category = "Chooser")
 	TObjectPtr<UChooserTable> Chooser;
 	
diff --git a/Engine/Plugins/Experimental/Chooser/Source/ProxyTable/Internal/ProxyTable.h b/Engine/Plugins/Experimental/Chooser/Source/ProxyTable/Internal/ProxyTable.h
--- a/Engine/Plugins/Experimental/Chooser/Source/ProxyTable/Internal/ProxyTable.h
+++ b/Engine/Plugins/Experimental/Chooser/Source/ProxyTable/Internal/ProxyTable.h
@@ -9,6 +9,8 @@
 #include "InstancedStruct.h"
 #include "ProxyTable.generated.h"
 
+DECLARE_MULTICAST_DELEGATE_OneParam(FProxyTypeChanged, const UClass* OutputObjectType);
+
 UCLASS(MinimalAPI,BlueprintType)
 class UProxyAsset : public UObject, public IHasContextClass
 {
@@ -16,6 +18,17 @@ class UProxyAsset : public UObject, public IHasContextClass
 public:
 	UProxyAsset() {}
 
+#if WITH_EDITOR
+	FProxyTypeChanged OnTypeChanged;
+	virtual void PostEditUndo() override;
+	virtual void PostEditChangeProperty(FPropertyChangedEvent& PropertyChangedEvent) override;
+	virtual void PostLoad() override;
+
+	// caching the type so that on Undo, we can tell if we should fire the changed delegate
+	UClass* CachedPreviousType = nullptr;
+#endif
+	
+
 	UPROPERTY(EditAnywhere, Category = "Proxy", Meta = (AllowAbstract=true))
 	TObjectPtr<UClass> Type;
 	
diff --git a/Engine/Plugins/Experimental/Chooser/Source/ProxyTable/Private/ProxyTable.cpp b/Engine/Plugins/Experimental/Chooser/Source/ProxyTable/Private/ProxyTable.cpp
--- a/Engine/Plugins/Experimental/Chooser/Source/ProxyTable/Private/ProxyTable.cpp
+++ b/Engine/Plugins/Experimental/Chooser/Source/ProxyTable/Private/ProxyTable.cpp
@@ -1,6 +1,40 @@
 // Copyright Epic Games, Inc. All Rights Reserved.
 #include "ProxyTable.h"
 #include "ProxyTableFunctionLibrary.h"
+#if WITH_EDITOR
+void UProxyAsset::PostEditUndo()
+{
+	UObject::PostEditUndo();
+
+	if (CachedPreviousType != Type)
+	{
+		OnTypeChanged.Broadcast(Type);
+		CachedPreviousType = Type;
+	}
+}
+
+void UProxyAsset::PostEditChangeProperty(FPropertyChangedEvent& PropertyChangedEvent)
+{
+	UObject::PostEditChangeProperty(PropertyChangedEvent);
+	
+	static FName TypeName = "Type";
+	if (PropertyChangedEvent.Property->GetName() == TypeName)
+	{
+		if (CachedPreviousType != Type)
+		{
+			OnTypeChanged.Broadcast(Type);
+		}
+		CachedPreviousType = Type;
+	}
+}
+
+void UProxyAsset::PostLoad()
+{
+	Super::PostLoad();
+	CachedPreviousType = Type;
+}
+#endif
+
 
 FLookupProxy::FLookupProxy()
 {
diff --git a/Engine/Plugins/Experimental/Chooser/Source/ProxyTableUncooked/Private/EvaluateProxyNode.cpp b/Engine/Plugins/Experimental/Chooser/Source/ProxyTableUncooked/Private/EvaluateProxyNode.cpp
--- a/Engine/Plugins/Experimental/Chooser/Source/ProxyTableUncooked/Private/EvaluateProxyNode.cpp
+++ b/Engine/Plugins/Experimental/Chooser/Source/ProxyTableUncooked/Private/EvaluateProxyNode.cpp
@@ -38,19 +38,57 @@ UK2Node_EvaluateProxy::UK2Node_EvaluateProxy(const FObjectInitializer& ObjectIni
 	: Super(ObjectInitializer)
 {
 	NodeTooltip = LOCTEXT("NodeTooltip", "Evaluates an Proxy Table, and returns the resulting Object or Objects.");
+}
+
+void UK2Node_EvaluateProxy::UnregisterProxyCallback()
+{
+	if(CurrentCallbackProxy)
+	{
+		CurrentCallbackProxy->OnTypeChanged.RemoveAll(this);
+		CurrentCallbackProxy = nullptr;
+	}
+}
+
+UK2Node_EvaluateProxy::~UK2Node_EvaluateProxy()
+{
+	UnregisterProxyCallback();
+}
+
+void UK2Node_EvaluateProxy::PostEditUndo()
+{
+	Super::PostEditUndo();
+	ProxyChanged();
+}
 
-	// need to unregister this somewhere - also this is pretty excessive
-	/*
-	FCoreUObjectDelegates::OnObjectPropertyChanged.AddLambda([this](UObject* Object, struct FPropertyChangedEvent& Event)
+void UK2Node_EvaluateProxy::DestroyNode()
+{
+	UnregisterProxyCallback();
+	Super::DestroyNode();
+}
+
+void UK2Node_EvaluateProxy::ProxyChanged()
+{
+	if (Proxy != CurrentCallbackProxy)
 	{
-		if (Object == Proxy)
+		UnregisterProxyCallback();
+	
+		if (Proxy)
 		{
-			AllocateDefaultPins();
+			Proxy->OnTypeChanged.AddUObject(this, &UK2Node_EvaluateProxy::TypeChanged);
 		}
-	});
-	*/
+	
+		CurrentCallbackProxy = Proxy;
+
+		AllocateDefaultPins();
+	}
+}
+
+void UK2Node_EvaluateProxy::TypeChanged(const UClass*)
+{
+	AllocateDefaultPins();
 }
 
+
 void UK2Node_EvaluateProxy::PreloadRequiredAssets()
 {
 	if (Proxy)
@@ -105,12 +143,22 @@ FText UK2Node_EvaluateProxy::GetPinDisplayName(const UEdGraphPin* Pin) const
 void UK2Node_EvaluateProxy::PostEditChangeProperty(struct FPropertyChangedEvent& PropertyChangedEvent)
 {
 	if (PropertyChangedEvent.Property->GetName() == "Proxy")
+	{
+		ProxyChanged();
+	}
+	if (PropertyChangedEvent.Property->GetName() == "Mode")
 	{
 		AllocateDefaultPins();
 	}
 	Super::PostEditChangeProperty(PropertyChangedEvent);
 }
 
+void UK2Node_EvaluateProxy::PostLoad()
+{
+	Super::PostLoad();
+	ProxyChanged();
+}
+
 void UK2Node_EvaluateProxy::PinConnectionListChanged(UEdGraphPin* Pin)
 {
 	Modify();
@@ -125,7 +173,6 @@ void UK2Node_EvaluateProxy::PinDefaultValueChanged(UEdGraphPin* Pin)
 
 void UK2Node_EvaluateProxy::PinTypeChanged(UEdGraphPin* Pin)
 {
-
 	Super::PinTypeChanged(Pin);
 }
 
diff --git a/Engine/Plugins/Experimental/Chooser/Source/ProxyTableUncooked/Private/EvaluateProxyNode.h b/Engine/Plugins/Experimental/Chooser/Source/ProxyTableUncooked/Private/EvaluateProxyNode.h
--- a/Engine/Plugins/Experimental/Chooser/Source/ProxyTableUncooked/Private/EvaluateProxyNode.h
+++ b/Engine/Plugins/Experimental/Chooser/Source/ProxyTableUncooked/Private/EvaluateProxyNode.h
@@ -18,9 +18,12 @@ UCLASS(MinimalAPI)
 class UK2Node_EvaluateProxy : public UK2Node
 {
 	GENERATED_UCLASS_BODY()
+	~UK2Node_EvaluateProxy();
 
 	//~ Begin UObject Interface
+	virtual void PostEditUndo() override;
 	virtual void PostEditChangeProperty(struct FPropertyChangedEvent& PropertyChangedEvent) override;
+	virtual void PostLoad() override;
 	//~ End UObject Interface
 
 	//~ Begin UEdGraphNode Interface.
@@ -45,10 +48,18 @@ class UK2Node_EvaluateProxy : public UK2Node
 	virtual FText GetMenuCategory() const override;
 	virtual int32 GetNodeRefreshPriority() const override { return EBaseNodeRefreshPriority::Low_UsesDependentWildcard; }
 	virtual void PreloadRequiredAssets() override;
-    
+   
+	virtual void DestroyNode() override; 
 	//~ End UK2Node Interface.
 
 private:
+
+	void ProxyChanged();
+	void TypeChanged(const UClass*);
+	void UnregisterProxyCallback();
+	UProxyAsset* CurrentCallbackProxy = nullptr;
+
+
 	UPROPERTY(EditAnywhere, Category = "Proxy")
 	TObjectPtr<UProxyAsset> Proxy;
 
