Proxy table cooked data optimizations

diff --git a/Engine/Plugins/Experimental/Chooser/Source/ProxyTable/Internal/ProxyTable.h b/Engine/Plugins/Experimental/Chooser/Source/ProxyTable/Internal/ProxyTable.h
--- a/Engine/Plugins/Experimental/Chooser/Source/ProxyTable/Internal/ProxyTable.h
+++ b/Engine/Plugins/Experimental/Chooser/Source/ProxyTable/Internal/ProxyTable.h
@@ -7,6 +7,8 @@
 #include "ChooserPropertyAccess.h"
 #include "IChooserParameterProxyTable.h"
 #include "InstancedStruct.h"
+#include "InstancedStructContainer.h"
+#include "Misc/Guid.h"
 #include "ProxyTable.generated.h"
 
 DECLARE_MULTICAST_DELEGATE_OneParam(FProxyTypeChanged, const UClass* OutputObjectType);
@@ -22,7 +24,6 @@ public:
 	FProxyTypeChanged OnTypeChanged;
 	virtual void PostEditUndo() override;
 	virtual void PostEditChangeProperty(FPropertyChangedEvent& PropertyChangedEvent) override;
-	virtual void PostLoad() override;
 
 	// caching the type so that on Undo, we can tell if we should fire the changed delegate
 	UClass* CachedPreviousType = nullptr;
@@ -38,13 +39,24 @@ public:
 	UPROPERTY(EditAnywhere, Meta = (ExcludeBaseStruct, BaseStruct ="/Script/ProxyTable.ChooserParameterProxyTableBase"), Category = "Proxy Table Reference")
 	FInstancedStruct ProxyTable;
 
+	UPROPERTY()
+	FGuid Guid;
+
 	virtual UClass* GetContextClass() override { return ContextClass; }
+
+#if WITH_EDITORONLY_DATA
+	virtual void PostLoad() override;
+	virtual void PostDuplicate(EDuplicateMode::Type DuplicateMode) override;
+#endif
 };
 
 USTRUCT()
-struct FProxyEntry
+struct PROXYTABLE_API FProxyEntry
 {
 	GENERATED_BODY()
+
+	bool operator== (const FProxyEntry& Other) const;
+	bool operator< (const FProxyEntry& Other) const;
 	
 	UPROPERTY(EditAnywhere, Category = "Data")
 	TObjectPtr<UProxyAsset> Proxy;
@@ -55,8 +67,16 @@ struct FProxyEntry
 	
 	UPROPERTY(DisplayName="Value", EditAnywhere, Meta = (ExcludeBaseStruct, BaseStruct = "/Script/Chooser.ObjectChooserBase"), Category = "Data")
 	FInstancedStruct ValueStruct;
+
+	const FGuid GetGuid() const;
 };
 
+#if WITH_EDITORONLY_DATA
+inline uint32 GetTypeHash(const FProxyEntry& Entry) { return GetTypeHash(Entry.GetGuid()); }
+#endif
+
+DECLARE_MULTICAST_DELEGATE(FProxyTableChanged);
+
 UCLASS(MinimalAPI,BlueprintType)
 class UProxyTable : public UObject
 {
@@ -64,11 +84,32 @@ class UProxyTable : public UObject
 public:
 	UProxyTable() {}
 
+	UPROPERTY()
+	TArray<FGuid> Keys;
+
+	UPROPERTY()
+	FInstancedStructContainer Values;
+	
+	UObject* FindProxyObject(const FGuid& Key, const UObject* ContextObject) const;
+
+#if WITH_EDITORONLY_DATA
+public:
+	FProxyTableChanged OnProxyTableChanged;
+	
 	UPROPERTY(EditAnywhere, Category = "Hidden")
 	TArray<FProxyEntry> Entries;
 	
 	UPROPERTY(EditAnywhere, Category = "Inheritance")
 	TArray<TObjectPtr<UProxyTable>> InheritEntriesFrom;
+
+	virtual void PostLoad() override;
+	virtual void PostTransacted(const FTransactionObjectEvent& TransactionEvent) override;
+private:
+	void BuildRuntimeData();
+	TArray<TWeakObjectPtr<UProxyTable>> TableDependencies;
+	TArray<TWeakObjectPtr<UProxyAsset>> ProxyDependencies;
+#endif
+
 };
 
 
diff --git a/Engine/Plugins/Experimental/Chooser/Source/ProxyTable/Private/ProxyTable.cpp b/Engine/Plugins/Experimental/Chooser/Source/ProxyTable/Private/ProxyTable.cpp
--- a/Engine/Plugins/Experimental/Chooser/Source/ProxyTable/Private/ProxyTable.cpp
+++ b/Engine/Plugins/Experimental/Chooser/Source/ProxyTable/Private/ProxyTable.cpp
@@ -1,6 +1,12 @@
 // Copyright Epic Games, Inc. All Rights Reserved.
 #include "ProxyTable.h"
 #include "ProxyTableFunctionLibrary.h"
+#include "AssetRegistry/AssetRegistryModule.h"
+#include "Logging/LogMacros.h"
+#include "UObject/Package.h"
+
+DEFINE_LOG_CATEGORY_STATIC(LogProxyTable,Log,All);
+
 #if WITH_EDITOR
 void UProxyAsset::PostEditUndo()
 {
@@ -28,41 +34,205 @@ void UProxyAsset::PostEditChangeProperty(FPropertyChangedEvent& PropertyChangedE
 	}
 }
 
+#endif
+
+
+FLookupProxy::FLookupProxy()
+{
+}
+
+#if WITH_EDITORONLY_DATA
+
+/////////////////////////////////////////////////////////////////////////////////////////
+// Proxy Asset
+
 void UProxyAsset::PostLoad()
 {
 	Super::PostLoad();
+
+#if WITH_EDITOR
 	CachedPreviousType = Type;
-}
 #endif
 
+	if (!Guid.IsValid())
+	{
+		// if we load a ProxyAsset that was created before the Guid, assign it a deterministic guid based on the name and path.
+		Guid.A = GetTypeHash(GetName());
+		Guid.B = GetTypeHash(GetPackage()->GetPathName());
+	}
+}
 
-FLookupProxy::FLookupProxy()
+void UProxyAsset::PostDuplicate(EDuplicateMode::Type DuplicateMode)
+{
+	UObject::PostDuplicate(DuplicateMode);
+	if (DuplicateMode == EDuplicateMode::Normal)
+	{
+		// create a new guid when duplicating
+		Guid = FGuid::NewGuid();
+	}
+}
+
+//////////////////////////////////////////////////////////////////////////////////////
+/// Proxy Entry
+
+// Gets the uint32 key that will be used as the unique identifier for the Proxy Entry
+const FGuid FProxyEntry::GetGuid() const
+{
+	if (Proxy)
+	{
+		return Proxy->Guid;
+	}
+	else
+	{
+		// fallback for old FName key content
+		FGuid Guid;
+		if (Key != NAME_None)
+		{
+			Guid.A = GetTypeHash(Key);
+		}
+		return Guid;
+	}
+}
+
+// == operator for TArray::Contains
+bool FProxyEntry::operator== (const FProxyEntry& Other) const
+{
+	return GetGuid() == Other.GetGuid();
+}
+
+// < operator for Algo::BinarySearch, and TArray::StableSort
+bool FProxyEntry::operator< (const FProxyEntry& Other) const
 {
+	return GetGuid() < Other.GetGuid();
 }
 
-static UObject* FindProxyObject(const UProxyTable* Table, const UProxyAsset* Proxy, const UObject* ContextObject)
+//////////////////////////////////////////////////////////////////////////////////////
+/// Proxy Table
+
+static void BuildRuntimeDataRecursive(UProxyTable* RootTable, UProxyTable* Table, TArray<FProxyEntry>& OutEntriesArray, TArray<TWeakObjectPtr<UProxyTable>>& OutDependencies)
 {
-	if (Table)
+	if(Table == nullptr)
 	{
-		for (const FProxyEntry& Entry : Table->Entries)
+		return;
+	}
+	
+	OutDependencies.Add(Table);
+	
+	for (const FProxyEntry& Entry : Table->Entries)
+	{
+		int32 FoundIndex = OutEntriesArray.Find(Entry);
+		if (FoundIndex == INDEX_NONE)
 		{
-			if (Entry.Proxy == Proxy && Entry.ValueStruct.IsValid())
+			OutEntriesArray.Add(Entry);
+		}
+		else
+		{
+			// check for Guid collisions
+			if (Entry.Proxy && OutEntriesArray[FoundIndex].Proxy)
+			{
+				if (Entry.Proxy != OutEntriesArray[FoundIndex].Proxy)
+				{
+					UE_LOG(LogProxyTable, Error, TEXT("Proxy Assets %s, and %s have the same Guid. They may have been duplicated outside the editor."),
+						*Entry.Proxy.GetName(), *OutEntriesArray[FoundIndex].Proxy.GetName());
+				}
+			}
+			else
 			{
-				const FObjectChooserBase& EntryValue = Entry.ValueStruct.Get<FObjectChooserBase>();
-				return EntryValue.ChooseObject(ContextObject);
+				// fallback for FName based keys
+				if (Entry.Key != OutEntriesArray[FoundIndex].Key)
+				{
+					UE_LOG(LogProxyTable, Error, TEXT("Proxy Key %s, and %s have the same Hash."),
+						Entry.Key, OutEntriesArray[FoundIndex].Key);
+				}
 			}
 		}
+	}
+
+	for (const TObjectPtr<UProxyTable> ParentTable : Table->InheritEntriesFrom)
+	{
+		if (!OutDependencies.Contains(ParentTable))
+		{
+			BuildRuntimeDataRecursive(RootTable, ParentTable, OutEntriesArray, OutDependencies);
+		}
+	}
+}
 
-		// search parent tables (uncooked data only)
-		for (const TObjectPtr<UProxyTable> ParentTable : Table->InheritEntriesFrom)
+void UProxyTable::BuildRuntimeData()
+{
+	// Unregister callbacks on current dependencies
+	for (TWeakObjectPtr<UProxyTable> Dependency : TableDependencies)
+	{
+		if (Dependency.IsValid())
 		{
-			if (UObject* Value = FindProxyObject(ParentTable, Proxy, ContextObject))
+			Dependency->OnProxyTableChanged.RemoveAll(this);
+		}
+	}
+	TableDependencies.Empty();
+	ProxyDependencies.Empty();
+
+	TArray<FProxyEntry> RuntimeEntries;
+	BuildRuntimeDataRecursive(this, this, RuntimeEntries, TableDependencies);
+
+	// sort by Key
+	RuntimeEntries.StableSort();
+
+	// Copy to Key and Value arrays
+	int EntryCount = RuntimeEntries.Num();
+	
+	Keys.Empty(EntryCount);
+	Values.Empty();
+	
+	TArray<FConstStructView> Views;
+	Views.Reserve(EntryCount);
+	
+	for(const FProxyEntry& Entry : RuntimeEntries) 
+	{
+		Keys.Add(Entry.GetGuid());
+   		Views.Add(Entry.ValueStruct);
+	}
+
+	Values = Views;
+	
+	// register callbacks on updated dependencies
+	for (TWeakObjectPtr<UProxyTable> Dependency : TableDependencies)
+	{
+		Dependency->OnProxyTableChanged.AddUObject(this, &UProxyTable::BuildRuntimeData);
+	}
+
+	// keep a copy of proxy assets just for debugging purposes (indexes match the Key and Value arrays)
+	for(const FProxyEntry& Entry : RuntimeEntries) 
+	{
+		if (Entry.Proxy)
 			{
-				return Value;
+			ProxyDependencies.Add(Entry.Proxy);
 			}
 		}
 	}
 
+
+void UProxyTable::PostLoad()
+{
+	Super::PostLoad();
+	BuildRuntimeData();
+}
+
+void UProxyTable::PostTransacted(const FTransactionObjectEvent& TransactionEvent)
+{
+	UObject::PostTransacted(TransactionEvent);
+	OnProxyTableChanged.Broadcast();
+}
+
+#endif
+
+UObject* UProxyTable::FindProxyObject(const FGuid& Key, const UObject* ContextObject) const
+{
+	const int FoundIndex = Algo::BinarySearch(Keys, Key);
+	if (FoundIndex != INDEX_NONE)
+	{
+		const FObjectChooserBase &EntryValue = Values[FoundIndex].Get<const FObjectChooserBase>();
+		return EntryValue.ChooseObject(ContextObject);
+	}
+	
 	return nullptr;
 }
 
@@ -73,9 +243,12 @@ static UObject* FindProxyObject(const UProxyAsset* Proxy, const UObject* Context
 		const UProxyTable* Table;
 		if (Proxy->ProxyTable.Get<FChooserParameterProxyTableBase>().GetValue(ContextObject, Table))
 		{
-			if (UObject* Value = FindProxyObject(Table, Proxy, ContextObject))
+			if(Table)
 			{
-				return Value;
+				if (UObject* Value = Table->FindProxyObject(Proxy->Guid, ContextObject))
+				{
+					return Value;
+				}
 			}
 		}
 	}
@@ -117,6 +290,9 @@ bool FProxyTableContextProperty::GetValue(const UObject* ContextObject, const UP
 	return false;
 }
 
+/////////////////////////////////////////////////////////////////////////////////////
+// Blueprint Library Functions
+
 UObject* UProxyTableFunctionLibrary::EvaluateProxyAsset(const UObject* ContextObject, const UProxyAsset* Proxy, TSubclassOf<UObject> ObjectClass)
 {
 	UObject* Result = FindProxyObject(Proxy, ContextObject);
@@ -129,34 +305,20 @@ UObject* UProxyTableFunctionLibrary::EvaluateProxyAsset(const UObject* ContextOb
 }
 
 // fallback for FName based Keys:
-
-static UObject* FindProxyObject(const UProxyTable* Table, FName Key, const UObject* ContextObject)
+UObject* UProxyTableFunctionLibrary::EvaluateProxyTable(const UObject* ContextObject, const UProxyTable* ProxyTable, FName Key)
 {
-	if (Table)
+	if (ProxyTable)
 	{
-		for (const FProxyEntry& Entry : Table->Entries)
-		{
-			if (Entry.Key == Key && Entry.ValueStruct.IsValid())
-			{
-				const FObjectChooserBase& EntryValue = Entry.ValueStruct.Get<FObjectChooserBase>();
-				return EntryValue.ChooseObject(ContextObject);
-			}
-		}
-
-		// search parent tables (uncooked data only)
-		for (const TObjectPtr<UProxyTable> ParentTable : Table->InheritEntriesFrom)
+		if(ProxyTable)
 		{
-			if (UObject* Value = FindProxyObject(ParentTable, Key, ContextObject))
+			FGuid Guid;
+			Guid.A = GetTypeHash(Key);
+			if (UObject* Value = ProxyTable->FindProxyObject(Guid, ContextObject))
 			{
 				return Value;
 			}
 		}
 	}
-
+	
 	return nullptr;
 }
-
-UObject* UProxyTableFunctionLibrary::EvaluateProxyTable(const UObject* ContextObject, const UProxyTable* ProxyTable, FName Key)
-{
-	return FindProxyObject(ProxyTable, Key, ContextObject);
-}
diff --git a/Engine/Plugins/Experimental/Chooser/Source/ProxyTableEditor/Private/ProxyAssetFactory.cpp b/Engine/Plugins/Experimental/Chooser/Source/ProxyTableEditor/Private/ProxyAssetFactory.cpp
--- a/Engine/Plugins/Experimental/Chooser/Source/ProxyTableEditor/Private/ProxyAssetFactory.cpp
+++ b/Engine/Plugins/Experimental/Chooser/Source/ProxyTableEditor/Private/ProxyAssetFactory.cpp
@@ -17,5 +17,7 @@ bool UProxyAssetFactory::ConfigureProperties()
 
 UObject* UProxyAssetFactory::FactoryCreateNew(UClass* Class, UObject* InParent, FName Name, EObjectFlags Flags, UObject* Context, FFeedbackContext* Warn, FName CallingContext)
 {
-	return NewObject<UProxyAsset>(InParent, Class, Name, Flags | RF_Public | RF_Standalone | RF_Transactional | RF_LoadCompleted);
+	UProxyAsset* Asset = NewObject<UProxyAsset>(InParent, Class, Name, Flags | RF_Public | RF_Standalone | RF_Transactional | RF_LoadCompleted);
+	Asset->Guid = FGuid::NewGuid();
+	return Asset;
 }
\ No newline at end of file
diff --git a/Engine/Plugins/Experimental/Chooser/Source/ProxyTableEditor/Private/ProxyTableEditor.cpp b/Engine/Plugins/Experimental/Chooser/Source/ProxyTableEditor/Private/ProxyTableEditor.cpp
--- a/Engine/Plugins/Experimental/Chooser/Source/ProxyTableEditor/Private/ProxyTableEditor.cpp
+++ b/Engine/Plugins/Experimental/Chooser/Source/ProxyTableEditor/Private/ProxyTableEditor.cpp
@@ -12,8 +12,7 @@
 #include "Editor.h"
 #include "Framework/MultiBox/MultiBoxBuilder.h"
 #include "Widgets/Input/SHyperlink.h"
-#include "Widgets/Input/SNumericEntryBox.h"
-#include "Widgets/Views/SListView.h"
+#include "Widgets/Input/SEditableTextBox.h"
 #include "Widgets/Input/SComboButton.h"
 #include "SAssetDropTarget.h"
 #include "SClassViewer.h"
@@ -416,6 +415,7 @@ public:
 	{
 		static FName Handles = "Handles";
 		static FName Key = "Key";
+		static FName OldKey = "OldKey";
 		static FName Value = "Value";
 
 		if (Row->RowIndex >=0)
@@ -473,6 +473,20 @@ public:
 						.Visibility_Lambda([this]() { return bDragActive && bDropAbove ? EVisibility::Visible : EVisibility::Hidden; })
 					];
 			}
+			else if (ColumnName == OldKey)
+			{
+				bool bReadOnly = Row->ProxyTable != Editor->GetProxyTable();
+				return SNew(SEditableTextBox)
+					.IsEnabled(!bReadOnly)
+					.Text_Lambda([this](){ return Row->ProxyTable->Entries.Num() > Row->RowIndex ?  FText::FromName(Row->ProxyTable->Entries[Row->RowIndex].Key) : FText();})
+					.OnTextCommitted_Lambda([this](const FText& Text, ETextCommit::Type CommitType)
+					{
+						if (Row->ProxyTable->Entries.Num() > Row->RowIndex)
+						{
+							Row->ProxyTable->Entries[Row->RowIndex].Key = FName(Text.ToString());
+						}
+					});	
+			}
 			else if (ColumnName == Key)
 			{
 				bool bReadOnly = Row->ProxyTable != Editor->GetProxyTable();
@@ -718,10 +732,28 @@ void FProxyTableEditor::UpdateTableColumns()
 					.DefaultLabel(LOCTEXT("KeyColumnName", "Proxy"))
 					.ManualWidth(500));
 
+
+		if (UProxyTable* Table = GetProxyTable())
+    	{
+    		if (Table->Entries.Num() > 0)
+    		{
+    			// if the first entry has a non-none key, assume this is an old table and make an extra column with the old FName Key property
+    			if (Table->Entries[0].Key != NAME_None)
+    			{
+    					HeaderRow->AddColumn(SHeaderRow::Column("OldKey")
+                    					.DefaultLabel(LOCTEXT("OldKeyColumnName", "Key (Deprecated)"))
+                    					.ManualWidth(500));
+    			}
+    		}
+    	}
+
 	HeaderRow->AddColumn(SHeaderRow::Column("Value")
 					.DefaultLabel(LOCTEXT("ValueColumnName", "Value"))
 					.ManualWidth(500));
 
+	
+
+
 }
 
 TSharedRef<SDockTab> FProxyTableEditor::SpawnTableTab( const FSpawnTabArgs& Args )
@@ -838,10 +870,10 @@ void FProxyTableEditor::AddInheritedRows(UProxyTable* ProxyTable)
 	for(int i =0; i<ProxyTable->Entries.Num(); i++)
 	{
 		// check if there's already an entry in TableRows for the same ProxyAsset
-		if (!ReferencedProxyAssets.Find(ProxyTable->Entries[i].Proxy))
+		if (!ReferencedProxyEntries.Find(ProxyTable->Entries[i]))
 		{
 			ParentTableRow->Children.Add(MakeShared<FProxyTableRow>(i, ProxyTable));
-			ReferencedProxyAssets.Add(ProxyTable->Entries[i].Proxy);
+			ReferencedProxyEntries.Add(ProxyTable->Entries[i]);
 		}
 	}
 
@@ -857,14 +889,14 @@ void FProxyTableEditor::UpdateTableRows()
 	UProxyTable* ProxyTable = Cast<UProxyTable>(EditingObjects[0]);
 
 	TableRows.SetNum(0);
-	ReferencedProxyAssets.Empty(ReferencedProxyAssets.Num());
+	ReferencedProxyEntries.Empty(ReferencedProxyEntries.Num());
 	ReferencedProxyTables.Empty(ReferencedProxyTables.Num());
 
 	// add rows from this table
 	for(int i =0; i<ProxyTable->Entries.Num(); i++)
 	{
 		TableRows.Add(MakeShared<FProxyTableRow>(i, ProxyTable));
-		ReferencedProxyAssets.Add(ProxyTable->Entries[i].Proxy);
+		ReferencedProxyEntries.Add(ProxyTable->Entries[i]);
 	}
 
 	// Add 1 at the end, for the "Add Row" control
diff --git a/Engine/Plugins/Experimental/Chooser/Source/ProxyTableEditor/Private/ProxyTableEditor.h b/Engine/Plugins/Experimental/Chooser/Source/ProxyTableEditor/Private/ProxyTableEditor.h
--- a/Engine/Plugins/Experimental/Chooser/Source/ProxyTableEditor/Private/ProxyTableEditor.h
+++ b/Engine/Plugins/Experimental/Chooser/Source/ProxyTableEditor/Private/ProxyTableEditor.h
@@ -140,7 +140,7 @@ namespace UE::ProxyTableEditor
 		
 	public:
 		TMap<UProxyTable*, bool> ImportedTablesExpansionState;
-		TSet<UProxyAsset*> ReferencedProxyAssets;
+		TSet<FProxyEntry> ReferencedProxyEntries;
 		TSet<UProxyTable*> ReferencedProxyTables;
 		
 		void SelectRow(TSharedPtr<FProxyTableRow> Row)
