Fix log passing FName as c string to Log in ProxyTable

diff --git a/Engine/Plugins/Experimental/Chooser/Source/ProxyTable/Private/ProxyTable.cpp b/Engine/Plugins/Experimental/Chooser/Source/ProxyTable/Private/ProxyTable.cpp
--- a/Engine/Plugins/Experimental/Chooser/Source/ProxyTable/Private/ProxyTable.cpp
+++ b/Engine/Plugins/Experimental/Chooser/Source/ProxyTable/Private/ProxyTable.cpp
@@ -142,7 +142,7 @@ static void BuildRuntimeDataRecursive(UProxyTable* RootTable, UProxyTable* Table
 				if (Entry.Key != OutEntriesArray[FoundIndex].Key)
 				{
 					UE_LOG(LogProxyTable, Error, TEXT("Proxy Key %s, and %s have the same Hash."),
-						Entry.Key, OutEntriesArray[FoundIndex].Key);
+						*Entry.Key.ToString(), *OutEntriesArray[FoundIndex].Key.ToString());
 				}
 			}
 		}
