#UE Memory corruption fix

diff --git a/Engine/Plugins/Experimental/Chooser/Source/ProxyTableUncooked/Private/EvaluateProxyNode.cpp b/Engine/Plugins/Experimental/Chooser/Source/ProxyTableUncooked/Private/EvaluateProxyNode.cpp
--- a/Engine/Plugins/Experimental/Chooser/Source/ProxyTableUncooked/Private/EvaluateProxyNode.cpp
+++ b/Engine/Plugins/Experimental/Chooser/Source/ProxyTableUncooked/Private/EvaluateProxyNode.cpp
@@ -49,9 +49,11 @@ void UK2Node_EvaluateProxy::UnregisterProxyCallback()
 	}
 }
 
-UK2Node_EvaluateProxy::~UK2Node_EvaluateProxy()
+void UK2Node_EvaluateProxy::BeginDestroy()
 {
 	UnregisterProxyCallback();
+	
+	Super::BeginDestroy();
 }
 
 void UK2Node_EvaluateProxy::PostEditUndo()
diff --git a/Engine/Plugins/Experimental/Chooser/Source/ProxyTableUncooked/Private/EvaluateProxyNode.h b/Engine/Plugins/Experimental/Chooser/Source/ProxyTableUncooked/Private/EvaluateProxyNode.h
--- a/Engine/Plugins/Experimental/Chooser/Source/ProxyTableUncooked/Private/EvaluateProxyNode.h
+++ b/Engine/Plugins/Experimental/Chooser/Source/ProxyTableUncooked/Private/EvaluateProxyNode.h
@@ -18,7 +18,7 @@ UCLASS(MinimalAPI)
 class UK2Node_EvaluateProxy : public UK2Node
 {
 	GENERATED_UCLASS_BODY()
-	~UK2Node_EvaluateProxy();
+	virtual void BeginDestroy() override;
 
 	//~ Begin UObject Interface
 	virtual void PostEditUndo() override;
