Update details panel customization for property access chains to handle changes to the context class correctly

diff --git a/Engine/Plugins/Experimental/Chooser/Source/Chooser/Internal/Chooser.h b/Engine/Plugins/Experimental/Chooser/Source/Chooser/Internal/Chooser.h
--- a/Engine/Plugins/Experimental/Chooser/Source/Chooser/Internal/Chooser.h
+++ b/Engine/Plugins/Experimental/Chooser/Source/Chooser/Internal/Chooser.h
@@ -21,12 +21,14 @@ public:
 
 #if WITH_EDITOR
 	FChooserOutputObjectTypeChanged OnOutputObjectTypeChanged;
+	
 	virtual void PostEditUndo() override;
 	virtual void PostEditChangeProperty(FPropertyChangedEvent& PropertyChangedEvent) override;
 	virtual void PostLoad() override;
 
-	// caching the OutputObjectType so that on Undo, we can tell if we should fire the changed delegate
+	// caching the OutputObjectType and ContextObjectType so that on Undo, we can tell if we should fire the changed delegate
 	UClass* CachedPreviousOutputObjectType;
+	UClass* CachedPreviousContextObjectType;
 #endif
 	
 #if WITH_EDITORONLY_DATA
diff --git a/Engine/Plugins/Experimental/Chooser/Source/Chooser/Private/Chooser.cpp b/Engine/Plugins/Experimental/Chooser/Source/Chooser/Private/Chooser.cpp
--- a/Engine/Plugins/Experimental/Chooser/Source/Chooser/Private/Chooser.cpp
+++ b/Engine/Plugins/Experimental/Chooser/Source/Chooser/Private/Chooser.cpp
@@ -21,6 +21,11 @@ void UChooserTable::PostEditUndo()
 		OnOutputObjectTypeChanged.Broadcast(OutputObjectType);
 		CachedPreviousOutputObjectType = OutputObjectType;
 	}
+	if (CachedPreviousContextObjectType != ContextObjectType)
+	{
+		OnContextClassChanged.Broadcast(ContextObjectType);
+		CachedPreviousContextObjectType = ContextObjectType;
+	}
 }
 
 void UChooserTable::PostEditChangeProperty(FPropertyChangedEvent& PropertyChangedEvent)
@@ -33,8 +38,18 @@ void UChooserTable::PostEditChangeProperty(FPropertyChangedEvent& PropertyChange
 		if (CachedPreviousOutputObjectType != OutputObjectType)
 		{
 			OnOutputObjectTypeChanged.Broadcast(OutputObjectType);
+			CachedPreviousOutputObjectType = OutputObjectType;
+		}
+	}
+	
+	static FName ContextPropertyName = "ContextObjectType";
+	if (PropertyChangedEvent.Property->GetName() == ContextPropertyName)
+	{
+		if (CachedPreviousContextObjectType != ContextObjectType)
+		{
+			OnContextClassChanged.Broadcast(ContextObjectType);
+			CachedPreviousContextObjectType = ContextObjectType;
 		}
-		CachedPreviousOutputObjectType = OutputObjectType;
 	}
 }
 
diff --git a/Engine/Plugins/Experimental/Chooser/Source/Chooser/Public/ChooserPropertyAccess.h b/Engine/Plugins/Experimental/Chooser/Source/Chooser/Public/ChooserPropertyAccess.h
--- a/Engine/Plugins/Experimental/Chooser/Source/Chooser/Public/ChooserPropertyAccess.h
+++ b/Engine/Plugins/Experimental/Chooser/Source/Chooser/Public/ChooserPropertyAccess.h
@@ -17,10 +17,13 @@ class UHasContextClass : public UInterface
 	GENERATED_BODY()
 };
 
+DECLARE_MULTICAST_DELEGATE_OneParam(FContextClassChanged, UClass*)
+
 class IHasContextClass
 {
 	GENERATED_BODY()
 public:
+	FContextClassChanged OnContextClassChanged;
 	virtual UClass* GetContextClass() { return nullptr; };
 };
 
diff --git a/Engine/Plugins/Experimental/Chooser/Source/ChooserEditor/Private/PropertyAccessChainCustomization.cpp b/Engine/Plugins/Experimental/Chooser/Source/ChooserEditor/Private/PropertyAccessChainCustomization.cpp
--- a/Engine/Plugins/Experimental/Chooser/Source/ChooserEditor/Private/PropertyAccessChainCustomization.cpp
+++ b/Engine/Plugins/Experimental/Chooser/Source/ChooserEditor/Private/PropertyAccessChainCustomization.cpp
@@ -17,191 +17,247 @@
 
 namespace UE::ChooserEditor
 {
-
-void FPropertyAccessChainCustomization::CustomizeHeader(TSharedRef<IPropertyHandle> PropertyHandle, FDetailWidgetRow& HeaderRow, IPropertyTypeCustomizationUtils& CustomizationUtils)
+	
+// Wrapper widget for Property access widget, which can update when the target Class changes
+class SPropertyAccessChainWidget : public SCompoundWidget
 {
-	FProperty* Property = PropertyHandle->GetProperty();
+public:
+	SLATE_BEGIN_ARGS(SPropertyAccessChainWidget)
+	{}
 
-	UClass* ContextClass = nullptr;
-	static FName TypeMetaData = "BindingType";
-	static FName ColorMetaData = "BindingColor";
+	SLATE_ARGUMENT(IHasContextClass*, ContextClassOwner)
+	SLATE_ARGUMENT(FString, TypeFilter);
+	SLATE_ARGUMENT(FString, BindingColor);
+            
+	SLATE_END_ARGS()
 
-	FString TypeFilter = PropertyHandle->GetMetaData(TypeMetaData);
-	FString BindingColor = PropertyHandle->GetMetaData(ColorMetaData);
-
-	TArray<UObject*> OuterObjects;
-	PropertyHandle->GetOuterObjects(OuterObjects);
-	UObject* OuterObject = OuterObjects[0];
-	
-	while (OuterObject && !OuterObject->Implements<UHasContextClass>())
+	TSharedRef<SWidget> CreatePropertyAccessWidget()
 	{
-		OuterObject = OuterObject->GetOuter();
-	}
+		FPropertyBindingWidgetArgs Args;
+		Args.bAllowPropertyBindings = true;
 
-	if (OuterObject)
-	{
-		IHasContextClass* HasContext = static_cast<IHasContextClass*>(OuterObject->GetInterfaceAddress(UHasContextClass::StaticClass()));
-		ContextClass = HasContext->GetContextClass();
-	}
-	
-	TSharedPtr<SWidget> Widget = SNullWidget::NullWidget;
-
-	FPropertyBindingWidgetArgs Args;
-	Args.bAllowPropertyBindings = true;
-	
-	Args.OnCanBindProperty = FOnCanBindProperty::CreateLambda([TypeFilter](FProperty* Property)
-	{
-		if (TypeFilter == "" || Property == nullptr)
+		UClass* ContextClass = UObject::StaticClass();
+		if (ContextClassOwner)
 		{
-			return true;
+			ContextClass = ContextClassOwner->GetContextClass();
 		}
-		else if (TypeFilter == "enum")
+		
+		Args.OnCanBindProperty = FOnCanBindProperty::CreateLambda([this](FProperty* Property)
 		{
-            if (const FEnumProperty* EnumProperty = CastField<const FEnumProperty>(Property))
-            {
-             	return true;
+			if (TypeFilter == "" || Property == nullptr)
+			{
+				return true;
 			}
-            else if (const FByteProperty* ByteProperty = CastField<const FByteProperty>(Property))
-            {
-            	return ByteProperty->Enum != nullptr;
-            }
-			return false;
-		}
+			else if (TypeFilter == "enum")
+			{
+				if (const FEnumProperty* EnumProperty = CastField<const FEnumProperty>(Property))
+				{
+					return true;
+				}
+				else if (const FByteProperty* ByteProperty = CastField<const FByteProperty>(Property))
+				{
+					return ByteProperty->Enum != nullptr;
+				}
+				return false;
+			}
+			
+			return Property->GetCPPType() == TypeFilter;
+		});
 		
-		return Property->GetCPPType() == TypeFilter;
-	});
-	
-	Args.OnCanBindToClass = FOnCanBindToClass::CreateLambda([ContextClass](UClass* InClass)
-	{
-		return true;
-	});
+		Args.OnCanBindToClass = FOnCanBindToClass::CreateLambda([](UClass* InClass)
+		{
+			return true;
+		});
 
 
-	FLinearColor BindingColorValue = FLinearColor::Gray;
-	if (BindingColor != "")
-	{
-		const UGraphEditorSettings* GraphEditorSettings = GetDefault<UGraphEditorSettings>();
-		if (const FStructProperty* ColorProperty = FindFProperty<FStructProperty>(GraphEditorSettings->GetClass(), FName(BindingColor)))
+		FLinearColor BindingColorValue = FLinearColor::Gray;
+		if (BindingColor != "")
 		{
-			BindingColorValue = *ColorProperty->ContainerPtrToValuePtr<FLinearColor>(GraphEditorSettings);
+			const UGraphEditorSettings* GraphEditorSettings = GetDefault<UGraphEditorSettings>();
+			if (const FStructProperty* ColorProperty = FindFProperty<FStructProperty>(GraphEditorSettings->GetClass(), FName(BindingColor)))
+			{
+				BindingColorValue = *ColorProperty->ContainerPtrToValuePtr<FLinearColor>(GraphEditorSettings);
+			}
 		}
-	}
 
-	Args.CurrentBindingColor = MakeAttributeLambda([BindingColorValue]() {
-		return BindingColorValue;
-	});
+		Args.CurrentBindingColor = MakeAttributeLambda([BindingColorValue]() {
+			return BindingColorValue;
+		});
 
-	Args.OnCanAcceptPropertyOrChildren = FOnCanBindProperty::CreateLambda([](FProperty* InProperty)
-		{
-			// Make only blueprint visible properties visible for binding.
-			return InProperty->HasAnyPropertyFlags(CPF_BlueprintVisible);
-		});	
+		Args.OnCanAcceptPropertyOrChildren = FOnCanBindProperty::CreateLambda([](FProperty* InProperty)
+			{
+				// Make only blueprint visible properties visible for binding.
+				return InProperty->HasAnyPropertyFlags(CPF_BlueprintVisible);
+			});	
 
-	Args.OnAddBinding = FOnAddBinding::CreateLambda([TypeFilter, PropertyHandle](FName InPropertyName, const TArray<FBindingChainElement>& InBindingChain)
-		{
-			TArray<UObject*> OuterObjects;
-			PropertyHandle->GetOuterObjects(OuterObjects);
-			const FScopedTransaction Transaction(NSLOCTEXT("ChooserPropertyBinding", "Change Property Binding", "Change Property Binding"));
-		
-			for (uint32 i=0; i<PropertyHandle->GetNumOuterObjects(); i++)
+		Args.OnAddBinding = FOnAddBinding::CreateLambda([this](FName InPropertyName, const TArray<FBindingChainElement>& InBindingChain)
 			{
-				void* ValuePtr = nullptr;
-				PropertyHandle->GetValueData(ValuePtr);// todo get per object value data?
-				FChooserPropertyBinding* PropertyValue = reinterpret_cast<FChooserPropertyBinding*>(ValuePtr);
+				TArray<UObject*> OuterObjects;
+				PropertyHandle->GetOuterObjects(OuterObjects);
+				const FScopedTransaction Transaction(NSLOCTEXT("ChooserPropertyBinding", "Change Property Binding", "Change Property Binding"));
 			
-				if (PropertyValue != nullptr)
+				for (uint32 i=0; i<PropertyHandle->GetNumOuterObjects(); i++)
 				{
-					PropertyHandle->NotifyPreChange();
-					
-					OuterObjects[i]->Modify(true);
-					Chooser::CopyPropertyChain(InBindingChain, PropertyValue->PropertyBindingChain);
-
-					if (TypeFilter == "enum")
+					void* ValuePtr = nullptr;
+					PropertyHandle->GetValueData(ValuePtr);// todo get per object value data?
+					FChooserPropertyBinding* PropertyValue = reinterpret_cast<FChooserPropertyBinding*>(ValuePtr);
+				
+					if (PropertyValue != nullptr)
 					{
-						FField* Property = InBindingChain.Last().Field.ToField();
-						FChooserEnumPropertyBinding* EnumPropertyValue = static_cast<FChooserEnumPropertyBinding*>(PropertyValue);
+						PropertyHandle->NotifyPreChange();
 						
-						if (const FEnumProperty* EnumProperty = CastField<const FEnumProperty>(Property))
-						{
-							EnumPropertyValue->Enum = EnumProperty->GetEnum();
-						}
-						else if (const FByteProperty* ByteProperty = CastField<const FByteProperty>(Property))
+						OuterObjects[i]->Modify(true);
+						Chooser::CopyPropertyChain(InBindingChain, PropertyValue->PropertyBindingChain);
+
+						if (TypeFilter == "enum")
 						{
-							EnumPropertyValue->Enum = ByteProperty->Enum;
+							FField* Property = InBindingChain.Last().Field.ToField();
+							FChooserEnumPropertyBinding* EnumPropertyValue = static_cast<FChooserEnumPropertyBinding*>(PropertyValue);
+							
+							if (const FEnumProperty* EnumProperty = CastField<const FEnumProperty>(Property))
+							{
+								EnumPropertyValue->Enum = EnumProperty->GetEnum();
+							}
+							else if (const FByteProperty* ByteProperty = CastField<const FByteProperty>(Property))
+							{
+								EnumPropertyValue->Enum = ByteProperty->Enum;
+							}
 						}
-					}
 
-					PropertyHandle->NotifyPostChange(EPropertyChangeType::ValueSet);
+						PropertyHandle->NotifyPostChange(EPropertyChangeType::ValueSet);
+					}
 				}
-			}
-		});
-
-	Args.CurrentBindingToolTipText = MakeAttributeLambda([PropertyHandle]()
-			{
-				const FText Bind = NSLOCTEXT("ContextPropertyWidget", "Bind", "Bind");
-				FText CurrentValue = Bind;
-				
-				void* ValuePtr = nullptr;
-				PropertyHandle->GetValueData(ValuePtr);// todo get per object value data?
-				FChooserPropertyBinding* PropertyValue = reinterpret_cast<FChooserPropertyBinding*>(ValuePtr);
+			});
 
-				if (PropertyValue != nullptr && PropertyValue->PropertyBindingChain.Num()>0)
+		Args.CurrentBindingToolTipText = MakeAttributeLambda([this]()
 				{
-					TArray<FText> BindingChainText;
-					BindingChainText.Reserve(PropertyValue->PropertyBindingChain.Num());
-				 
-					for (const FName& Name : PropertyValue->PropertyBindingChain)
+					const FText Bind = NSLOCTEXT("ContextPropertyWidget", "Bind", "Bind");
+					FText CurrentValue = Bind;
+					
+					void* ValuePtr = nullptr;
+					PropertyHandle->GetValueData(ValuePtr);// todo get per object value data?
+					FChooserPropertyBinding* PropertyValue = reinterpret_cast<FChooserPropertyBinding*>(ValuePtr);
+
+					if (PropertyValue != nullptr && PropertyValue->PropertyBindingChain.Num()>0)
 					{
-						BindingChainText.Add(FText::FromName(Name));
+						TArray<FText> BindingChainText;
+						BindingChainText.Reserve(PropertyValue->PropertyBindingChain.Num());
+					 
+						for (const FName& Name : PropertyValue->PropertyBindingChain)
+						{
+							BindingChainText.Add(FText::FromName(Name));
+						}
+						
+						CurrentValue = FText::Join(NSLOCTEXT("ContextPropertyWidget", "PropertyPathSeparator","."), BindingChainText);
 					}
-					
-					CurrentValue = FText::Join(NSLOCTEXT("ContextPropertyWidget", "PropertyPathSeparator","."), BindingChainText);
-				}
-	
-				return CurrentValue;	
-			});
-	
-	Args.CurrentBindingText = MakeAttributeLambda([PropertyHandle]()
-			{
-				const FText Bind = NSLOCTEXT("ContextPropertyWidget", "Bind", "Bind");
-				FText CurrentValue = Bind;
 		
-				void* ValuePtr = nullptr;
-				PropertyHandle->GetValueData(ValuePtr);// todo get per object value data?
-				FChooserPropertyBinding* PropertyValue = reinterpret_cast<FChooserPropertyBinding*>(ValuePtr);
-	
-				int BindingChainLength = PropertyValue->PropertyBindingChain.Num();
-				if (BindingChainLength > 0)
+					return CurrentValue;	
+				});
+		
+		Args.CurrentBindingText = MakeAttributeLambda([this]()
 				{
-					if (BindingChainLength == 1)
-					{
-						// single property, just use the property name
-						CurrentValue = FText::FromName(PropertyValue->PropertyBindingChain.Last());
-					}
-					else
+					const FText Bind = NSLOCTEXT("ContextPropertyWidget", "Bind", "Bind");
+					FText CurrentValue = Bind;
+			
+					void* ValuePtr = nullptr;
+					PropertyHandle->GetValueData(ValuePtr);// todo get per object value data?
+					FChooserPropertyBinding* PropertyValue = reinterpret_cast<FChooserPropertyBinding*>(ValuePtr);
+		
+					int BindingChainLength = PropertyValue->PropertyBindingChain.Num();
+					if (BindingChainLength > 0)
 					{
-						// for longer chains always show the last struct/object name, and the final property name (full path in tooltip)
-						CurrentValue = FText::Join(NSLOCTEXT("ContextPropertyWidget", "PropertyPathSeparator","."),
-							TArray<FText>({
-								FText::FromName(PropertyValue->PropertyBindingChain[BindingChainLength-2]),
-								FText::FromName(PropertyValue->PropertyBindingChain[BindingChainLength-1])
-							}));
+						if (BindingChainLength == 1)
+						{
+							// single property, just use the property name
+							CurrentValue = FText::FromName(PropertyValue->PropertyBindingChain.Last());
+						}
+						else
+						{
+							// for longer chains always show the last struct/object name, and the final property name (full path in tooltip)
+							CurrentValue = FText::Join(NSLOCTEXT("ContextPropertyWidget", "PropertyPathSeparator","."),
+								TArray<FText>({
+									FText::FromName(PropertyValue->PropertyBindingChain[BindingChainLength-2]),
+									FText::FromName(PropertyValue->PropertyBindingChain[BindingChainLength-1])
+								}));
+						}
 					}
-				}
-	
-				return CurrentValue;
+		
+					return CurrentValue;
+				});
+
+		Args.CurrentBindingImage = MakeAttributeLambda([]() -> const FSlateBrush*
+			{
+				static FName PropertyIcon(TEXT("Kismet.Tabs.Variables"));
+				return FAppStyle::GetBrush(PropertyIcon);
 			});
 
-	Args.CurrentBindingImage = MakeAttributeLambda([]() -> const FSlateBrush*
+		const IPropertyAccessEditor& PropertyAccessEditor = IModularFeatures::Get().GetModularFeature<IPropertyAccessEditor>("PropertyAccessEditor");
+		FBindingContextStruct StructInfo;
+		StructInfo.Struct = ContextClass;
+		return  PropertyAccessEditor.MakePropertyBindingWidget({StructInfo}, Args);
+	}
+
+	void UpdateWidget()
+	{
+		ChildSlot[ CreatePropertyAccessWidget() ];
+	}
+
+	void ContextClassChanged(UClass* NewContextClass)
+	{
+		UpdateWidget();
+	}
+
+	void Construct( const FArguments& InArgs, TSharedRef<IPropertyHandle> InPropertyHandle)
+	{
+		PropertyHandle = InPropertyHandle;
+		TypeFilter = InArgs._TypeFilter;
+		BindingColor = InArgs._BindingColor;
+		ContextClassOwner = InArgs._ContextClassOwner;
+		UpdateWidget();
+
+		if (ContextClassOwner)
 		{
-			static FName PropertyIcon(TEXT("Kismet.Tabs.Variables"));
-			return FAppStyle::GetBrush(PropertyIcon);
-		});
+			ContextClassOwner->OnContextClassChanged.AddSP(this, &SPropertyAccessChainWidget::ContextClassChanged);
+		}
+	}
+
+	~SPropertyAccessChainWidget()
+	{
+	}
+
+private:
+	FString TypeFilter;
+	FString BindingColor;
+	FDelegateHandle ContextClassChangedHandle();
+	TSharedPtr<IPropertyHandle> PropertyHandle;
+	IHasContextClass* ContextClassOwner = nullptr;
+};
+
+void FPropertyAccessChainCustomization::CustomizeHeader(TSharedRef<IPropertyHandle> PropertyHandle, FDetailWidgetRow& HeaderRow, IPropertyTypeCustomizationUtils& CustomizationUtils)
+{
+	FProperty* Property = PropertyHandle->GetProperty();
+
+	UClass* ContextClass = nullptr;
+	static FName TypeMetaData = "BindingType";
+	static FName ColorMetaData = "BindingColor";
+
+	FString TypeFilter = PropertyHandle->GetMetaData(TypeMetaData);
+	FString BindingColor = PropertyHandle->GetMetaData(ColorMetaData);
 
-	IPropertyAccessEditor& PropertyAccessEditor = IModularFeatures::Get().GetModularFeature<IPropertyAccessEditor>("PropertyAccessEditor");
-	FBindingContextStruct StructInfo;
-	StructInfo.Struct = ContextClass;
-	Widget =  PropertyAccessEditor.MakePropertyBindingWidget({StructInfo}, Args);
+	TArray<UObject*> OuterObjects;
+	PropertyHandle->GetOuterObjects(OuterObjects);
+	UObject* OuterObject = OuterObjects[0];
+	
+	while (OuterObject && !OuterObject->Implements<UHasContextClass>())
+	{
+		OuterObject = OuterObject->GetOuter();
+	}
+	
+	IHasContextClass* HasContext = nullptr;
+	if (OuterObject)
+	{
+		HasContext = static_cast<IHasContextClass*>(OuterObject->GetInterfaceAddress(UHasContextClass::StaticClass()));
+	}
 
 	HeaderRow
 	.NameContent()
@@ -210,7 +266,7 @@ void FPropertyAccessChainCustomization::CustomizeHeader(TSharedRef<IPropertyHand
 	]
 	.ValueContent()
 	[
-		Widget.ToSharedRef()
+		SNew(SPropertyAccessChainWidget, PropertyHandle).ContextClassOwner(HasContext).TypeFilter(TypeFilter).BindingColor(BindingColor)
 	];
 }
 	
diff --git a/Engine/Plugins/Experimental/Chooser/Source/ProxyTable/Internal/ProxyTable.h b/Engine/Plugins/Experimental/Chooser/Source/ProxyTable/Internal/ProxyTable.h
--- a/Engine/Plugins/Experimental/Chooser/Source/ProxyTable/Internal/ProxyTable.h
+++ b/Engine/Plugins/Experimental/Chooser/Source/ProxyTable/Internal/ProxyTable.h
@@ -27,6 +27,7 @@ public:
 
 	// caching the type so that on Undo, we can tell if we should fire the changed delegate
 	UClass* CachedPreviousType = nullptr;
+	UClass* CachedPreviousContextClass = nullptr;
 #endif
 	
 
diff --git a/Engine/Plugins/Experimental/Chooser/Source/ProxyTable/Private/ProxyTable.cpp b/Engine/Plugins/Experimental/Chooser/Source/ProxyTable/Private/ProxyTable.cpp
--- a/Engine/Plugins/Experimental/Chooser/Source/ProxyTable/Private/ProxyTable.cpp
+++ b/Engine/Plugins/Experimental/Chooser/Source/ProxyTable/Private/ProxyTable.cpp
@@ -17,6 +17,12 @@ void UProxyAsset::PostEditUndo()
 		OnTypeChanged.Broadcast(Type);
 		CachedPreviousType = Type;
 	}
+	
+	if (CachedPreviousContextClass != ContextClass)
+	{
+		OnContextClassChanged.Broadcast(ContextClass);
+		CachedPreviousContextClass = ContextClass;
+	}
 }
 
 void UProxyAsset::PostEditChangeProperty(FPropertyChangedEvent& PropertyChangedEvent)
@@ -24,6 +30,7 @@ void UProxyAsset::PostEditChangeProperty(FPropertyChangedEvent& PropertyChangedE
 	UObject::PostEditChangeProperty(PropertyChangedEvent);
 	
 	static FName TypeName = "Type";
+	static FName ContextClassName = "ContextClass";
 	if (PropertyChangedEvent.Property->GetName() == TypeName)
 	{
 		if (CachedPreviousType != Type)
@@ -32,6 +39,14 @@ void UProxyAsset::PostEditChangeProperty(FPropertyChangedEvent& PropertyChangedE
 		}
 		CachedPreviousType = Type;
 	}
+	else if (PropertyChangedEvent.Property->GetName() == ContextClassName)
+	{
+		if (CachedPreviousContextClass != ContextClass)
+		{
+			OnContextClassChanged.Broadcast(ContextClass);
+		}
+		CachedPreviousContextClass = ContextClass;
+	}
 }
 
 #endif
