Chooser editor layout fixes

diff --git a/Engine/Plugins/Experimental/Chooser/Source/ChooserEditor/Private/EnumColumnEditor.cpp b/Engine/Plugins/Experimental/Chooser/Source/ChooserEditor/Private/EnumColumnEditor.cpp
--- a/Engine/Plugins/Experimental/Chooser/Source/ChooserEditor/Private/EnumColumnEditor.cpp
+++ b/Engine/Plugins/Experimental/Chooser/Source/ChooserEditor/Private/EnumColumnEditor.cpp
@@ -86,38 +86,41 @@ public:
 		ChildSlot
 		[
 			SNew(SHorizontalBox)
-			+ SHorizontalBox::Slot().MaxWidth(45.0f)
+			+ SHorizontalBox::Slot().AutoWidth()
 			[
-				SNew(SButton).ButtonStyle(FAppStyle::Get(),"FlatButton").TextStyle(FAppStyle::Get(),"RichTextBlock.Bold").HAlign(HAlign_Center)
-				.Text_Lambda([this, Row]()
-				{
-					// this is a bit of an odd thing to do, but as we are polling for the "=" / "!=" button state anyway,
-					// also poll to see if the currently bound Enum has changed, and if it has, recreate the enum combo box.
-					const UEnum* CurrentEnumSource = nullptr;
-					if (EnumColumn->InputValue.IsValid())
-					{
-						CurrentEnumSource = EnumColumn->InputValue.Get<FChooserParameterEnumBase>().GetEnum(); 
-					}
-					if (EnumSource != CurrentEnumSource)
+				SNew(SBox).WidthOverride(45)
+				[
+					SNew(SButton).ButtonStyle(FAppStyle::Get(),"FlatButton").TextStyle(FAppStyle::Get(),"RichTextBlock.Bold").HAlign(HAlign_Center)
+					.Text_Lambda([this, Row]()
 					{
-						EnumComboBorder->SetContent(CreateEnumComboBox());
-						EnumSource = EnumColumn->InputValue.Get<FChooserParameterEnumBase>().GetEnum();
-					}
+						// this is a bit of an odd thing to do, but as we are polling for the "=" / "!=" button state anyway,
+						// also poll to see if the currently bound Enum has changed, and if it has, recreate the enum combo box.
+						const UEnum* CurrentEnumSource = nullptr;
+						if (EnumColumn->InputValue.IsValid())
+						{
+							CurrentEnumSource = EnumColumn->InputValue.Get<FChooserParameterEnumBase>().GetEnum(); 
+						}
+						if (EnumSource != CurrentEnumSource)
+						{
+							EnumComboBorder->SetContent(CreateEnumComboBox());
+							EnumSource = EnumColumn->InputValue.Get<FChooserParameterEnumBase>().GetEnum();
+						}
 					
-					return (EnumColumn->RowValues.IsValidIndex(Row) && EnumColumn->RowValues[Row].CompareNotEqual ? LOCTEXT("Not Equal", "!=") : LOCTEXT("Equal", "="));
-				})
-				.OnClicked_Lambda([this, Row]()
-				{
-					if (EnumColumn->RowValues.IsValidIndex(Row))
+						return (EnumColumn->RowValues.IsValidIndex(Row) && EnumColumn->RowValues[Row].CompareNotEqual ? LOCTEXT("Not Equal", "!=") : LOCTEXT("Equal", "="));
+					})
+					.OnClicked_Lambda([this, Row]()
 					{
-						const FScopedTransaction Transaction(LOCTEXT("Edit Comparison", "Edit Comparison Operation"));
-						TransactionObject->Modify(true);
-						EnumColumn->RowValues[Row].CompareNotEqual = !EnumColumn->RowValues[Row].CompareNotEqual;
-					}
-					return FReply::Handled();
-				})
+						if (EnumColumn->RowValues.IsValidIndex(Row))
+						{
+							const FScopedTransaction Transaction(LOCTEXT("Edit Comparison", "Edit Comparison Operation"));
+							TransactionObject->Modify(true);
+							EnumColumn->RowValues[Row].CompareNotEqual = !EnumColumn->RowValues[Row].CompareNotEqual;
+						}
+						return FReply::Handled();
+					})
+				]
 			]
-			+ SHorizontalBox::Slot().FillWidth(true)
+			+ SHorizontalBox::Slot().FillWidth(1)
 			[
 				SAssignNew(EnumComboBorder, SBorder).Padding(0).BorderBackgroundColor(FLinearColor(0,0,0,0))
 				[
diff --git a/Engine/Plugins/Experimental/Chooser/Source/ChooserEditor/Private/FloatRangeColumnEditor.cpp b/Engine/Plugins/Experimental/Chooser/Source/ChooserEditor/Private/FloatRangeColumnEditor.cpp
--- a/Engine/Plugins/Experimental/Chooser/Source/ChooserEditor/Private/FloatRangeColumnEditor.cpp
+++ b/Engine/Plugins/Experimental/Chooser/Source/ChooserEditor/Private/FloatRangeColumnEditor.cpp
@@ -22,11 +22,11 @@ TSharedRef<SWidget> CreateFloatRangeColumnWidget(UChooserTable* Chooser, FChoose
 	FFloatRangeColumn* FloatRangeColumn = static_cast<FFloatRangeColumn*>(Column);
 
 	return SNew(SHorizontalBox)
-	+ SHorizontalBox::Slot().MaxWidth(10)
+	+ SHorizontalBox::Slot().AutoWidth().VAlign(VAlign_Center)
 	[
 		SNew(STextBlock).Text(LOCTEXT("FloatRangeLeft", "("))
 	]
-	+ SHorizontalBox::Slot().FillWidth(1.0)
+	+ SHorizontalBox::Slot().FillWidth(0.5f)
 	[
 		SNew(SNumericEntryBox<float>)
 		.MaxValue_Lambda([FloatRangeColumn, Row]()
@@ -47,11 +47,11 @@ TSharedRef<SWidget> CreateFloatRangeColumnWidget(UChooserTable* Chooser, FChoose
 			}
 		})
 	]
-	+ SHorizontalBox::Slot().MaxWidth(10)
+	+ SHorizontalBox::Slot().AutoWidth().VAlign(VAlign_Center)
 	[
 		SNew(STextBlock).Text(LOCTEXT("FloatRangeComma", " ,"))
 	]
-	+ SHorizontalBox::Slot().FillWidth(1.0)
+	+ SHorizontalBox::Slot().FillWidth(0.5f)
 	[
 		SNew(SNumericEntryBox<float>)
 		.MinValue_Lambda([FloatRangeColumn, Row]()
@@ -72,7 +72,7 @@ TSharedRef<SWidget> CreateFloatRangeColumnWidget(UChooserTable* Chooser, FChoose
 			}
 		})
 	]
-	+ SHorizontalBox::Slot().MaxWidth(10)
+	+ SHorizontalBox::Slot().AutoWidth().VAlign(VAlign_Center)
 	[
 		SNew(STextBlock).Text(LOCTEXT("FloatRangeRight", " )"))
 	];
