Object Column for Chooser Table
PR #10259: added columns & cells for objects (assets) (Contributed by [REDACTED_PERSON])

diff --git a/Engine/Plugins/Experimental/Chooser/Source/Chooser/Internal/ObjectColumn.h b/Engine/Plugins/Experimental/Chooser/Source/Chooser/Internal/ObjectColumn.h
new file mode 100644
--- /dev/null
+++ b/Engine/Plugins/Experimental/Chooser/Source/Chooser/Internal/ObjectColumn.h
@@ -0,0 +1,104 @@
+// Copyright Epic Games, Inc. All Rights Reserved.
+#pragma once
+
+#include "ChooserPropertyAccess.h"
+#include "CoreMinimal.h"
+#include "IChooserColumn.h"
+#include "IChooserParameterObject.h"
+#include "InstancedStruct.h"
+#include "ObjectColumn.generated.h"
+
+struct FBindingChainElement;
+
+USTRUCT(DisplayName = "Object Property Binding")
+struct CHOOSER_API FObjectContextProperty : public FChooserParameterObjectBase
+{
+	GENERATED_BODY()
+
+	UPROPERTY(EditAnywhere, Meta = (BindingType = "object", BindingColor = "ObjectPinTypeColor"), Category = "Binding")
+	FChooserObjectPropertyBinding Binding;
+
+	virtual bool GetValue(const UObject* ContextObject, FSoftObjectPath& OutResult) const override;
+
+#if WITH_EDITOR
+
+	static bool CanBind(const FProperty& Property)
+	{
+		return Property.IsA<FObjectPropertyBase>() && !(Property.IsA<FClassProperty>() || Property.IsA<FSoftClassProperty>());
+	}
+
+	void SetBinding(const TArray<FBindingChainElement>& InBindingChain);
+
+	virtual void GetDisplayName(FText& OutName) const override
+	{
+		if (!Binding.PropertyBindingChain.IsEmpty())
+		{
+			OutName = FText::FromName(Binding.PropertyBindingChain.Last());
+		}
+	}
+
+	virtual UClass* GetAllowedClass() const override { return Binding.AllowedClass; }
+
+#endif
+};
+
+UENUM()
+enum class EObjectColumnCellValueComparison
+{
+	MatchEqual,
+	MatchNotEqual,
+	MatchAny,
+
+	Modulus // used for cycling through the other values
+};
+
+USTRUCT()
+struct FChooserObjectRowData
+{
+	GENERATED_BODY()
+
+	UPROPERTY(EditAnywhere, Category = "Runtime")
+	EObjectColumnCellValueComparison Comparison = EObjectColumnCellValueComparison::MatchEqual;
+
+	UPROPERTY(EditAnywhere, Category = "Runtime")
+	TSoftObjectPtr<UObject> Value;
+
+	bool Evaluate(const FSoftObjectPath& LeftHandSide) const;
+};
+
+USTRUCT()
+struct CHOOSER_API FObjectColumn : public FChooserColumnBase
+{
+	GENERATED_BODY()
+
+	FObjectColumn();
+
+	UPROPERTY(EditAnywhere, Meta = (ExcludeBaseStruct, BaseStruct = "/Script/Chooser.ChooserParameterObjectBase"), Category = "Data")
+	FInstancedStruct InputValue;
+
+#if WITH_EDITORONLY_DATA
+	UPROPERTY(EditAnywhere, Category = "Data")
+	FChooserObjectRowData DefaultRowValue;
+#endif
+
+	UPROPERTY(EditAnywhere, Category = "Data")
+	// array of results (cells for this column for each row in the table)
+	// should match the length of the Results array
+	TArray<FChooserObjectRowData> RowValues;
+
+	virtual void Filter(const UObject* ContextObject, const TArray<uint32>& IndexListIn, TArray<uint32>& IndexListOut) const override;
+
+	CHOOSER_COLUMN_BOILERPLATE(FChooserParameterObjectBase);
+
+#if WITH_EDITOR
+	virtual void PostLoad() override
+	{
+		Super::PostLoad();
+
+		if (InputValue.IsValid())
+		{
+			InputValue.GetMutable<FChooserParameterBase>().PostLoad();
+		}
+	}
+#endif
+};
diff --git a/Engine/Plugins/Experimental/Chooser/Source/Chooser/Private/ObjectColumn.cpp b/Engine/Plugins/Experimental/Chooser/Source/Chooser/Private/ObjectColumn.cpp
new file mode 100644
--- /dev/null
+++ b/Engine/Plugins/Experimental/Chooser/Source/Chooser/Private/ObjectColumn.cpp
@@ -0,0 +1,106 @@
+// Copyright Epic Games, Inc. All Rights Reserved.
+#include "ObjectColumn.h"
+#include "ChooserPropertyAccess.h"
+
+#if WITH_EDITOR
+	#include "IPropertyAccessEditor.h"
+#endif
+
+bool FObjectContextProperty::GetValue(const UObject* ContextObject, FSoftObjectPath& OutResult) const
+{
+	UStruct* StructType = ContextObject->GetClass();
+	const void* Container = ContextObject;
+
+	if (UE::Chooser::ResolvePropertyChain(Container, StructType, Binding.PropertyBindingChain))
+	{
+		if (const FObjectPropertyBase* ObjectProperty = FindFProperty<FObjectPropertyBase>(StructType, Binding.PropertyBindingChain.Last()))
+		{
+			// If the property has a value, then create a soft object path from it;
+			// otherwise, it could be a weak pointer to something that is not loaded.
+			const UObject* LoadedObject = ObjectProperty->GetObjectPropertyValue_InContainer(Container);
+			if (LoadedObject != nullptr)
+			{
+				OutResult = LoadedObject;
+				return true;
+			}
+
+			if (ObjectProperty->IsA<FSoftObjectProperty>())
+			{
+				const FSoftObjectPtr& SoftObjectPtr = *ObjectProperty->ContainerPtrToValuePtr<FSoftObjectPtr>(Container);
+				OutResult = SoftObjectPtr.ToSoftObjectPath();
+				return true;
+			}
+		}
+	}
+
+	return false;
+}
+
+#if WITH_EDITOR
+
+void FObjectContextProperty::SetBinding(const TArray<FBindingChainElement>& InBindingChain)
+{
+	const UClass* PreviousClass = Binding.AllowedClass;
+	Binding.AllowedClass = nullptr;
+
+	UE::Chooser::CopyPropertyChain(InBindingChain, Binding.PropertyBindingChain);
+
+	const FField* Field = InBindingChain.Last().Field.ToField();
+	if (const FObjectPropertyBase* ObjectProperty = CastField<FObjectPropertyBase>(Field))
+	{
+		Binding.AllowedClass = ObjectProperty->PropertyClass;
+	}
+}
+
+#endif // WITH_EDITOR
+
+FObjectColumn::FObjectColumn()
+{
+#if WITH_EDITOR
+	InputValue.InitializeAs(FObjectContextProperty::StaticStruct());
+#endif
+}
+
+bool FChooserObjectRowData::Evaluate(const FSoftObjectPath& LeftHandSide) const
+{
+	switch (Comparison)
+	{
+		case EObjectColumnCellValueComparison::MatchEqual:
+			return LeftHandSide == Value.ToSoftObjectPath();
+
+		case EObjectColumnCellValueComparison::MatchNotEqual:
+			return LeftHandSide != Value.ToSoftObjectPath();
+
+		case EObjectColumnCellValueComparison::MatchAny:
+			return true;
+
+		default:
+			return false;
+	}
+}
+
+void FObjectColumn::Filter(const UObject* ContextObject, const TArray<uint32>& IndexListIn, TArray<uint32>& IndexListOut) const
+{
+	FSoftObjectPath Result;
+	if (ContextObject != nullptr &&
+		InputValue.IsValid() &&
+		InputValue.Get<FChooserParameterObjectBase>().GetValue(ContextObject, Result))
+	{
+		for (const uint32 Index : IndexListIn)
+		{
+			if (RowValues.IsValidIndex(Index))
+			{
+				const FChooserObjectRowData& RowValue = RowValues[Index];
+				if (RowValue.Evaluate(Result))
+				{
+					IndexListOut.Emplace(Index);
+				}
+			}
+		}
+	}
+	else
+	{
+		// passthrough fallback (behaves better during live editing)
+		IndexListOut = IndexListIn;
+	}
+}
diff --git a/Engine/Plugins/Experimental/Chooser/Source/Chooser/Public/ChooserPropertyAccess.h b/Engine/Plugins/Experimental/Chooser/Source/Chooser/Public/ChooserPropertyAccess.h
--- a/Engine/Plugins/Experimental/Chooser/Source/Chooser/Public/ChooserPropertyAccess.h
+++ b/Engine/Plugins/Experimental/Chooser/Source/Chooser/Public/ChooserPropertyAccess.h
@@ -47,6 +47,17 @@ struct FChooserEnumPropertyBinding : public FChooserPropertyBinding
 #endif
 };
 
+USTRUCT()
+struct FChooserObjectPropertyBinding : public FChooserPropertyBinding
+{
+	GENERATED_BODY()
+
+#if WITH_EDITORONLY_DATA
+	UPROPERTY()
+	TObjectPtr<UClass> AllowedClass = nullptr;
+#endif
+};
+
 namespace UE::Chooser
 {
 	CHOOSER_API bool ResolvePropertyChain(const void*& Container, UStruct*& StructType, const TArray<FName>& PropertyBindingChain);
diff --git a/Engine/Plugins/Experimental/Chooser/Source/Chooser/Public/IChooserParameterObject.h b/Engine/Plugins/Experimental/Chooser/Source/Chooser/Public/IChooserParameterObject.h
new file mode 100644
--- /dev/null
+++ b/Engine/Plugins/Experimental/Chooser/Source/Chooser/Public/IChooserParameterObject.h
@@ -0,0 +1,21 @@
+// Copyright Epic Games, Inc. All Rights Reserved.
+
+#pragma once
+
+#include "CoreMinimal.h"
+#include "IChooserParameterBase.h"
+#include "IChooserParameterObject.generated.h"
+
+class UObject;
+
+USTRUCT()
+struct FChooserParameterObjectBase : public FChooserParameterBase
+{
+	GENERATED_BODY()
+
+	virtual bool GetValue(const UObject* ContextObject, FSoftObjectPath& OutResult) const { return false; }
+
+#if WITH_EDITOR
+	virtual UClass* GetAllowedClass() const { return UObject::StaticClass(); }
+#endif
+};
diff --git a/Engine/Plugins/Experimental/Chooser/Source/ChooserEditor/Private/ChooserEditorModule.cpp b/Engine/Plugins/Experimental/Chooser/Source/ChooserEditor/Private/ChooserEditorModule.cpp
--- a/Engine/Plugins/Experimental/Chooser/Source/ChooserEditor/Private/ChooserEditorModule.cpp
+++ b/Engine/Plugins/Experimental/Chooser/Source/ChooserEditor/Private/ChooserEditorModule.cpp
@@ -7,6 +7,7 @@
 #include "EnumColumnEditor.h"
 #include "FloatRangeColumnEditor.h"
 #include "GameplayTagColumnEditor.h"
+#include "ObjectColumnEditor.h"
 #include "ChooserTableEditorCommands.h"
 #include "ChooserPropertyAccess.h"
 #include "PropertyEditorModule.h"
@@ -23,6 +24,7 @@ void FModule::StartupModule()
 	RegisterFloatRangeWidgets();
 	RegisterBoolWidgets();
 	RegisterEnumWidgets();
+	RegisterObjectWidgets();
 	
 	FChooserTableEditorCommands::Register();
 
diff --git a/Engine/Plugins/Experimental/Chooser/Source/ChooserEditor/Private/ObjectColumnEditor.cpp b/Engine/Plugins/Experimental/Chooser/Source/ChooserEditor/Private/ObjectColumnEditor.cpp
new file mode 100644
--- /dev/null
+++ b/Engine/Plugins/Experimental/Chooser/Source/ChooserEditor/Private/ObjectColumnEditor.cpp
@@ -0,0 +1,157 @@
+// Copyright Epic Games, Inc. All Rights Reserved.
+
+#include "ObjectColumnEditor.h"
+#include "ContextPropertyWidget.h"
+#include "GraphEditorSettings.h"
+#include "ObjectChooserWidgetFactories.h"
+#include "ObjectColumn.h"
+#include "PropertyCustomizationHelpers.h"
+#include "TransactionCommon.h"
+#include "Widgets/Input/SButton.h"
+#include "Widgets/Layout/SBox.h"
+
+#define LOCTEXT_NAMESPACE "ObjectColumnEditor"
+
+namespace UE::ChooserEditor
+{
+	static UClass* GetAllowedClass(const FObjectColumn* ObjectColumn)
+	{
+		UClass* AllowedClass = nullptr;
+		if (const FChooserParameterObjectBase* InputValue = ObjectColumn->InputValue.GetPtr<FChooserParameterObjectBase>())
+		{
+			AllowedClass = InputValue->GetAllowedClass();
+		}
+
+		if (AllowedClass == nullptr)
+		{
+			AllowedClass = UObject::StaticClass();
+		}
+
+		return AllowedClass;
+	}
+
+	static TSharedRef<SWidget> CreateObjectPicker(UObject* TransactionObject, FObjectColumn* ObjectColumn, int32 Row)
+	{
+		return SNew(SObjectPropertyEntryBox)
+			.ObjectPath_Lambda([ObjectColumn, Row]() {
+				return ObjectColumn->RowValues.IsValidIndex(Row) ? ObjectColumn->RowValues[Row].Value.ToString() : FString();
+			})
+			.AllowedClass(GetAllowedClass(ObjectColumn))
+			.OnObjectChanged_Lambda([TransactionObject, ObjectColumn, Row](const FAssetData& AssetData) {
+				if (ObjectColumn->RowValues.IsValidIndex(Row))
+				{
+					const FScopedTransaction Transaction(LOCTEXT("Edit RHS", "Edit Object Value"));
+					TransactionObject->Modify(true);
+					ObjectColumn->RowValues[Row].Value = AssetData.ToSoftObjectPath();
+				}
+			})
+			.DisplayUseSelected(false)
+			.DisplayBrowse(false)
+			.DisplayThumbnail(false)
+			.Visibility_Lambda([ObjectColumn, Row]() {
+				return (ObjectColumn->RowValues.IsValidIndex(Row) &&
+						ObjectColumn->RowValues[Row].Comparison == EObjectColumnCellValueComparison::MatchAny)
+						   ? EVisibility::Collapsed
+						   : EVisibility::Visible;
+			});
+	}
+
+	static TSharedRef<SWidget> CreateObjectColumnWidget(UChooserTable* Chooser, FChooserColumnBase* Column, int32 Row)
+	{
+		// Extend `SHorizontalBox` a little bit so we can poll for changes & recreate the object picker if necessary.
+		struct SHorizontalBoxEx : public SHorizontalBox
+		{
+			virtual void Tick(const FGeometry& AllottedGeometry, const double InCurrentTime, const float InDeltaTime) override
+			{
+				SHorizontalBox::Tick(AllottedGeometry, InCurrentTime, InDeltaTime);
+
+				UClass* CurrentAllowedClass = GetAllowedClass(ObjectColumn);
+				if (AllowedClass != CurrentAllowedClass)
+				{
+					AllowedClass = CurrentAllowedClass;
+					GetSlot(1)[ObjectPickerFactory()];
+				}
+			}
+
+			FObjectColumn* ObjectColumn = nullptr;
+			UClass* AllowedClass = nullptr;
+			TFunction<TSharedRef<SWidget>()> ObjectPickerFactory;
+		};
+
+		FObjectColumn* ObjectColumn = static_cast<FObjectColumn*>(Column);
+
+		// clang-format off
+		TSharedRef<SHorizontalBoxEx> CellWidget = SNew(SHorizontalBoxEx)
+			+SHorizontalBox::Slot()
+			.AutoWidth()
+			[
+				SNew(SBox)
+				.WidthOverride(55)
+				[
+					SNew(SButton)
+					.ButtonStyle(FAppStyle::Get(), "FlatButton")
+					.TextStyle(FAppStyle::Get(), "RichTextBlock.Bold")
+					.HAlign(HAlign_Center)
+					.Text_Lambda([ObjectColumn, Row]() {
+						if (ObjectColumn->RowValues.IsValidIndex(Row))
+						{
+							switch (ObjectColumn->RowValues[Row].Comparison)
+							{
+								case EObjectColumnCellValueComparison::MatchEqual:
+									return LOCTEXT("CompEqual", "=");
+
+								case EObjectColumnCellValueComparison::MatchNotEqual:
+									return LOCTEXT("CompNotEqual", "!=");
+
+								case EObjectColumnCellValueComparison::MatchAny:
+									return LOCTEXT("CompAny", "Any");
+							}
+						}
+						return FText::GetEmpty();
+					})
+					.OnClicked_Lambda([Chooser, ObjectColumn, Row]() {
+						if (ObjectColumn->RowValues.IsValidIndex(Row))
+						{
+							const FScopedTransaction Transaction(LOCTEXT("Edit Comparison", "Edit Comparison Operation"));
+							Chooser->Modify(true);
+							// cycle through comparison options
+							EObjectColumnCellValueComparison& Comparison = ObjectColumn->RowValues[Row].Comparison;
+							const int32 NextComparison = (static_cast<int32>(Comparison) + 1) % static_cast<int32>(EObjectColumnCellValueComparison::Modulus);
+							Comparison = static_cast<EObjectColumnCellValueComparison>(NextComparison);
+						}
+						return FReply::Handled();
+					})
+				]
+			]
+			+SHorizontalBox::Slot()
+			.FillWidth(1)
+			[
+				CreateObjectPicker(Chooser, ObjectColumn, Row)
+			];
+		// clang-format on
+
+		// Setup the things we need to check for changes to the allowed class.
+		{
+			CellWidget->SetCanTick(true);
+			CellWidget->ObjectColumn = ObjectColumn;
+			CellWidget->AllowedClass = GetAllowedClass(ObjectColumn);
+			CellWidget->ObjectPickerFactory = [Chooser, ObjectColumn, Row]() { return CreateObjectPicker(Chooser, ObjectColumn, Row); };
+		}
+
+		return CellWidget;
+	}
+
+	static TSharedRef<SWidget> CreateObjectPropertyWidget(bool bReadOnly, UObject* TransactionObject, void* Value, UClass* ContextClass, UClass* ResultBaseClass)
+	{
+		return CreatePropertyWidget<FObjectContextProperty>(bReadOnly, TransactionObject, Value, ContextClass, GetDefault<UGraphEditorSettings>()->ObjectPinTypeColor);
+	}
+
+	void RegisterObjectWidgets()
+	{
+		FObjectChooserWidgetFactories::RegisterWidgetCreator(FObjectContextProperty::StaticStruct(), CreateObjectPropertyWidget);
+		FObjectChooserWidgetFactories::RegisterColumnWidgetCreator(FObjectColumn::StaticStruct(), CreateObjectColumnWidget);
+	}
+
+} // namespace UE::ChooserEditor
+
+#undef LOCTEXT_NAMESPACE
diff --git a/Engine/Plugins/Experimental/Chooser/Source/ChooserEditor/Private/ObjectColumnEditor.h b/Engine/Plugins/Experimental/Chooser/Source/ChooserEditor/Private/ObjectColumnEditor.h
new file mode 100644
--- /dev/null
+++ b/Engine/Plugins/Experimental/Chooser/Source/ChooserEditor/Private/ObjectColumnEditor.h
@@ -0,0 +1,10 @@
+// Copyright Epic Games, Inc. All Rights Reserved.
+
+#pragma once
+
+#include "CoreMinimal.h"
+
+namespace UE::ChooserEditor
+{
+	void RegisterObjectWidgets();
+}
