Fix static analysis error in ProxyTable.cpp

diff --git a/Engine/Plugins/Experimental/Chooser/Source/ProxyTable/Private/ProxyTable.cpp b/Engine/Plugins/Experimental/Chooser/Source/ProxyTable/Private/ProxyTable.cpp
--- a/Engine/Plugins/Experimental/Chooser/Source/ProxyTable/Private/ProxyTable.cpp
+++ b/Engine/Plugins/Experimental/Chooser/Source/ProxyTable/Private/ProxyTable.cpp
@@ -324,14 +324,11 @@ UObject* UProxyTableFunctionLibrary::EvaluateProxyTable(const UObject* ContextOb
 {
 	if (ProxyTable)
 	{
-		if(ProxyTable)
+		FGuid Guid;
+		Guid.A = GetTypeHash(Key);
+		if (UObject* Value = ProxyTable->FindProxyObject(Guid, ContextObject))
 		{
-			FGuid Guid;
-			Guid.A = GetTypeHash(Key);
-			if (UObject* Value = ProxyTable->FindProxyObject(Guid, ContextObject))
-			{
-				return Value;
-			}
+			return Value;
 		}
 	}
 	
