Fix post load ordering when using Zen Loader

diff --git a/Engine/Plugins/Experimental/Chooser/Source/ProxyTable/Private/ProxyTable.cpp b/Engine/Plugins/Experimental/Chooser/Source/ProxyTable/Private/ProxyTable.cpp
--- a/Engine/Plugins/Experimental/Chooser/Source/ProxyTable/Private/ProxyTable.cpp
+++ b/Engine/Plugins/Experimental/Chooser/Source/ProxyTable/Private/ProxyTable.cpp
@@ -135,6 +135,11 @@ static void BuildRuntimeDataRecursive(UProxyTable* RootTable, UProxyTable* Table
 	
 	for (const FProxyEntry& Entry : Table->Entries)
 	{
+		if (Entry.Proxy)
+		{
+			Entry.Proxy->ConditionalPostLoad();
+		}
+
 		int32 FoundIndex = OutEntriesArray.Find(Entry);
 		if (FoundIndex == INDEX_NONE)
 		{
