Fix issue with new Proxy Tables never triggering cook code until the editor reloads

diff --git a/Engine/Plugins/Experimental/Chooser/Source/ProxyTable/Private/ProxyTable.cpp b/Engine/Plugins/Experimental/Chooser/Source/ProxyTable/Private/ProxyTable.cpp
--- a/Engine/Plugins/Experimental/Chooser/Source/ProxyTable/Private/ProxyTable.cpp
+++ b/Engine/Plugins/Experimental/Chooser/Source/ProxyTable/Private/ProxyTable.cpp
@@ -130,8 +130,11 @@ static void BuildRuntimeDataRecursive(UProxyTable* RootTable, UProxyTable* Table
 	{
 		return;
 	}
-	
-	OutDependencies.Add(Table);
+
+	if (Table != RootTable)
+	{
+		OutDependencies.Add(Table);
+	}
 	
 	for (const FProxyEntry& Entry : Table->Entries)
 	{
@@ -223,11 +226,11 @@ void UProxyTable::BuildRuntimeData()
 	for(const FProxyEntry& Entry : RuntimeEntries) 
 	{
 		if (Entry.Proxy)
-			{
+		{
 			ProxyDependencies.Add(Entry.Proxy);
-			}
 		}
 	}
+}
 
 
 void UProxyTable::PostLoad()
@@ -239,6 +242,7 @@ void UProxyTable::PostLoad()
 void UProxyTable::PostTransacted(const FTransactionObjectEvent& TransactionEvent)
 {
 	UObject::PostTransacted(TransactionEvent);
+	BuildRuntimeData();
 	OnProxyTableChanged.Broadcast();
 }
 
