Allow calling thread safe functions for chooser column properties, including object functions in the middle of the property chain.  Exclude traversing object properties in order to force use of functions marked as thread safe.  Merge duplicated implementations of property access ui for the details panel and the column headers

diff --git a/Engine/Plugins/Experimental/Chooser/Source/Chooser/Internal/BoolColumn.h b/Engine/Plugins/Experimental/Chooser/Source/Chooser/Internal/BoolColumn.h
--- a/Engine/Plugins/Experimental/Chooser/Source/Chooser/Internal/BoolColumn.h
+++ b/Engine/Plugins/Experimental/Chooser/Source/Chooser/Internal/BoolColumn.h
@@ -25,7 +25,7 @@ public:
 	UPROPERTY()
 	TArray<FName> PropertyBindingChain_DEPRECATED;
 	
-	UPROPERTY(EditAnywhere, Meta = (BindingType = "bool", BindingColor = "BooleanPinTypeColor"), Category = "Binding")
+	UPROPERTY(EditAnywhere, Meta = (BindingType = "bool", BindingAllowFunctions = "true", BindingColor = "BooleanPinTypeColor"), Category = "Binding")
 	FChooserPropertyBinding Binding;
 	
 	virtual bool GetValue(const UObject* ContextObject, bool& OutResult) const override;
diff --git a/Engine/Plugins/Experimental/Chooser/Source/Chooser/Internal/FloatRangeColumn.h b/Engine/Plugins/Experimental/Chooser/Source/Chooser/Internal/FloatRangeColumn.h
--- a/Engine/Plugins/Experimental/Chooser/Source/Chooser/Internal/FloatRangeColumn.h
+++ b/Engine/Plugins/Experimental/Chooser/Source/Chooser/Internal/FloatRangeColumn.h
@@ -18,7 +18,7 @@ struct CHOOSER_API FFloatContextProperty :  public FChooserParameterFloatBase
 	UPROPERTY()
 	TArray<FName> PropertyBindingChain_DEPRECATED;
 	
-	UPROPERTY(EditAnywhere, Meta = (BindingType = "double", BindingColor = "FloatPinTypeColor"), Category = "Binding")
+	UPROPERTY(EditAnywhere, Meta = (BindingType = "double", BindingAllowFunctions = "true", BindingColor = "FloatPinTypeColor"), Category = "Binding")
 	FChooserPropertyBinding Binding;
 
 	virtual void PostLoad() override
diff --git a/Engine/Plugins/Experimental/Chooser/Source/Chooser/Private/BoolColumn.cpp b/Engine/Plugins/Experimental/Chooser/Source/Chooser/Private/BoolColumn.cpp
--- a/Engine/Plugins/Experimental/Chooser/Source/Chooser/Private/BoolColumn.cpp
+++ b/Engine/Plugins/Experimental/Chooser/Source/Chooser/Private/BoolColumn.cpp
@@ -4,6 +4,7 @@
 
 bool FBoolContextProperty::GetValue(const UObject* ContextObject, bool& OutResult) const
 {
+	
 	UStruct* StructType = ContextObject->GetClass();
 	const void* Container = ContextObject;
 	
@@ -14,6 +15,23 @@ bool FBoolContextProperty::GetValue(const UObject* ContextObject, bool& OutResul
 			OutResult = *Property->ContainerPtrToValuePtr<bool>(Container);
 			return true;
 		}
+		
+	    if (UClass* ClassType = Cast<UClass>(StructType))
+	    {
+			if (UFunction* Function = ClassType->FindFunctionByName(Binding.PropertyBindingChain.Last()))
+			{
+				UObject* Object = reinterpret_cast<UObject*>(const_cast<void*>(Container));
+				if (Function->IsNative())
+				{
+					FFrame Stack(Object, Function, nullptr, nullptr, Function->ChildProperties);
+					Function->Invoke(Object, Stack, &OutResult);
+				}
+				else
+				{
+					Object->ProcessEvent(Function, &OutResult);
+				}
+			} 
+		}
 	}
 
 	return false;
diff --git a/Engine/Plugins/Experimental/Chooser/Source/Chooser/Private/ChooserPropertyAccess.cpp b/Engine/Plugins/Experimental/Chooser/Source/Chooser/Private/ChooserPropertyAccess.cpp
--- a/Engine/Plugins/Experimental/Chooser/Source/Chooser/Private/ChooserPropertyAccess.cpp
+++ b/Engine/Plugins/Experimental/Chooser/Source/Chooser/Private/ChooserPropertyAccess.cpp
@@ -35,7 +35,40 @@ namespace UE::Chooser
 			}
 			else
 			{
-				return false;
+				// check if it's a member function
+				if (UClass* ClassType = Cast<UClass>(StructType))
+				{
+					if (UFunction* Function = ClassType->FindFunctionByName(PropertyBindingChain[PropertyChainIndex]))
+					{
+						UObject* Object = reinterpret_cast<UObject*>(const_cast<void*>(Container));
+						if (Function->IsNative())
+						{
+							FFrame Stack(Object, Function, nullptr, nullptr, Function->ChildProperties);
+							Function->Invoke(Object, Stack, &Container);
+						}
+						else
+						{
+							Object->ProcessEvent(Function, &Container);
+						}
+						
+						if (Container == nullptr)
+						{
+							return false;
+						}
+						else
+						{
+							StructType = reinterpret_cast<UObject*>(const_cast<void*>(Container))->GetClass();
+						}
+					}
+					else
+					{
+						return false;
+					}
+				}
+				else
+				{
+					return false;
+				}
 			}
 		}
 	
diff --git a/Engine/Plugins/Experimental/Chooser/Source/Chooser/Private/FloatRangeColumn.cpp b/Engine/Plugins/Experimental/Chooser/Source/Chooser/Private/FloatRangeColumn.cpp
--- a/Engine/Plugins/Experimental/Chooser/Source/Chooser/Private/FloatRangeColumn.cpp
+++ b/Engine/Plugins/Experimental/Chooser/Source/Chooser/Private/FloatRangeColumn.cpp
@@ -6,7 +6,7 @@ bool FFloatContextProperty::GetValue(const UObject* ContextObject, float& OutRes
 {
 	UStruct* StructType = ContextObject->GetClass();
 	const void* Container = ContextObject;
-	
+
 	if (UE::Chooser::ResolvePropertyChain(Container, StructType, Binding.PropertyBindingChain))
 	{
 		if (const FDoubleProperty* DoubleProperty = FindFProperty<FDoubleProperty>(StructType, Binding.PropertyBindingChain.Last()))
@@ -20,6 +20,44 @@ bool FFloatContextProperty::GetValue(const UObject* ContextObject, float& OutRes
 			OutResult = *FloatProperty->ContainerPtrToValuePtr<float>(Container);
 			return true;
 		}
+
+	    if (UClass* ClassType = Cast<UClass>(StructType))
+	    {
+			if (UFunction* Function = ClassType->FindFunctionByName(Binding.PropertyBindingChain.Last()))
+			{
+				bool bReturnsDouble = CastField<FDoubleProperty>(Function->GetReturnProperty()) != nullptr;
+					
+				UObject* Object = reinterpret_cast<UObject*>(const_cast<void*>(Container));
+				if (Function->IsNative())
+				{
+					FFrame Stack(Object, Function, nullptr, nullptr, Function->ChildProperties);
+					if (bReturnsDouble)
+					{
+						double result;
+						Function->Invoke(Object, Stack, &result);
+						OutResult = result;
+					}
+					else
+					{
+						Function->Invoke(Object, Stack, &OutResult);
+					}
+				}
+				else
+				{
+					if (bReturnsDouble)
+					{
+						double result = 0;
+						Object->ProcessEvent(Function, &result);
+						OutResult = result;
+					}
+					else
+					{
+						Object->ProcessEvent(Function, &OutResult);
+					}
+				}
+				return true;
+			} 
+		}
 	}
 
 	return false;
diff --git a/Engine/Plugins/Experimental/Chooser/Source/ChooserEditor/Private/BoolColumnEditor.cpp b/Engine/Plugins/Experimental/Chooser/Source/ChooserEditor/Private/BoolColumnEditor.cpp
--- a/Engine/Plugins/Experimental/Chooser/Source/ChooserEditor/Private/BoolColumnEditor.cpp
+++ b/Engine/Plugins/Experimental/Chooser/Source/ChooserEditor/Private/BoolColumnEditor.cpp
@@ -3,9 +3,9 @@
 #include "BoolColumnEditor.h"
 #include "BoolColumn.h"
 #include "OutputBoolColumn.h"
-#include "ContextPropertyWidget.h"
-#include "ChooserTableEditor.h"
+#include "SPropertyAccessChainWidget.h"
 #include "ObjectChooserWidgetFactories.h"
+#include "Widgets/SBoxPanel.h"
 #include "Widgets/Input/SButton.h"
 #include "Widgets/Input/SCheckBox.h"
 #include "GraphEditorSettings.h"
@@ -57,26 +57,46 @@ TSharedRef<SWidget> CreateOutputBoolColumnWidget(UChooserTable* Chooser, FChoose
 {
 	FOutputBoolColumn* BoolColumn = static_cast<FOutputBoolColumn*>(Column);
 
-	return SNew (SCheckBox)
-	.OnCheckStateChanged_Lambda([Chooser, BoolColumn,Row](ECheckBoxState State)
-	{
-		if (Row < BoolColumn->RowValues.Num())
+	return
+	
+	SNew(SHorizontalBox)
+	+ SHorizontalBox::Slot().FillWidth(1)
+	+ SHorizontalBox::Slot().AutoWidth().VAlign(VAlign_Center)
+	[
+		SNew (SCheckBox)
+		.OnCheckStateChanged_Lambda([Chooser, BoolColumn,Row](ECheckBoxState State)
+		{
+			if (Row < BoolColumn->RowValues.Num())
+			{
+				const FScopedTransaction Transaction(LOCTEXT("Change Bool Value", "Change Bool Value"));
+				Chooser->Modify(true);
+				BoolColumn->RowValues[Row] = (State == ECheckBoxState::Checked);
+			}
+		})
+		.IsChecked_Lambda([BoolColumn, Row]()
 		{
-			const FScopedTransaction Transaction(LOCTEXT("Change Bool Value", "Change Bool Value"));
-			Chooser->Modify(true);
-			BoolColumn->RowValues[Row] = (State == ECheckBoxState::Checked);
-		}
-	})
-	.IsChecked_Lambda([BoolColumn, Row]()
-	{
-		const bool value = (Row < BoolColumn->RowValues.Num()) ? BoolColumn->RowValues[Row] : false;
-		return value ? ECheckBoxState::Checked : ECheckBoxState::Unchecked;
-	});
+			const bool value = (Row < BoolColumn->RowValues.Num()) ? BoolColumn->RowValues[Row] : false;
+			return value ? ECheckBoxState::Checked : ECheckBoxState::Unchecked;
+		})
+	]
+	+ SHorizontalBox::Slot().FillWidth(1);
 }
 	
 TSharedRef<SWidget> CreateBoolPropertyWidget(bool bReadOnly, UObject* TransactionObject, void* Value, UClass* ContextClass, UClass* ResultBaseClass)
 {
-	return CreatePropertyWidget<FBoolContextProperty>(bReadOnly, TransactionObject, Value, ContextClass, GetDefault<UGraphEditorSettings>()->BooleanPinTypeColor);
+	IHasContextClass* HasContextClass = Cast<IHasContextClass>(TransactionObject);
+
+	FBoolContextProperty* ContextProperty = reinterpret_cast<FBoolContextProperty*>(Value);
+
+	return SNew(SPropertyAccessChainWidget).ContextClassOwner(HasContextClass).AllowFunctions(true).BindingColor("BooleanPinTypeColor").TypeFilter("bool")
+	.PropertyBindingValue(&ContextProperty->Binding)
+	.OnAddBinding_Lambda(
+		[ContextProperty, TransactionObject](FName InPropertyName, const TArray<FBindingChainElement>& InBindingChain)
+		{
+			const FScopedTransaction Transaction(NSLOCTEXT("ContextPropertyWidget", "Change Property Binding", "Change Property Binding"));
+			TransactionObject->Modify(true);
+			ContextProperty->SetBinding(InBindingChain);	
+		});
 }
 	
 void RegisterBoolWidgets()
diff --git a/Engine/Plugins/Experimental/Chooser/Source/ChooserEditor/Private/EnumColumnEditor.cpp b/Engine/Plugins/Experimental/Chooser/Source/ChooserEditor/Private/EnumColumnEditor.cpp
--- a/Engine/Plugins/Experimental/Chooser/Source/ChooserEditor/Private/EnumColumnEditor.cpp
+++ b/Engine/Plugins/Experimental/Chooser/Source/ChooserEditor/Private/EnumColumnEditor.cpp
@@ -2,7 +2,7 @@
 
 #include "EnumColumnEditor.h"
 #include "EnumColumn.h"
-#include "ContextPropertyWidget.h"
+#include "SPropertyAccessChainWidget.h"
 #include "ObjectChooserWidgetFactories.h"
 #include "ChooserTableEditor.h"
 #include "GraphEditorSettings.h"
@@ -153,7 +153,19 @@ TSharedRef<SWidget> CreateEnumColumnWidget(UChooserTable* Chooser, FChooserColum
 
 TSharedRef<SWidget> CreateEnumPropertyWidget(bool bReadOnly, UObject* TransactionObject, void* Value, UClass* ContextClass, UClass* ResultBaseClass)
 {
-	return CreatePropertyWidget<FEnumContextProperty>(bReadOnly, TransactionObject, Value, ContextClass, GetDefault<UGraphEditorSettings>()->BytePinTypeColor);
+		IHasContextClass* HasContextClass = Cast<IHasContextClass>(TransactionObject);
+    
+    	FEnumContextProperty* ContextProperty = reinterpret_cast<FEnumContextProperty*>(Value);
+    
+    	return SNew(SPropertyAccessChainWidget).ContextClassOwner(HasContextClass).AllowFunctions(false).BindingColor("BytePinTypeColor").TypeFilter("enum")
+    	.PropertyBindingValue(&ContextProperty->Binding)
+    	.OnAddBinding_Lambda(
+    		[ContextProperty, TransactionObject](FName InPropertyName, const TArray<FBindingChainElement>& InBindingChain)
+    		{
+    			const FScopedTransaction Transaction(NSLOCTEXT("ContextPropertyWidget", "Change Property Binding", "Change Property Binding"));
+    			TransactionObject->Modify(true);
+    			ContextProperty->SetBinding(InBindingChain);	
+    		});
 }
 	
 void RegisterEnumWidgets()
diff --git a/Engine/Plugins/Experimental/Chooser/Source/ChooserEditor/Private/FloatRangeColumnEditor.cpp b/Engine/Plugins/Experimental/Chooser/Source/ChooserEditor/Private/FloatRangeColumnEditor.cpp
--- a/Engine/Plugins/Experimental/Chooser/Source/ChooserEditor/Private/FloatRangeColumnEditor.cpp
+++ b/Engine/Plugins/Experimental/Chooser/Source/ChooserEditor/Private/FloatRangeColumnEditor.cpp
@@ -2,11 +2,11 @@
 
 #include "FloatRangeColumnEditor.h"
 #include "FloatRangeColumn.h"
-#include "ContextPropertyWidget.h"
 #include "ObjectChooserWidgetFactories.h"
 #include "ChooserTableEditor.h"
 #include "Widgets/Input/SNumericEntryBox.h"
 #include "GraphEditorSettings.h"
+#include "SPropertyAccessChainWidget.h"
 
 #define LOCTEXT_NAMESPACE "FloatRangeColumnEditor"
 
@@ -14,7 +14,19 @@ namespace UE::ChooserEditor
 {
 TSharedRef<SWidget> CreateFloatPropertyWidget(bool bReadOnly, UObject* TransactionObject, void* Value, UClass* ContextClass, UClass* ResultBaseClass)
 {
-	return CreatePropertyWidget<FFloatContextProperty>(bReadOnly, TransactionObject, Value, ContextClass, GetDefault<UGraphEditorSettings>()->FloatPinTypeColor);
+	IHasContextClass* HasContextClass = Cast<IHasContextClass>(TransactionObject);
+
+	FFloatContextProperty* ContextProperty = reinterpret_cast<FFloatContextProperty*>(Value);
+
+	return SNew(SPropertyAccessChainWidget).ContextClassOwner(HasContextClass).AllowFunctions(true).BindingColor("FloatPinTypeColor").TypeFilter("double")
+	.PropertyBindingValue(&ContextProperty->Binding)
+	.OnAddBinding_Lambda(
+		[ContextProperty, TransactionObject](FName InPropertyName, const TArray<FBindingChainElement>& InBindingChain)
+		{
+			const FScopedTransaction Transaction(NSLOCTEXT("ContextPropertyWidget", "Change Property Binding", "Change Property Binding"));
+			TransactionObject->Modify(true);
+			ContextProperty->SetBinding(InBindingChain);	
+		});
 }
 	
 TSharedRef<SWidget> CreateFloatRangeColumnWidget(UChooserTable* Chooser, FChooserColumnBase* Column, int Row)
diff --git a/Engine/Plugins/Experimental/Chooser/Source/ChooserEditor/Private/GameplayTagColumnEditor.cpp b/Engine/Plugins/Experimental/Chooser/Source/ChooserEditor/Private/GameplayTagColumnEditor.cpp
--- a/Engine/Plugins/Experimental/Chooser/Source/ChooserEditor/Private/GameplayTagColumnEditor.cpp
+++ b/Engine/Plugins/Experimental/Chooser/Source/ChooserEditor/Private/GameplayTagColumnEditor.cpp
@@ -1,8 +1,7 @@
 // Copyright Epic Games, Inc. All Rights Reserved.
 
 #include "GameplayTagColumnEditor.h"
-#include "ContextPropertyWidget.h"
-#include "ChooserTableEditor.h"
+#include "SPropertyAccessChainWidget.h"
 #include "GameplayTagColumn.h"
 #include "ObjectChooserWidgetFactories.h"
 #include "SGameplayTagWidget.h"
@@ -51,7 +50,19 @@ TSharedRef<SWidget> CreateGameplayTagColumnWidget(UChooserTable* Chooser, FChoos
 
 TSharedRef<SWidget> CreateGameplayTagPropertyWidget(bool bReadOnly, UObject* TransactionObject, void* Value, UClass* ContextClass, UClass* ResultBaseClass)
 {
-	return CreatePropertyWidget<FGameplayTagContextProperty>(bReadOnly, TransactionObject, Value, ContextClass, GetDefault<UGraphEditorSettings>()->StructPinTypeColor);
+	IHasContextClass* HasContextClass = Cast<IHasContextClass>(TransactionObject);
+
+	FGameplayTagContextProperty* ContextProperty = reinterpret_cast<FGameplayTagContextProperty*>(Value);
+
+	return SNew(SPropertyAccessChainWidget).ContextClassOwner(HasContextClass).AllowFunctions(false).BindingColor("StructPinTypeColor").TypeFilter("FGameplayTagContainer")
+	.PropertyBindingValue(&ContextProperty->Binding)
+	.OnAddBinding_Lambda(
+		[ContextProperty, TransactionObject](FName InPropertyName, const TArray<FBindingChainElement>& InBindingChain)
+		{
+			const FScopedTransaction Transaction(NSLOCTEXT("ContextPropertyWidget", "Change Property Binding", "Change Property Binding"));
+			TransactionObject->Modify(true);
+			ContextProperty->SetBinding(InBindingChain);	
+		});
 }
 	
 void RegisterGameplayTagWidgets()
diff --git a/Engine/Plugins/Experimental/Chooser/Source/ChooserEditor/Private/ObjectColumnEditor.cpp b/Engine/Plugins/Experimental/Chooser/Source/ChooserEditor/Private/ObjectColumnEditor.cpp
--- a/Engine/Plugins/Experimental/Chooser/Source/ChooserEditor/Private/ObjectColumnEditor.cpp
+++ b/Engine/Plugins/Experimental/Chooser/Source/ChooserEditor/Private/ObjectColumnEditor.cpp
@@ -1,7 +1,7 @@
 // Copyright Epic Games, Inc. All Rights Reserved.
 
 #include "ObjectColumnEditor.h"
-#include "ContextPropertyWidget.h"
+#include "SPropertyAccessChainWidget.h"
 #include "GraphEditorSettings.h"
 #include "ObjectChooserWidgetFactories.h"
 #include "ObjectColumn.h"
@@ -143,7 +143,19 @@ namespace UE::ChooserEditor
 
 	static TSharedRef<SWidget> CreateObjectPropertyWidget(bool bReadOnly, UObject* TransactionObject, void* Value, UClass* ContextClass, UClass* ResultBaseClass)
 	{
-		return CreatePropertyWidget<FObjectContextProperty>(bReadOnly, TransactionObject, Value, ContextClass, GetDefault<UGraphEditorSettings>()->ObjectPinTypeColor);
+		IHasContextClass* HasContextClass = Cast<IHasContextClass>(TransactionObject);
+
+		FObjectContextProperty* ContextProperty = reinterpret_cast<FObjectContextProperty*>(Value);
+
+		return SNew(SPropertyAccessChainWidget).ContextClassOwner(HasContextClass).AllowFunctions(false).BindingColor("ObjectPinTypeColor").TypeFilter("object")
+		.PropertyBindingValue(&ContextProperty->Binding)
+		.OnAddBinding_Lambda(
+			[ContextProperty, TransactionObject](FName InPropertyName, const TArray<FBindingChainElement>& InBindingChain)
+			{
+				const FScopedTransaction Transaction(NSLOCTEXT("ContextPropertyWidget", "Change Property Binding", "Change Property Binding"));
+				TransactionObject->Modify(true);
+				ContextProperty->SetBinding(InBindingChain);	
+			});
 	}
 
 	void RegisterObjectWidgets()
diff --git a/Engine/Plugins/Experimental/Chooser/Source/ChooserEditor/Private/PropertyAccessChainCustomization.cpp b/Engine/Plugins/Experimental/Chooser/Source/ChooserEditor/Private/PropertyAccessChainCustomization.cpp
--- a/Engine/Plugins/Experimental/Chooser/Source/ChooserEditor/Private/PropertyAccessChainCustomization.cpp
+++ b/Engine/Plugins/Experimental/Chooser/Source/ChooserEditor/Private/PropertyAccessChainCustomization.cpp
@@ -11,228 +11,13 @@
 #include "IDetailChildrenBuilder.h"
 #include "IPropertyAccessEditor.h"
 #include "ScopedTransaction.h"
-#include "Features/IModularFeatures.h"
+#include "SPropertyAccessChainWidget.h"
 
 #define LOCTEXT_NAMESPACE "PropertyAccessChainCustomization"
 
 namespace UE::ChooserEditor
 {
 	
-// Wrapper widget for Property access widget, which can update when the target Class changes
-class SPropertyAccessChainWidget : public SCompoundWidget
-{
-public:
-	SLATE_BEGIN_ARGS(SPropertyAccessChainWidget)
-	{}
-
-	SLATE_ARGUMENT(IHasContextClass*, ContextClassOwner)
-	SLATE_ARGUMENT(FString, TypeFilter);
-	SLATE_ARGUMENT(FString, BindingColor);
-            
-	SLATE_END_ARGS()
-
-	TSharedRef<SWidget> CreatePropertyAccessWidget()
-	{
-		FPropertyBindingWidgetArgs Args;
-		Args.bAllowPropertyBindings = true;
-
-		UClass* ContextClass = UObject::StaticClass();
-		if (ContextClassOwner)
-		{
-			ContextClass = ContextClassOwner->GetContextClass();
-		}
-		
-		Args.OnCanBindProperty = FOnCanBindProperty::CreateLambda([this](FProperty* Property)
-		{
-			if (TypeFilter == "" || Property == nullptr)
-			{
-				return true;
-			}
-			else if (TypeFilter == "enum")
-			{
-				if (const FEnumProperty* EnumProperty = CastField<const FEnumProperty>(Property))
-				{
-					return true;
-				}
-				else if (const FByteProperty* ByteProperty = CastField<const FByteProperty>(Property))
-				{
-					return ByteProperty->Enum != nullptr;
-				}
-				return false;
-			}
-			
-			return Property->GetCPPType() == TypeFilter;
-		});
-		
-		Args.OnCanBindToClass = FOnCanBindToClass::CreateLambda([](UClass* InClass)
-		{
-			return true;
-		});
-
-
-		FLinearColor BindingColorValue = FLinearColor::Gray;
-		if (BindingColor != "")
-		{
-			const UGraphEditorSettings* GraphEditorSettings = GetDefault<UGraphEditorSettings>();
-			if (const FStructProperty* ColorProperty = FindFProperty<FStructProperty>(GraphEditorSettings->GetClass(), FName(BindingColor)))
-			{
-				BindingColorValue = *ColorProperty->ContainerPtrToValuePtr<FLinearColor>(GraphEditorSettings);
-			}
-		}
-
-		Args.CurrentBindingColor = MakeAttributeLambda([BindingColorValue]() {
-			return BindingColorValue;
-		});
-
-		Args.OnCanAcceptPropertyOrChildren = FOnCanBindProperty::CreateLambda([](FProperty* InProperty)
-			{
-				// Make only blueprint visible properties visible for binding.
-				return InProperty->HasAnyPropertyFlags(CPF_BlueprintVisible);
-			});	
-
-		Args.OnAddBinding = FOnAddBinding::CreateLambda([this](FName InPropertyName, const TArray<FBindingChainElement>& InBindingChain)
-			{
-				TArray<UObject*> OuterObjects;
-				PropertyHandle->GetOuterObjects(OuterObjects);
-				const FScopedTransaction Transaction(NSLOCTEXT("ChooserPropertyBinding", "Change Property Binding", "Change Property Binding"));
-			
-				for (uint32 i=0; i<PropertyHandle->GetNumOuterObjects(); i++)
-				{
-					void* ValuePtr = nullptr;
-					PropertyHandle->GetValueData(ValuePtr);// todo get per object value data?
-					FChooserPropertyBinding* PropertyValue = reinterpret_cast<FChooserPropertyBinding*>(ValuePtr);
-				
-					if (PropertyValue != nullptr)
-					{
-						PropertyHandle->NotifyPreChange();
-						
-						OuterObjects[i]->Modify(true);
-						Chooser::CopyPropertyChain(InBindingChain, PropertyValue->PropertyBindingChain);
-
-						if (TypeFilter == "enum")
-						{
-							FField* Property = InBindingChain.Last().Field.ToField();
-							FChooserEnumPropertyBinding* EnumPropertyValue = static_cast<FChooserEnumPropertyBinding*>(PropertyValue);
-							
-							if (const FEnumProperty* EnumProperty = CastField<const FEnumProperty>(Property))
-							{
-								EnumPropertyValue->Enum = EnumProperty->GetEnum();
-							}
-							else if (const FByteProperty* ByteProperty = CastField<const FByteProperty>(Property))
-							{
-								EnumPropertyValue->Enum = ByteProperty->Enum;
-							}
-						}
-
-						PropertyHandle->NotifyPostChange(EPropertyChangeType::ValueSet);
-					}
-				}
-			});
-
-		Args.CurrentBindingToolTipText = MakeAttributeLambda([this]()
-				{
-					const FText Bind = NSLOCTEXT("ContextPropertyWidget", "Bind", "Bind");
-					FText CurrentValue = Bind;
-					
-					void* ValuePtr = nullptr;
-					PropertyHandle->GetValueData(ValuePtr);// todo get per object value data?
-					FChooserPropertyBinding* PropertyValue = reinterpret_cast<FChooserPropertyBinding*>(ValuePtr);
-
-					if (PropertyValue != nullptr && PropertyValue->PropertyBindingChain.Num()>0)
-					{
-						TArray<FText> BindingChainText;
-						BindingChainText.Reserve(PropertyValue->PropertyBindingChain.Num());
-					 
-						for (const FName& Name : PropertyValue->PropertyBindingChain)
-						{
-							BindingChainText.Add(FText::FromName(Name));
-						}
-						
-						CurrentValue = FText::Join(NSLOCTEXT("ContextPropertyWidget", "PropertyPathSeparator","."), BindingChainText);
-					}
-		
-					return CurrentValue;	
-				});
-		
-		Args.CurrentBindingText = MakeAttributeLambda([this]()
-				{
-					const FText Bind = NSLOCTEXT("ContextPropertyWidget", "Bind", "Bind");
-					FText CurrentValue = Bind;
-			
-					void* ValuePtr = nullptr;
-					PropertyHandle->GetValueData(ValuePtr);// todo get per object value data?
-					FChooserPropertyBinding* PropertyValue = reinterpret_cast<FChooserPropertyBinding*>(ValuePtr);
-		
-					int BindingChainLength = PropertyValue->PropertyBindingChain.Num();
-					if (BindingChainLength > 0)
-					{
-						if (BindingChainLength == 1)
-						{
-							// single property, just use the property name
-							CurrentValue = FText::FromName(PropertyValue->PropertyBindingChain.Last());
-						}
-						else
-						{
-							// for longer chains always show the last struct/object name, and the final property name (full path in tooltip)
-							CurrentValue = FText::Join(NSLOCTEXT("ContextPropertyWidget", "PropertyPathSeparator","."),
-								TArray<FText>({
-									FText::FromName(PropertyValue->PropertyBindingChain[BindingChainLength-2]),
-									FText::FromName(PropertyValue->PropertyBindingChain[BindingChainLength-1])
-								}));
-						}
-					}
-		
-					return CurrentValue;
-				});
-
-		Args.CurrentBindingImage = MakeAttributeLambda([]() -> const FSlateBrush*
-			{
-				static FName PropertyIcon(TEXT("Kismet.Tabs.Variables"));
-				return FAppStyle::GetBrush(PropertyIcon);
-			});
-
-		const IPropertyAccessEditor& PropertyAccessEditor = IModularFeatures::Get().GetModularFeature<IPropertyAccessEditor>("PropertyAccessEditor");
-		FBindingContextStruct StructInfo;
-		StructInfo.Struct = ContextClass;
-		return  PropertyAccessEditor.MakePropertyBindingWidget({StructInfo}, Args);
-	}
-
-	void UpdateWidget()
-	{
-		ChildSlot[ CreatePropertyAccessWidget() ];
-	}
-
-	void ContextClassChanged(UClass* NewContextClass)
-	{
-		UpdateWidget();
-	}
-
-	void Construct( const FArguments& InArgs, TSharedRef<IPropertyHandle> InPropertyHandle)
-	{
-		PropertyHandle = InPropertyHandle;
-		TypeFilter = InArgs._TypeFilter;
-		BindingColor = InArgs._BindingColor;
-		ContextClassOwner = InArgs._ContextClassOwner;
-		UpdateWidget();
-
-		if (ContextClassOwner)
-		{
-			ContextClassOwner->OnContextClassChanged.AddSP(this, &SPropertyAccessChainWidget::ContextClassChanged);
-		}
-	}
-
-	~SPropertyAccessChainWidget()
-	{
-	}
-
-private:
-	FString TypeFilter;
-	FString BindingColor;
-	FDelegateHandle ContextClassChangedHandle();
-	TSharedPtr<IPropertyHandle> PropertyHandle;
-	IHasContextClass* ContextClassOwner = nullptr;
-};
-
 void FPropertyAccessChainCustomization::CustomizeHeader(TSharedRef<IPropertyHandle> PropertyHandle, FDetailWidgetRow& HeaderRow, IPropertyTypeCustomizationUtils& CustomizationUtils)
 {
 	FProperty* Property = PropertyHandle->GetProperty();
@@ -240,9 +25,13 @@ void FPropertyAccessChainCustomization::CustomizeHeader(TSharedRef<IPropertyHand
 	UClass* ContextClass = nullptr;
 	static FName TypeMetaData = "BindingType";
 	static FName ColorMetaData = "BindingColor";
+	static FName AllowFunctionsMetaData = "BindingAllowFunctions";
 
 	FString TypeFilter = PropertyHandle->GetMetaData(TypeMetaData);
 	FString BindingColor = PropertyHandle->GetMetaData(ColorMetaData);
+	FString BindingAllowFunctions = PropertyHandle->GetMetaData(AllowFunctionsMetaData);
+	bool AllowFunctions = BindingAllowFunctions.ToBool();
+		
 
 	TArray<UObject*> OuterObjects;
 	PropertyHandle->GetOuterObjects(OuterObjects);
@@ -266,7 +55,51 @@ void FPropertyAccessChainCustomization::CustomizeHeader(TSharedRef<IPropertyHand
 	]
 	.ValueContent()
 	[
-		SNew(SPropertyAccessChainWidget, PropertyHandle).ContextClassOwner(HasContext).TypeFilter(TypeFilter).BindingColor(BindingColor)
+		SNew(SPropertyAccessChainWidget).ContextClassOwner(HasContext).TypeFilter(TypeFilter).BindingColor(BindingColor).AllowFunctions(AllowFunctions)
+		.PropertyBindingValue_Lambda([PropertyHandle]()
+		{
+			void* data;
+			PropertyHandle->GetValueData(data);
+			return reinterpret_cast<const FChooserPropertyBinding*>(data);
+		})
+		.OnAddBinding_Lambda([PropertyHandle, TypeFilter](FName InPropertyName, const TArray<FBindingChainElement>& InBindingChain)
+         	 	{
+         	 		TArray<UObject*> OuterObjects;
+         	 		PropertyHandle->GetOuterObjects(OuterObjects);
+         	 		const FScopedTransaction Transaction(NSLOCTEXT("ChooserPropertyBinding", "Change Property Binding", "Change Property Binding"));
+         	 	
+         			for (uint32 i=0; i<PropertyHandle->GetNumOuterObjects(); i++)
+         	 		{
+         	 			void* ValuePtr = nullptr;
+         	 			PropertyHandle->GetValueData(ValuePtr);// todo get per object value data?
+         	 			FChooserPropertyBinding* PropertyValue = reinterpret_cast<FChooserPropertyBinding*>(ValuePtr);
+         	 		
+         	 			if (PropertyValue != nullptr)
+         	 			{
+         	 				PropertyHandle->NotifyPreChange();
+         	 				
+         	 				OuterObjects[i]->Modify(true);
+         	 				Chooser::CopyPropertyChain(InBindingChain, PropertyValue->PropertyBindingChain);
+         	
+         	 				if (TypeFilter == "enum")
+         	 				{
+         	 					FField* Property = InBindingChain.Last().Field.ToField();
+         	 					FChooserEnumPropertyBinding* EnumPropertyValue = static_cast<FChooserEnumPropertyBinding*>(PropertyValue);
+         	 					
+         	 					if (const FEnumProperty* EnumProperty = CastField<const FEnumProperty>(Property))
+         	 					{
+         	 						EnumPropertyValue->Enum = EnumProperty->GetEnum();
+         	 					}
+         	 					else if (const FByteProperty* ByteProperty = CastField<const FByteProperty>(Property))
+         	 					{
+         	 						EnumPropertyValue->Enum = ByteProperty->Enum;
+         	 					}
+         	 				}
+         	
+         	 				PropertyHandle->NotifyPostChange(EPropertyChangeType::ValueSet);
+         	 			}
+         	 		}
+         	 	})
 	];
 }
 	
diff --git a/Engine/Plugins/Experimental/Chooser/Source/ChooserEditor/Private/SPropertyAccessChainWidget.cpp b/Engine/Plugins/Experimental/Chooser/Source/ChooserEditor/Private/SPropertyAccessChainWidget.cpp
new file mode 100644
--- /dev/null
+++ b/Engine/Plugins/Experimental/Chooser/Source/ChooserEditor/Private/SPropertyAccessChainWidget.cpp
@@ -0,0 +1,218 @@
+﻿// Copyright Epic Games, Inc. All Rights Reserved.
+
+#include "SPropertyAccessChainWidget.h"
+
+#include "Chooser.h"
+#include "ChooserPropertyAccess.h"
+#include "PropertyHandle.h"
+#include "DetailWidgetRow.h"
+#include "GraphEditorSettings.h"
+#include "SClassViewer.h"
+#include "IDetailChildrenBuilder.h"
+#include "IPropertyAccessEditor.h"
+#include "ScopedTransaction.h"
+#include "Features/IModularFeatures.h"
+
+#define LOCTEXT_NAMESPACE "PropertyAccessChainWidget"
+
+namespace UE::ChooserEditor
+{
+	
+TSharedRef<SWidget> SPropertyAccessChainWidget::CreatePropertyAccessWidget()
+{
+	FPropertyBindingWidgetArgs Args;
+	Args.bAllowPropertyBindings = true;
+
+	UClass* ContextClass = UObject::StaticClass();
+	if (ContextClassOwner)
+	{
+		ContextClass = ContextClassOwner->GetContextClass();
+	}
+
+	Args.bAllowUObjectFunctions = true;
+	Args.bAllowOnlyThreadSafeFunctions = true;
+
+	auto CanBindProperty = [this](FProperty* Property)
+	{
+		if (TypeFilter == "" || Property == nullptr)
+		{
+			return true;
+		}
+		if (TypeFilter == "object")
+		{
+			// special case for objects references of any type
+			return CastField<FObjectProperty>(Property) != nullptr;
+		}
+		if (TypeFilter == "double")
+		{
+			// special case for doubles to bind to either floats or doubles
+			return Property->GetCPPType() == "float" || Property->GetCPPType() == "double";
+		}
+		else if (TypeFilter == "enum")
+		{
+			// special case for enums, to find properties of type EnumProperty or ByteProperty which have an Enum
+	
+			if (const FEnumProperty* EnumProperty = CastField<const FEnumProperty>(Property))
+			{
+				return true;
+			}
+			else if (const FByteProperty* ByteProperty = CastField<const FByteProperty>(Property))
+			{
+				return ByteProperty->Enum != nullptr;
+			}
+			return false;
+		}
+		
+		return Property->GetCPPType() == TypeFilter;
+	};
+
+	Args.OnCanBindProperty = FOnCanBindProperty::CreateLambda(CanBindProperty);
+
+	Args.OnCanBindFunction = FOnCanBindFunction::CreateLambda([CanBindProperty](UFunction* Function)
+	{
+		if (Function->NumParms !=1)
+		{
+			// only allow binding object member functions which have no parameters
+			return false;
+		}
+
+		if (FProperty* ReturnProperty = Function->GetReturnProperty())
+		{
+			return CanBindProperty(ReturnProperty);
+		}
+	
+		return false;
+	});
+
+	Args.OnCanBindToClass = FOnCanBindToClass::CreateLambda([](UClass* InClass)
+	{
+		return true;
+	});
+
+	FLinearColor BindingColorValue = FLinearColor::Gray;
+	if (BindingColor != "")
+	{
+		const UGraphEditorSettings* GraphEditorSettings = GetDefault<UGraphEditorSettings>();
+		if (const FStructProperty* ColorProperty = FindFProperty<FStructProperty>(GraphEditorSettings->GetClass(), FName(BindingColor)))
+		{
+			BindingColorValue = *ColorProperty->ContainerPtrToValuePtr<FLinearColor>(GraphEditorSettings);
+		}
+	}
+
+	Args.CurrentBindingColor = MakeAttributeLambda([BindingColorValue]() {
+		return BindingColorValue;
+	});
+	
+	Args.OnCanBindToSubObjectClass = FOnCanBindToSubObjectClass::CreateLambda([](UClass* InClass)
+		{
+			// CanBindToSubObjectClass does the opposite of what it's name says.  True means don't allow bindings
+			// don't allow binding to any object propertoes (forcing use of thread safe functions to access objects)
+			return true;
+		});
+
+	Args.OnCanAcceptPropertyOrChildren = FOnCanBindProperty::CreateLambda([](FProperty* InProperty)
+		{
+			// Make only blueprint visible properties visible for binding.
+			return InProperty->HasAnyPropertyFlags(CPF_BlueprintVisible);
+		});	
+
+	Args.OnAddBinding = OnAddBinding;
+
+
+	Args.CurrentBindingToolTipText = MakeAttributeLambda([this]()
+			{
+				const FText Bind = NSLOCTEXT("ContextPropertyWidget", "Bind", "Bind");
+				FText CurrentValue = Bind;
+				
+				const FChooserPropertyBinding* PropertyValue = PropertyBindingValue.Get();
+
+				if (PropertyValue != nullptr && PropertyValue->PropertyBindingChain.Num()>0)
+				{
+					TArray<FText> BindingChainText;
+					BindingChainText.Reserve(PropertyValue->PropertyBindingChain.Num());
+				 
+					for (const FName& Name : PropertyValue->PropertyBindingChain)
+					{
+						BindingChainText.Add(FText::FromName(Name));
+					}
+					
+					CurrentValue = FText::Join(NSLOCTEXT("ContextPropertyWidget", "PropertyPathSeparator","."), BindingChainText);
+				}
+	
+				return CurrentValue;	
+			});
+	
+	Args.CurrentBindingText = MakeAttributeLambda([this]()
+			{
+				const FText Bind = NSLOCTEXT("ContextPropertyWidget", "Bind", "Bind");
+				FText CurrentValue = Bind;
+		
+				const FChooserPropertyBinding* PropertyValue = PropertyBindingValue.Get();
+	
+				int BindingChainLength = PropertyValue->PropertyBindingChain.Num();
+				if (BindingChainLength > 0)
+				{
+					if (BindingChainLength == 1)
+					{
+						// single property, just use the property name
+						CurrentValue = FText::FromName(PropertyValue->PropertyBindingChain.Last());
+					}
+					else
+					{
+						// for longer chains always show the last struct/object name, and the final property name (full path in tooltip)
+						CurrentValue = FText::Join(NSLOCTEXT("ContextPropertyWidget", "PropertyPathSeparator","."),
+							TArray<FText>({
+								FText::FromName(PropertyValue->PropertyBindingChain[BindingChainLength-2]),
+								FText::FromName(PropertyValue->PropertyBindingChain[BindingChainLength-1])
+							}));
+					}
+				}
+	
+				return CurrentValue;
+			});
+
+	Args.CurrentBindingImage = MakeAttributeLambda([]() -> const FSlateBrush*
+		{
+			static FName PropertyIcon(TEXT("Kismet.Tabs.Variables"));
+			return FAppStyle::GetBrush(PropertyIcon);
+		});
+
+	const IPropertyAccessEditor& PropertyAccessEditor = IModularFeatures::Get().GetModularFeature<IPropertyAccessEditor>("PropertyAccessEditor");
+	FBindingContextStruct StructInfo;
+	StructInfo.Struct = ContextClass;
+	return  PropertyAccessEditor.MakePropertyBindingWidget({StructInfo}, Args);
+}
+
+void SPropertyAccessChainWidget::UpdateWidget()
+{
+	ChildSlot[ CreatePropertyAccessWidget() ];
+}
+
+void SPropertyAccessChainWidget::ContextClassChanged(UClass* NewContextClass)
+{
+	UpdateWidget();
+}
+
+void SPropertyAccessChainWidget::Construct( const FArguments& InArgs)
+{
+	TypeFilter = InArgs._TypeFilter;
+	BindingColor = InArgs._BindingColor;
+	ContextClassOwner = InArgs._ContextClassOwner;
+	bAllowFunctions = InArgs._AllowFunctions;
+	OnAddBinding = InArgs._OnAddBinding;
+	PropertyBindingValue = InArgs._PropertyBindingValue;
+	UpdateWidget();
+
+	if (ContextClassOwner)
+	{
+		ContextClassOwner->OnContextClassChanged.AddSP(this, &SPropertyAccessChainWidget::ContextClassChanged);
+	}
+}
+
+SPropertyAccessChainWidget::~SPropertyAccessChainWidget()
+{
+}
+
+}
+
+#undef LOCTEXT_NAMESPACE
\ No newline at end of file
diff --git a/Engine/Plugins/Experimental/Chooser/Source/ChooserEditor/Public/ContextPropertyWidget.h b/Engine/Plugins/Experimental/Chooser/Source/ChooserEditor/Public/ContextPropertyWidget.h
deleted file mode 100644
diff --git a/Engine/Plugins/Experimental/Chooser/Source/ChooserEditor/Public/SPropertyAccessChainWidget.h b/Engine/Plugins/Experimental/Chooser/Source/ChooserEditor/Public/SPropertyAccessChainWidget.h
new file mode 100644
--- /dev/null
+++ b/Engine/Plugins/Experimental/Chooser/Source/ChooserEditor/Public/SPropertyAccessChainWidget.h
@@ -0,0 +1,49 @@
+// Copyright Epic Games, Inc. All Rights Reserved.
+#pragma once
+
+#include "ScopedTransaction.h"
+#include "IPropertyAccessEditor.h"
+#include "Features/IModularFeatures.h"
+#include "Kismet2/BlueprintEditorUtils.h"
+#include "Styling/AppStyle.h"
+#include "Widgets/SCompoundWidget.h"
+#include "ChooserPropertyAccess.h"
+
+namespace UE::ChooserEditor
+{
+		
+// Wrapper widget for Property access widget, which can update when the target Class changes
+class SPropertyAccessChainWidget : public SCompoundWidget
+{
+public:
+	SLATE_BEGIN_ARGS(SPropertyAccessChainWidget)
+	{}
+
+	SLATE_ARGUMENT(IHasContextClass*, ContextClassOwner)
+	SLATE_ARGUMENT(FString, TypeFilter);
+	SLATE_ARGUMENT(FString, BindingColor);
+	SLATE_ARGUMENT(bool, AllowFunctions);
+	SLATE_EVENT(FOnAddBinding, OnAddBinding);
+	SLATE_ATTRIBUTE(const FChooserPropertyBinding*, PropertyBindingValue);
+            
+	SLATE_END_ARGS()
+
+	void Construct( const FArguments& InArgs);
+	virtual ~SPropertyAccessChainWidget();
+
+private:
+	TSharedRef<SWidget> CreatePropertyAccessWidget();
+	void UpdateWidget();
+	void ContextClassChanged(UClass* NewContextClass);
+		
+	FString TypeFilter;
+	FString BindingColor;
+	FDelegateHandle ContextClassChangedHandle();
+	IHasContextClass* ContextClassOwner = nullptr;
+	FOnAddBinding OnAddBinding;
+	TAttribute<const FChooserPropertyBinding*> PropertyBindingValue;
+	bool bAllowFunctions;
+};
+	
+
+}
\ No newline at end of file
diff --git a/Engine/Plugins/Runtime/PropertyAccess/Source/PropertyAccessEditor/Private/SPropertyBinding.cpp b/Engine/Plugins/Runtime/PropertyAccess/Source/PropertyAccessEditor/Private/SPropertyBinding.cpp
--- a/Engine/Plugins/Runtime/PropertyAccess/Source/PropertyAccessEditor/Private/SPropertyBinding.cpp
+++ b/Engine/Plugins/Runtime/PropertyAccess/Source/PropertyAccessEditor/Private/SPropertyBinding.cpp
@@ -154,7 +154,13 @@ void SPropertyBinding::ForEachBindableFunction(UClass* FromClass, Predicate Pred
 					}
 				}
 
-				if(bValidObjectFunction || bValidStructFunction || Args.OnCanBindFunction.Execute(Function))
+				bool bValidThreadSafety = true;
+				if (Args.bAllowOnlyThreadSafeFunctions)
+				{
+					bValidThreadSafety = FBlueprintEditorUtils::HasFunctionBlueprintThreadSafeMetaData(Function);
+				}
+
+				if(bValidThreadSafety && (bValidObjectFunction || bValidStructFunction || Args.OnCanBindFunction.Execute(Function)))
 				{
 					Pred(FFunctionInfo(Function));
 				}
diff --git a/Engine/Source/Editor/UnrealEd/Public/IPropertyAccessEditor.h b/Engine/Source/Editor/UnrealEd/Public/IPropertyAccessEditor.h
--- a/Engine/Source/Editor/UnrealEd/Public/IPropertyAccessEditor.h
+++ b/Engine/Source/Editor/UnrealEd/Public/IPropertyAccessEditor.h
@@ -223,6 +223,9 @@ struct FPropertyBindingWidgetArgs
 
 	/** Whether to allow UObject functions as non-leaf nodes */
 	bool bAllowUObjectFunctions = false;
+	
+	/** Whether to allow only functions marked thread safe */
+	bool bAllowOnlyThreadSafeFunctions = false;
 
 	/** Whether to allow UScriptStruct functions as non-leaf nodes */
 	bool bAllowStructFunctions = false;	
