Randomization column for Chooser Tables

diff --git a/Engine/Plugins/Experimental/Chooser/Source/Chooser/Internal/RandomizeColumn.h b/Engine/Plugins/Experimental/Chooser/Source/Chooser/Internal/RandomizeColumn.h
new file mode 100644
--- /dev/null
+++ b/Engine/Plugins/Experimental/Chooser/Source/Chooser/Internal/RandomizeColumn.h
@@ -0,0 +1,85 @@
+// Copyright Epic Games, Inc. All Rights Reserved.
+#pragma once
+
+#include "CoreMinimal.h"
+#include "IChooserColumn.h"
+#include "IChooserParameterRandomize.h"
+#include "ChooserPropertyAccess.h"
+#include "InstancedStruct.h"
+#include "RandomizeColumn.generated.h"
+
+USTRUCT(DisplayName = "Randomize Property Binding")
+struct CHOOSER_API FRandomizeContextProperty :  public FChooserParameterRandomizeBase
+{
+	GENERATED_BODY()
+public:
+
+	UPROPERTY(EditAnywhere, Meta = (BindingType = "FChooserRandomizationContext", BindingAllowFunctions = "false", BindingColor = "StructPinTypeColor"), Category = "Binding")
+	FChooserPropertyBinding Binding;
+	
+	virtual bool GetValue(const UObject* ContextObject, const FChooserRandomizationContext*& OutResult) const override;
+
+#if WITH_EDITOR
+	static bool CanBind(const FProperty& Property)
+	{
+		static FString RandomizeTypeName = "bool";
+		return Property.GetCPPType() == RandomizeTypeName;
+	}
+
+	void SetBinding(const TArray<FBindingChainElement>& InBindingChain)
+	{
+		UE::Chooser::CopyPropertyChain(InBindingChain, Binding.PropertyBindingChain);
+	}
+
+	virtual void GetDisplayName(FText& OutName) const override
+	{
+		if (!Binding.PropertyBindingChain.IsEmpty())
+		{
+			OutName = FText::FromName(Binding.PropertyBindingChain.Last());
+		}
+	}
+#endif
+};
+
+USTRUCT()
+struct CHOOSER_API FRandomizeColumn : public FChooserColumnBase
+{
+	GENERATED_BODY()
+	public:
+	FRandomizeColumn();
+	
+	UPROPERTY(EditAnywhere, Meta = (ExcludeBaseStruct, BaseStruct = "/Script/Chooser.RandomizeContextProperty"), Category = "Data")
+	FInstancedStruct InputValue;
+	
+	
+	UPROPERTY(EditAnywhere, Category= "Data", meta=(Tooltip="Multiplies the weight of the previous chosen result (set to 0 to never pick the same result twice in a row)"));
+	float RepeatProbabilityMultiplier = 1.0f;
+
+#if WITH_EDITORONLY_DATA
+	UPROPERTY(EditAnywhere, Category= "Data", DisplayName="DefaultRowValue");
+	float DefaultRowValue = 1.0f;
+#endif
+	
+	UPROPERTY(EditAnywhere, Category= "Data", DisplayName="RowValues");
+	TArray<float> RowValues; 
+	
+	virtual void Filter(const UObject* ContextObject, const TArray<uint32>& IndexListIn, TArray<uint32>& IndexListOut) const override;
+	virtual void SetOutputs(UObject* ContextObject, int RowIndex) const override;
+	
+	virtual bool HasFilters() const override { return true; }
+	virtual bool HasOutputs() const override { return true; }
+
+	virtual void PostLoad() override
+	{
+		if (InputValue.IsValid())
+		{
+			InputValue.GetMutable<FChooserParameterBase>().PostLoad();
+		}
+	}
+
+	CHOOSER_COLUMN_BOILERPLATE(FChooserParameterRandomizeBase);
+
+#if WITH_EDITOR
+	virtual bool IsRandomizeColumn() const override { return true; }
+#endif
+};
\ No newline at end of file
diff --git a/Engine/Plugins/Experimental/Chooser/Source/Chooser/Private/RandomizeColumn.cpp b/Engine/Plugins/Experimental/Chooser/Source/Chooser/Private/RandomizeColumn.cpp
new file mode 100644
--- /dev/null
+++ b/Engine/Plugins/Experimental/Chooser/Source/Chooser/Private/RandomizeColumn.cpp
@@ -0,0 +1,115 @@
+// Copyright Epic Games, Inc. All Rights Reserved.
+#include "RandomizeColumn.h"
+#include "ChooserPropertyAccess.h"
+
+bool FRandomizeContextProperty::GetValue(const UObject* ContextObject, const FChooserRandomizationContext*& OutResult) const
+{
+	
+	UStruct* StructType = ContextObject->GetClass();
+	const void* Container = ContextObject;
+	
+	if (UE::Chooser::ResolvePropertyChain(Container, StructType, Binding.PropertyBindingChain))
+	{
+		if (FStructProperty* Property = FindFProperty<FStructProperty>(StructType, Binding.PropertyBindingChain.Last()))
+		{
+			OutResult = Property->ContainerPtrToValuePtr<FChooserRandomizationContext>(Container);
+			return true;
+		}
+	}
+
+	return false;
+}
+
+FRandomizeColumn::FRandomizeColumn()
+{
+	InputValue.InitializeAs(FRandomizeContextProperty::StaticStruct());
+}
+
+void FRandomizeColumn::Filter(const UObject* ContextObject, const TArray<uint32>& IndexListIn, TArray<uint32>& IndexListOut) const
+{
+	int Count = IndexListIn.Num();
+	int Selection = 0;
+
+	const FChooserRandomizationContext* RandomizationContext = nullptr;
+	if (ContextObject && InputValue.IsValid())
+	{
+		InputValue.Get<FChooserParameterRandomizeBase>().GetValue(ContextObject,RandomizationContext);
+	}
+	
+	if (Count > 1)
+	{
+		int LastSelectedIndex = -1;
+		if (RandomizationContext)
+		{
+			if (const FRandomizationState* State = RandomizationContext->StateMap.Find(this))
+			{
+				LastSelectedIndex = State->LastSelectedRow;
+			}
+		}
+		
+		// compute the sum of all weights/probabilities
+		float TotalWeight = 0;
+		for (uint32 Index : IndexListIn)
+		{
+			float RowWeight = 1.0f;
+			if (RowValues.Num() > static_cast<int32>(Index))
+			{
+				RowWeight = RowValues[Index];
+			}
+			
+			if (Index == LastSelectedIndex)
+			{
+				RowWeight *= RepeatProbabilityMultiplier;
+			}
+			
+			TotalWeight += RowWeight;
+		}
+
+		// pick a random float from 0-total weight
+		float RandomNumber = FMath::FRandRange(0.0f, TotalWeight);
+		float Weight = 0;
+
+		// add up the weights again, and select the index where our sum clears the random float
+		for (; Selection < Count - 1; Selection++)
+		{
+			int Index = static_cast<int32>(IndexListIn[Selection]);
+			float RowWeight = 1.0f;
+			
+			if (RowValues.Num() > Index)
+			{
+				RowWeight = RowValues[Index];
+			}
+			
+			if (Index == LastSelectedIndex)
+			{
+				RowWeight *= RepeatProbabilityMultiplier;
+			}
+			
+			Weight += RowWeight;
+			if (Weight > RandomNumber)
+			{
+				break;
+			}
+		}
+	}
+
+	if (Selection < Count)
+	{
+		IndexListOut.Add(IndexListIn[Selection]);
+	}
+}
+
+void FRandomizeColumn::SetOutputs(UObject* ContextObject, int RowIndex) const
+{
+	if (ContextObject && InputValue.IsValid())
+	{
+		const FChooserRandomizationContext* ConstRandomizationContext = nullptr;
+		InputValue.Get<FChooserParameterRandomizeBase>().GetValue(ContextObject,ConstRandomizationContext);
+		if(ConstRandomizationContext)
+		{
+			FChooserRandomizationContext* RandomizationContext = const_cast<FChooserRandomizationContext*>(ConstRandomizationContext);
+			FRandomizationState& State = RandomizationContext->StateMap.FindOrAdd(this, FRandomizationState());
+			State.LastSelectedRow = RowIndex;
+		}
+	}
+}
diff --git a/Engine/Plugins/Experimental/Chooser/Source/Chooser/Public/IChooserColumn.h b/Engine/Plugins/Experimental/Chooser/Source/Chooser/Public/IChooserColumn.h
--- a/Engine/Plugins/Experimental/Chooser/Source/Chooser/Public/IChooserColumn.h
+++ b/Engine/Plugins/Experimental/Chooser/Source/Chooser/Public/IChooserColumn.h
@@ -47,6 +47,10 @@ public:
 	virtual const UScriptStruct* GetInputType() const { return nullptr; };
 	virtual void SetInputType(const UScriptStruct* Type) { };
 	virtual FChooserParameterBase* GetInputValue() { return nullptr; };
+
+	// random columns must go last, and get a special icon
+	// using a virtual fucntion to identify them (rather than hard coding a specific type) to potentially support multiple varieties of randomization column.
+	virtual bool IsRandomizeColumn() const { return false; }
 #endif
 };
 
diff --git a/Engine/Plugins/Experimental/Chooser/Source/Chooser/Public/IChooserParameterRandomize.h b/Engine/Plugins/Experimental/Chooser/Source/Chooser/Public/IChooserParameterRandomize.h
new file mode 100644
--- /dev/null
+++ b/Engine/Plugins/Experimental/Chooser/Source/Chooser/Public/IChooserParameterRandomize.h
@@ -0,0 +1,30 @@
+// Copyright Epic Games, Inc. All Rights Reserved.
+
+#pragma once
+
+#include "CoreMinimal.h"
+#include "UObject/Object.h"
+#include "UObject/Interface.h"
+#include "InstancedStruct.h"
+#include "IChooserParameterBase.h"
+#include "IChooserParameterRandomize.generated.h"
+
+struct FRandomizationState
+{
+	int LastSelectedRow;
+};
+
+USTRUCT(BlueprintType)
+struct FChooserRandomizationContext
+{
+	GENERATED_BODY();
+
+	TMap<const void*, FRandomizationState> StateMap;
+};
+
+USTRUCT()
+struct FChooserParameterRandomizeBase : public FChooserParameterBase
+{
+	GENERATED_BODY()
+	virtual bool GetValue(const UObject* ContextObject, const FChooserRandomizationContext*& OutResult) const { return false; }
+};
diff --git a/Engine/Plugins/Experimental/Chooser/Source/ChooserEditor/Private/ChooserEditorModule.cpp b/Engine/Plugins/Experimental/Chooser/Source/ChooserEditor/Private/ChooserEditorModule.cpp
--- a/Engine/Plugins/Experimental/Chooser/Source/ChooserEditor/Private/ChooserEditorModule.cpp
+++ b/Engine/Plugins/Experimental/Chooser/Source/ChooserEditor/Private/ChooserEditorModule.cpp
@@ -11,6 +11,7 @@
 #include "ChooserTableEditorCommands.h"
 #include "ChooserPropertyAccess.h"
 #include "PropertyEditorModule.h"
+#include "RandomizeColumnEditor.h"
 
 #define LOCTEXT_NAMESPACE "ChooserEditorModule"
 
@@ -25,6 +26,7 @@ void FModule::StartupModule()
 	RegisterBoolWidgets();
 	RegisterEnumWidgets();
 	RegisterObjectWidgets();
+	RegisterRandomizeWidgets();
 	
 	FChooserTableEditorCommands::Register();
 
diff --git a/Engine/Plugins/Experimental/Chooser/Source/ChooserEditor/Private/ChooserTableEditor.cpp b/Engine/Plugins/Experimental/Chooser/Source/ChooserEditor/Private/ChooserTableEditor.cpp
--- a/Engine/Plugins/Experimental/Chooser/Source/ChooserEditor/Private/ChooserTableEditor.cpp
+++ b/Engine/Plugins/Experimental/Chooser/Source/ChooserEditor/Private/ChooserTableEditor.cpp
@@ -56,13 +56,13 @@ void FChooserTableEditor::RegisterTabSpawners(const TSharedRef<class FTabManager
 		.SetDisplayName( LOCTEXT("PropertiesTab", "Details") )
 		.SetGroup(WorkspaceMenuCategory.ToSharedRef())
 		.SetIcon(FSlateIcon("EditorStyle", "LevelEditor.Tabs.Details"));
-	
+		
 	InTabManager->RegisterTabSpawner( TableTabId, FOnSpawnTab::CreateSP(this, &FChooserTableEditor::SpawnTableTab) )
 		.SetDisplayName( LOCTEXT("TableTab", "Chooser Table") )
 		.SetGroup(WorkspaceMenuCategory.ToSharedRef())
 		.SetIcon(FSlateIcon("EditorStyle", "LevelEditor.Tabs.Details"));
 }
-
+	
 void FChooserTableEditor::UnregisterTabSpawners(const TSharedRef<class FTabManager>& InTabManager)
 {
 	FAssetEditorToolkit::UnregisterTabSpawners(InTabManager);
@@ -76,10 +76,10 @@ const FName FChooserTableEditor::ChooserEditorAppIdentifier( TEXT( "ChooserEdito
 FChooserTableEditor::~FChooserTableEditor()
 {
 	FCoreUObjectDelegates::OnObjectsReplaced.RemoveAll(this);
-	
+
 	DetailsView.Reset();
 }
-	
+
 void FChooserTableEditor::RegisterToolbar()
 {
 	UToolMenus* ToolMenus = UToolMenus::Get();
@@ -175,14 +175,14 @@ void FChooserTableEditor::PostUndo(bool bSuccess)
 	UpdateTableColumns();
 	UpdateTableRows();
 }
-	
+
 void FChooserTableEditor::PostRedo(bool bSuccess)
 {
 	UpdateTableColumns();
 	UpdateTableRows();
 }
 
-	
+
 void FChooserTableEditor::NotifyPreChange(FProperty* PropertyAboutToChange)
 {
 }
@@ -195,7 +195,7 @@ void FChooserTableEditor::NotifyPostChange(const FPropertyChangedEvent& Property
 	UpdateTableRows();
 }
 
-	
+
 FText FChooserTableEditor::GetToolkitName() const
 {
 	const TArray<UObject*>& EditingObjs = GetEditingObjects();
@@ -345,9 +345,9 @@ public:
 		Operation->RowIndex = InRowIndex;
 		Operation->DefaultHoverText = LOCTEXT("Chooser Row", "Chooser Row");
 		Operation->CurrentHoverText = Operation->DefaultHoverText;
-			
+		
 		Operation->Construct();
-	
+
 		return Operation;
 	};
 };
@@ -398,10 +398,10 @@ class SChooserTableRow : public SMultiColumnTableRow<TSharedPtr<FChooserTableEdi
 {
 public:
 	SLATE_BEGIN_ARGS(SChooserTableRow) {}
-		/** The list item for this row */
-		SLATE_ARGUMENT(TSharedPtr<FChooserTableEditor::FChooserTableRow>, Entry)
-		SLATE_ARGUMENT(UChooserTable*, Chooser)
-		SLATE_ARGUMENT(FChooserTableEditor*, Editor)
+	/** The list item for this row */
+	SLATE_ARGUMENT(TSharedPtr<FChooserTableEditor::FChooserTableRow>, Entry)
+	SLATE_ARGUMENT(UChooserTable*, Chooser)
+	SLATE_ARGUMENT(FChooserTableEditor*, Editor)
 	SLATE_END_ARGS()
 
 
@@ -422,13 +422,13 @@ public:
 	{
 		static FName Result = "Result";
 		static FName Handles = "Handles";
-		
+	
 		if (RowIndex->RowIndex < Chooser->ResultsStructs.Num())
 		{
 			if (ColumnName == Handles && RowIndex->RowIndex < Chooser->ResultsStructs.Num())
 			{
 				// row drag handle
-				
+			
 				return SNew(SChooserRowHandle).ChooserEditor(Editor).RowIndex(RowIndex->RowIndex);
 			}
 			else if (ColumnName == Result) 
@@ -443,22 +443,22 @@ public:
 				}),
 				&CacheBorder
 				);
-				
+			
 				return SNew(SOverlay)
-                		+ SOverlay::Slot()
-                		[
-                			ResultWidget.ToSharedRef()
-                		]
-                		+ SOverlay::Slot().VAlign(VAlign_Bottom)
-                		[
-                			SNew(SSeparator).SeparatorImage(FCoreStyle::Get().GetBrush("FocusRectangle"))
-                			.Visibility_Lambda([this]() { return bDragActive && !bDropAbove ? EVisibility::Visible : EVisibility::Hidden; })
-                		]
-                		+ SOverlay::Slot().VAlign(VAlign_Top)
-                		[
-                			SNew(SSeparator).SeparatorImage(FCoreStyle::Get().GetBrush("FocusRectangle"))
-                			.Visibility_Lambda([this]() { return bDragActive && bDropAbove ? EVisibility::Visible : EVisibility::Hidden; })
-                		];
+						+ SOverlay::Slot()
+						[
+							ResultWidget.ToSharedRef()
+						]
+						+ SOverlay::Slot().VAlign(VAlign_Bottom)
+						[
+							SNew(SSeparator).SeparatorImage(FCoreStyle::Get().GetBrush("FocusRectangle"))
+							.Visibility_Lambda([this]() { return bDragActive && !bDropAbove ? EVisibility::Visible : EVisibility::Hidden; })
+						]
+						+ SOverlay::Slot().VAlign(VAlign_Top)
+						[
+							SNew(SSeparator).SeparatorImage(FCoreStyle::Get().GetBrush("FocusRectangle"))
+							.Visibility_Lambda([this]() { return bDragActive && bDropAbove ? EVisibility::Visible : EVisibility::Hidden; })
+						];
 			}
 			else
 			{
@@ -469,7 +469,7 @@ public:
 					const UStruct * ColumnStruct = Chooser->ColumnsStructs[ColumnIndex].GetScriptStruct();
 
 					TSharedPtr<SWidget> ColumnWidget = FObjectChooserWidgetFactories::CreateColumnWidget(Column, ColumnStruct, Chooser, RowIndex->RowIndex);
-					
+				
 					if (ColumnWidget.IsValid())
 					{
 						return ColumnWidget.ToSharedRef();
@@ -488,7 +488,7 @@ public:
 
 		return SNullWidget::NullWidget;
 	}
-	
+
 	virtual void OnDragEnter(const FGeometry& MyGeometry, const FDragDropEvent& DragDropEvent) override
 	{
 		if (TSharedPtr<FChooserRowDragDropOp> Operation = DragDropEvent.GetOperationAs<FChooserRowDragDropOp>())
@@ -502,7 +502,7 @@ public:
 	{
 		bDragActive = false;
 	}
-	
+
 	virtual FReply OnDragOver(const FGeometry& MyGeometry, const FDragDropEvent& DragDropEvent) override
 	{
 		if (TSharedPtr<FChooserRowDragDropOp> Operation = DragDropEvent.GetOperationAs<FChooserRowDragDropOp>())
@@ -513,8 +513,8 @@ public:
 		}
 		return FReply::Unhandled();
 	}
-	
-    virtual FReply OnDrop(const FGeometry& MyGeometry, const FDragDropEvent& DragDropEvent) override
+
+	virtual FReply OnDrop(const FGeometry& MyGeometry, const FDragDropEvent& DragDropEvent) override
 	{
 		if (TSharedPtr<FChooserRowDragDropOp> Operation = DragDropEvent.GetOperationAs<FChooserRowDragDropOp>())
 		{
@@ -547,7 +547,7 @@ private:
 TSharedRef<ITableRow> FChooserTableEditor::GenerateTableRow(TSharedPtr<FChooserTableRow> InItem, const TSharedRef<STableViewBase>& OwnerTable)
 {
 	UChooserTable* Chooser = Cast<UChooserTable>(EditingObjects[0]);
-	
+
 	return SNew(SChooserTableRow, OwnerTable)
 		.Entry(InItem).Chooser(Chooser).Editor(this);
 }
@@ -565,9 +565,9 @@ void FChooserTableEditor::MoveRow(int SourceRowIndex, int TargetRowIndex)
 {
 	UChooserTable* Chooser = Cast<UChooserTable>(EditingObjects[0]);
 	TargetRowIndex = FMath::Min(TargetRowIndex,Chooser->ResultsStructs.Num());
-	
+
 	const FScopedTransaction Transaction(LOCTEXT("Move Row", "Move Row"));
-	
+
 	Chooser->Modify(true);
 
 	for (FInstancedStruct& ColStruct : Chooser->ColumnsStructs)
@@ -575,7 +575,7 @@ void FChooserTableEditor::MoveRow(int SourceRowIndex, int TargetRowIndex)
 		FChooserColumnBase& Column = ColStruct.GetMutable<FChooserColumnBase>();
 		Column.MoveRow(SourceRowIndex, TargetRowIndex);
 	}
-	
+
 	FInstancedStruct Result = Chooser->ResultsStructs[SourceRowIndex];
 	Chooser->ResultsStructs.RemoveAt(SourceRowIndex);
 	if (SourceRowIndex < TargetRowIndex)
@@ -590,13 +590,13 @@ void FChooserTableEditor::MoveRow(int SourceRowIndex, int TargetRowIndex)
 void FChooserTableEditor::UpdateTableColumns()
 {
 	UChooserTable* Chooser = Cast<UChooserTable>(EditingObjects[0]);
-	
+
 	HeaderRow->ClearColumns();
-	
+
 	HeaderRow->AddColumn(SHeaderRow::Column("Handles")
 					.DefaultLabel(FText())
 					.ManualWidth(30));
-	
+
 	HeaderRow->AddColumn(SHeaderRow::Column("Result")
 					.DefaultLabel(LOCTEXT("ResultColumnName", "Result"))
 					.ManualWidth(300));
@@ -612,6 +612,22 @@ void FChooserTableEditor::UpdateTableColumns()
 			HeaderWidget = FObjectChooserWidgetFactories::CreateWidget(false, Chooser, InputValue, Column.GetInputType(), Chooser->ContextObjectType, Chooser->OutputObjectType);
 		}
 
+
+		const ISlateStyle& CoreStyle = FCoreStyle::Get();
+		const FSlateBrush* ColumnIcon = nullptr;
+		if (Column.IsRandomizeColumn())
+		{
+			ColumnIcon = FAppStyle::Get().GetBrush("Icons.Help");
+		}
+		else if (Column.HasOutputs())
+		{
+			ColumnIcon = CoreStyle.GetBrush("Icons.ArrowRight");
+		}
+		else
+		{
+			ColumnIcon = CoreStyle.GetBrush("Icons.Filter");
+		}
+
 		HeaderRow->AddColumn(SHeaderRow::FColumn::FArguments()
 			.ColumnId(ColumnId)
 			.ManualWidth(200)
@@ -619,8 +635,8 @@ void FChooserTableEditor::UpdateTableColumns()
 			{
 				UChooserColumnMenuContext* MenuContext = NewObject<UChooserColumnMenuContext>();
 				MenuContext->Editor = this;
-                MenuContext->Chooser = Chooser;
-                MenuContext->ColumnIndex = ColumnIndex;
+				MenuContext->Chooser = Chooser;
+				MenuContext->ColumnIndex = ColumnIndex;
 
 				FMenuBuilder MenuBuilder(true,nullptr);
 
@@ -633,7 +649,7 @@ void FChooserTableEditor::UpdateTableColumns()
 						)
 					);
 
-				if (ColumnIndex > 0)
+				if (ColumnIndex > 0 && !Chooser->ColumnsStructs[ColumnIndex].Get<FChooserColumnBase>().IsRandomizeColumn())
 				{
 					MenuBuilder.AddMenuEntry(LOCTEXT("Move Left","Move Left"),LOCTEXT("Move Left ToolTip", "Move this column to the left."),FSlateIcon(),
 						FUIAction(
@@ -646,7 +662,7 @@ void FChooserTableEditor::UpdateTableColumns()
 							})
 							));
 				}
-				if (ColumnIndex < Chooser->ColumnsStructs.Num() - 1)
+				if (ColumnIndex < Chooser->ColumnsStructs.Num() - 1 && !Chooser->ColumnsStructs[ColumnIndex+1].Get<FChooserColumnBase>().IsRandomizeColumn())
 				{
 					MenuBuilder.AddMenuEntry(LOCTEXT("Move Right","Move Right"),LOCTEXT("Move Right ToolTip", "Move this column to the right."),FSlateIcon(),
 						FUIAction(
@@ -672,7 +688,7 @@ void FChooserTableEditor::UpdateTableColumns()
 							UpdateTableColumns();
 						})
 						));
-				
+			
 				MenuBuilder.AddSubMenu(LOCTEXT("Input Type", "Input Type"),
 					LOCTEXT("InputTypeToolTip", "Change input parameter type"),
 					FNewMenuDelegate::CreateLambda([this, Chooser, &ColumnIndex](FMenuBuilder& Builder)
@@ -680,7 +696,7 @@ void FChooserTableEditor::UpdateTableColumns()
 						FStructViewerInitializationOptions Options;
 						Options.StructFilter = MakeShared<FStructFilter>(Chooser->ColumnsStructs[ColumnIndex].Get<FChooserColumnBase>().GetInputBaseType());
 						Options.NameTypeToDisplay = EStructViewerNameTypeToDisplay::DisplayName;
-						
+					
 						// Add class filter for columns here
 						TSharedRef<SWidget> Widget = FModuleManager::LoadModuleChecked<FStructViewerModule>("StructViewer").CreateStructViewer(Options, FOnStructPicked::CreateLambda([this, ColumnIndex](const UScriptStruct* ChosenStruct)
 						{
@@ -691,12 +707,12 @@ void FChooserTableEditor::UpdateTableColumns()
 							UpdateTableColumns();
 							UpdateTableRows();
 						}));
-					
+				
 						Builder.AddWidget(Widget, FText());
 					}));
-				
+			
 				return MenuBuilder.MakeWidget();
-				
+			
 			})
 			.HeaderContent()
 			[
@@ -712,7 +728,7 @@ void FChooserTableEditor::UpdateTableColumns()
 				{
 					UChooserTable* Chooser = Cast<UChooserTable>(EditingObjects[0]);
 					TableView->ClearSelection();
-					
+				
 					SelectColumn(ColumnId.GetNumber() - 1);
 					return FReply::Handled();
 				})
@@ -724,7 +740,7 @@ void FChooserTableEditor::UpdateTableColumns()
 						.BorderBackgroundColor(FLinearColor(0,0,0,0))
 						.Content()
 						[
-							SNew(SImage).Image(Column.HasOutputs() ? FCoreStyle::Get().GetBrush("Icons.ArrowRight") :  FCoreStyle::Get().GetBrush("Icons.Filter"))
+							SNew(SImage).Image(ColumnIcon)
 						]
 					]
 					+ SHorizontalBox::Slot()
@@ -732,9 +748,9 @@ void FChooserTableEditor::UpdateTableColumns()
 						HeaderWidget ? HeaderWidget.ToSharedRef() : SNullWidget::NullWidget
 					]
 				]
-				
+			
 			]);
-		
+	
 		ColumnId.SetNumber(ColumnId.GetNumber() + 1);
 	}
 
@@ -753,6 +769,57 @@ void FChooserTableEditor::UpdateTableColumns()
 
 }
 
+void FChooserTableEditor::AddColumn(const UScriptStruct* ColumnType)
+{
+	CreateColumnComboButton->SetIsOpen(false);
+	UChooserTable* Chooser = Cast<UChooserTable>(EditingObjects[0]);
+	const FScopedTransaction Transaction(LOCTEXT("Add Column Transaction", "Add Column"));
+	Chooser->Modify(true);
+
+	FInstancedStruct NewColumn;
+	NewColumn.InitializeAs(ColumnType);
+	const FChooserColumnBase& NewColumnRef = NewColumn.Get<FChooserColumnBase>();
+	int InsertIndex = 0;
+	if (NewColumnRef.IsRandomizeColumn())
+	{
+		// add randomization column at the end (do nothing if there already is one)
+		InsertIndex = Chooser->ColumnsStructs.Num();
+		if (InsertIndex == 0 || !Chooser->ColumnsStructs[InsertIndex - 1].Get<FChooserColumnBase>().IsRandomizeColumn())
+		{
+			Chooser->ColumnsStructs.Add(NewColumn);
+		}
+	}
+	else if (NewColumnRef.HasOutputs())
+	{
+		// add output columns at the end (but before any randomization column)
+		InsertIndex = Chooser->ColumnsStructs.Num();
+		if (InsertIndex > 0 && Chooser->ColumnsStructs[InsertIndex - 1].Get<FChooserColumnBase>().IsRandomizeColumn())
+		{
+			InsertIndex--;
+		}
+		Chooser->ColumnsStructs.Insert(NewColumn, InsertIndex);
+	}
+	else
+	{
+		// add other columns after the last non-output, non-randomization column
+		while(InsertIndex < Chooser->ColumnsStructs.Num())
+		{
+			const FChooserColumnBase& Column = Chooser->ColumnsStructs[InsertIndex].Get<FChooserColumnBase>();
+			if (Column.HasOutputs() || Column.IsRandomizeColumn())
+			{
+				break;
+			}
+			InsertIndex++;
+		}
+		Chooser->ColumnsStructs.Insert(NewColumn, InsertIndex);
+	}
+
+	UpdateTableColumns();
+	UpdateTableRows();
+
+	SelectColumn(InsertIndex);
+}
+
 TSharedRef<SDockTab> FChooserTableEditor::SpawnTableTab( const FSpawnTabArgs& Args )
 {
 	check( Args.GetTabId() == TableTabId );
@@ -768,43 +835,8 @@ TSharedRef<SDockTab> FChooserTableEditor::SpawnTableTab( const FSpawnTabArgs& Ar
 		Options.NameTypeToDisplay = EStructViewerNameTypeToDisplay::DisplayName;
 		
 		// Add class filter for columns here
-		TSharedRef<SWidget> Widget = FModuleManager::LoadModuleChecked<FStructViewerModule>("StructViewer").CreateStructViewer(Options, FOnStructPicked::CreateLambda([this](const UScriptStruct* ChosenStruct)
-		{
-			CreateColumnComboButton->SetIsOpen(false);
-			UChooserTable* Chooser = Cast<UChooserTable>(EditingObjects[0]);
-			const FScopedTransaction Transaction(LOCTEXT("Add Column Transaction", "Add Column"));
-			Chooser->Modify(true);
-
-			FInstancedStruct NewColumn;
-			NewColumn.InitializeAs(ChosenStruct);
-			const FChooserColumnBase& NewColumnRef = NewColumn.Get<FChooserColumnBase>();
-			int InsertIndex = 0;
-			if (NewColumnRef.HasOutputs())
-			{
-				InsertIndex = Chooser->ColumnsStructs.Num();
-				// add output columns at the end
-				Chooser->ColumnsStructs.Add(NewColumn);
-			}
-			else
-			{
-				// add other columns after the loast non-output column
-				while(InsertIndex < Chooser->ColumnsStructs.Num())
-				{
-					const FChooserColumnBase& Column = Chooser->ColumnsStructs[InsertIndex].Get<FChooserColumnBase>();
-					if (Column.HasOutputs())
-					{
-						break;
-					}
-					InsertIndex++;
-				}
-				Chooser->ColumnsStructs.Insert(NewColumn, InsertIndex);
-			}
-
-			UpdateTableColumns();
-			UpdateTableRows();
-
-			SelectColumn(InsertIndex);
-		}));
+		FStructViewerModule& StructViewerModule = FModuleManager::LoadModuleChecked<FStructViewerModule>("StructViewer");
+		TSharedRef<SWidget> Widget = StructViewerModule.CreateStructViewer(Options, FOnStructPicked::CreateRaw(this, &FChooserTableEditor::AddColumn));
 		return Widget;
 	})
 	.ComboButtonStyle(FAppStyle::Get(), "SimpleComboButton")
diff --git a/Engine/Plugins/Experimental/Chooser/Source/ChooserEditor/Private/ChooserTableEditor.h b/Engine/Plugins/Experimental/Chooser/Source/ChooserEditor/Private/ChooserTableEditor.h
--- a/Engine/Plugins/Experimental/Chooser/Source/ChooserEditor/Private/ChooserTableEditor.h
+++ b/Engine/Plugins/Experimental/Chooser/Source/ChooserEditor/Private/ChooserTableEditor.h
@@ -98,6 +98,7 @@ namespace UE::ChooserEditor
 		void UpdateTableRows();
 		void SelectColumn(int Index);
 		void DeleteColumn(int Index);
+		void AddColumn(const UScriptStruct* ColumnType);
 		void MoveRow(int SourceRowIndex, int TargetIndex);
 	private:
 
diff --git a/Engine/Plugins/Experimental/Chooser/Source/ChooserEditor/Private/RandomizeColumnEditor.cpp b/Engine/Plugins/Experimental/Chooser/Source/ChooserEditor/Private/RandomizeColumnEditor.cpp
new file mode 100644
--- /dev/null
+++ b/Engine/Plugins/Experimental/Chooser/Source/ChooserEditor/Private/RandomizeColumnEditor.cpp
@@ -0,0 +1,76 @@
+// Copyright Epic Games, Inc. All Rights Reserved.
+
+#include "RandomizeColumnEditor.h"
+#include "RandomizeColumn.h"
+#include "SPropertyAccessChainWidget.h"
+#include "ObjectChooserWidgetFactories.h"
+#include "Widgets/SBoxPanel.h"
+#include "Widgets/Layout/SBox.h"
+#include "Widgets/Input/SNumericEntryBox.h"
+#include "GraphEditorSettings.h"
+
+#define LOCTEXT_NAMESPACE "RandomizeColumnEditor"
+
+namespace UE::ChooserEditor
+{
+
+TSharedRef<SWidget> CreateRandomizeColumnWidget(UChooserTable* Chooser, FChooserColumnBase* Column, int Row)
+{
+	FRandomizeColumn* RandomizeColumn = static_cast<FRandomizeColumn*>(Column);
+	
+	return
+	SNew(SHorizontalBox)
+		+ SHorizontalBox::Slot().FillWidth(1)
+		+ SHorizontalBox::Slot().AutoWidth()
+		[
+			SNew(SBox).WidthOverride(75).Content()
+			[
+				SNew(SNumericEntryBox<float>)
+    				.Value_Lambda([RandomizeColumn, Row]()
+    				{
+    					if (!RandomizeColumn->RowValues.IsValidIndex(Row))
+    					{
+    						return 0.0f;
+    					}
+    					return RandomizeColumn->RowValues[Row];
+    				})
+    				.OnValueCommitted_Lambda([Chooser, Row, RandomizeColumn](float Value, ETextCommit::Type CommitType)
+    				{
+    					if (RandomizeColumn->RowValues.IsValidIndex(Row))
+    					{
+    						const FScopedTransaction Transaction(LOCTEXT("Edit Randomize Cell Data", "Edit Randomize Cell Data"));
+    						Chooser->Modify(true);
+    						RandomizeColumn->RowValues[Row] = Value;
+    					}
+    				})
+    		]
+    	]
+		+ SHorizontalBox::Slot().FillWidth(1);
+}
+	
+TSharedRef<SWidget> CreateRandomizePropertyWidget(bool bReadOnly, UObject* TransactionObject, void* Value, UClass* ContextClass, UClass* ResultBaseClass)
+{
+	IHasContextClass* HasContextClass = Cast<IHasContextClass>(TransactionObject);
+
+	FRandomizeContextProperty* ContextProperty = reinterpret_cast<FRandomizeContextProperty*>(Value);
+
+	return SNew(SPropertyAccessChainWidget).ContextClassOwner(HasContextClass).AllowFunctions(true).BindingColor("StructPinTypeColor").TypeFilter("FChooserRandomizationContext")
+	.PropertyBindingValue(&ContextProperty->Binding)
+	.OnAddBinding_Lambda(
+		[ContextProperty, TransactionObject](FName InPropertyName, const TArray<FBindingChainElement>& InBindingChain)
+		{
+			const FScopedTransaction Transaction(NSLOCTEXT("ContextPropertyWidget", "Change Property Binding", "Change Property Binding"));
+			TransactionObject->Modify(true);
+			ContextProperty->SetBinding(InBindingChain);	
+		});
+}
+	
+void RegisterRandomizeWidgets()
+{
+	FObjectChooserWidgetFactories::RegisterWidgetCreator(FRandomizeContextProperty::StaticStruct(), CreateRandomizePropertyWidget);
+	FObjectChooserWidgetFactories::RegisterColumnWidgetCreator(FRandomizeColumn::StaticStruct(), CreateRandomizeColumnWidget);
+}
+	
+}
+
+#undef LOCTEXT_NAMESPACE
diff --git a/Engine/Plugins/Experimental/Chooser/Source/ChooserEditor/Private/RandomizeColumnEditor.h b/Engine/Plugins/Experimental/Chooser/Source/ChooserEditor/Private/RandomizeColumnEditor.h
new file mode 100644
--- /dev/null
+++ b/Engine/Plugins/Experimental/Chooser/Source/ChooserEditor/Private/RandomizeColumnEditor.h
@@ -0,0 +1,10 @@
+// Copyright Epic Games, Inc. All Rights Reserved.
+
+#pragma once
+
+#include "CoreMinimal.h"
+
+namespace UE::ChooserEditor
+{
+	void RegisterRandomizeWidgets();
+}
\ No newline at end of file
