#Chooser Fix memory corruption in UK2Node_EvaluateChooser

diff --git a/Engine/Plugins/Experimental/Chooser/Source/ChooserUncooked/Private/EvaluateChooserNode.cpp b/Engine/Plugins/Experimental/Chooser/Source/ChooserUncooked/Private/EvaluateChooserNode.cpp
--- a/Engine/Plugins/Experimental/Chooser/Source/ChooserUncooked/Private/EvaluateChooserNode.cpp
+++ b/Engine/Plugins/Experimental/Chooser/Source/ChooserUncooked/Private/EvaluateChooserNode.cpp
@@ -49,9 +49,11 @@ void UK2Node_EvaluateChooser::UnregisterChooserCallback()
 	}
 }
 
-UK2Node_EvaluateChooser::~UK2Node_EvaluateChooser()
+void UK2Node_EvaluateChooser::BeginDestroy()
 {
 	UnregisterChooserCallback();
+	
+	Super::BeginDestroy();
 }
 
 void UK2Node_EvaluateChooser::PostEditUndo()
diff --git a/Engine/Plugins/Experimental/Chooser/Source/ChooserUncooked/Private/EvaluateChooserNode.h b/Engine/Plugins/Experimental/Chooser/Source/ChooserUncooked/Private/EvaluateChooserNode.h
--- a/Engine/Plugins/Experimental/Chooser/Source/ChooserUncooked/Private/EvaluateChooserNode.h
+++ b/Engine/Plugins/Experimental/Chooser/Source/ChooserUncooked/Private/EvaluateChooserNode.h
@@ -25,9 +25,9 @@ UCLASS(MinimalAPI)
 class UK2Node_EvaluateChooser : public UK2Node
 {
 	GENERATED_UCLASS_BODY()
-	~UK2Node_EvaluateChooser();
 	
 	//~ Begin UObject Interface
+	virtual void BeginDestroy() override;
 	virtual void PostEditUndo() override;
 	virtual void PostEditChangeProperty(struct FPropertyChangedEvent& PropertyChangedEvent) override;
 	virtual void PostLoad() override;
