Chooser Testing and Debugging

diff --git a/Engine/Plugins/Experimental/Chooser/Source/Chooser/Internal/BoolColumn.h b/Engine/Plugins/Experimental/Chooser/Source/Chooser/Internal/BoolColumn.h
--- a/Engine/Plugins/Experimental/Chooser/Source/Chooser/Internal/BoolColumn.h
+++ b/Engine/Plugins/Experimental/Chooser/Source/Chooser/Internal/BoolColumn.h
@@ -85,7 +85,15 @@ struct CHOOSER_API FBoolColumn : public FChooserColumnBase
 	UPROPERTY(EditAnywhere, Category= "Data", DisplayName="RowValues");
 	TArray<EBoolColumnCellValue> RowValuesWithAny; 
 	
-	virtual void Filter(const UObject* ContextObject, const TArray<uint32>& IndexListIn, TArray<uint32>& IndexListOut) const override;
+	virtual void Filter(FChooserDebuggingInfo& DebugInfo, const UObject* ContextObject, const TArray<uint32>& IndexListIn, TArray<uint32>& IndexListOut) const override;
+
+#if WITH_EDITOR
+	mutable bool TestValue;
+	virtual bool EditorTestFilter(int32 RowIndex) const override
+	{
+		return RowValuesWithAny.IsValidIndex(RowIndex) && (RowValuesWithAny[RowIndex] == EBoolColumnCellValue::MatchAny || TestValue == static_cast<bool>(RowValuesWithAny[RowIndex]));
+	}
+#endif
 
 	virtual void PostLoad() override
 	{
diff --git a/Engine/Plugins/Experimental/Chooser/Source/Chooser/Internal/Chooser.h b/Engine/Plugins/Experimental/Chooser/Source/Chooser/Internal/Chooser.h
--- a/Engine/Plugins/Experimental/Chooser/Source/Chooser/Internal/Chooser.h
+++ b/Engine/Plugins/Experimental/Chooser/Source/Chooser/Internal/Chooser.h
@@ -12,8 +12,8 @@
 
 DECLARE_MULTICAST_DELEGATE_OneParam(FChooserOutputObjectTypeChanged, const UClass* OutputObjectType);
 
-UCLASS(MinimalAPI, BlueprintType)
-class UChooserTable : public UObject, public IHasContextClass
+UCLASS(BlueprintType)
+class CHOOSER_API UChooserTable : public UObject, public IHasContextClass
 {
 	GENERATED_UCLASS_BODY()
 public:
@@ -25,12 +25,35 @@ public:
 	virtual void PostEditUndo() override;
 	virtual void PostEditChangeProperty(FPropertyChangedEvent& PropertyChangedEvent) override;
 	virtual void PostLoad() override;
+	
+	void SetDebugSelectedRow(int32 Index) const { DebugSelectedRow = Index; }
+	int32 GetDebugSelectedRow() const { return DebugSelectedRow; }
+	bool HasDebugTarget() const { return DebugTarget != nullptr; }
+	const UObject* GetDebugTarget() const { return DebugTarget.Get(); }
+	void SetDebugTarget(TWeakObjectPtr<const UObject> Target) { DebugTarget = Target; }
+	void ResetDebugTarget() { DebugTarget.Reset(); }
+	void IterateRecentContextObjects(TFunction<void(const UObject*)> Callback) const;
+	bool UpdateDebugging(const UObject* ContextObject) const;
+	
+	// enable display of which cells pass/fail based on current TestValue for each column
+	bool bEnableDebugTesting = false;
+	mutable bool bDebugTestValuesValid = false;
 
+private: 
 	// caching the OutputObjectType and ContextObjectType so that on Undo, we can tell if we should fire the changed delegate
 	UClass* CachedPreviousOutputObjectType;
 	UClass* CachedPreviousContextObjectType;
-#endif
 	
+	// objects this chooser has been recently evaluated on
+	mutable TSet<TWeakObjectPtr<const UObject>> RecentContextObjects;
+	mutable FCriticalSection DebugLock;
+	// reference to the UObject in PIE  which we want to get debug info for
+	TWeakObjectPtr<const UObject> DebugTarget;
+	// Row which was selected last time this chooser was evaluated on DebugTarget
+	mutable int32 DebugSelectedRow = -1;
+#endif
+
+public:
 #if WITH_EDITORONLY_DATA
 	// deprecated UObject Results
 	UPROPERTY()
diff --git a/Engine/Plugins/Experimental/Chooser/Source/Chooser/Internal/EnumColumn.h b/Engine/Plugins/Experimental/Chooser/Source/Chooser/Internal/EnumColumn.h
--- a/Engine/Plugins/Experimental/Chooser/Source/Chooser/Internal/EnumColumn.h
+++ b/Engine/Plugins/Experimental/Chooser/Source/Chooser/Internal/EnumColumn.h
@@ -108,7 +108,15 @@ public:
 	// should match the length of the Results array
 	TArray<FChooserEnumRowData> RowValues;
 
-	virtual void Filter(const UObject* ContextObject, const TArray<uint32>& IndexListIn, TArray<uint32>& IndexListOut) const override;
+	virtual void Filter(FChooserDebuggingInfo& DebugInfo, const UObject* ContextObject, const TArray<uint32>& IndexListIn, TArray<uint32>& IndexListOut) const override;
+	
+#if WITH_EDITOR
+	mutable int32 TestValue;
+	virtual bool EditorTestFilter(int32 RowIndex) const override
+	{
+		return RowValues.IsValidIndex(RowIndex) && RowValues[RowIndex].Evaluate(TestValue);
+	}
+#endif
 	
 	CHOOSER_COLUMN_BOILERPLATE(FChooserParameterEnumBase);
 
diff --git a/Engine/Plugins/Experimental/Chooser/Source/Chooser/Internal/FloatRangeColumn.h b/Engine/Plugins/Experimental/Chooser/Source/Chooser/Internal/FloatRangeColumn.h
--- a/Engine/Plugins/Experimental/Chooser/Source/Chooser/Internal/FloatRangeColumn.h
+++ b/Engine/Plugins/Experimental/Chooser/Source/Chooser/Internal/FloatRangeColumn.h
@@ -87,7 +87,15 @@ struct CHOOSER_API FFloatRangeColumn : public FChooserColumnBase
 	// should match the length of the Results array 
 	TArray<FChooserFloatRangeRowData> RowValues;
 	
-	virtual void Filter(const UObject* ContextObject, const TArray<uint32>& IndexListIn, TArray<uint32>& IndexListOut) const override;
+	virtual void Filter(FChooserDebuggingInfo& DebugInfo, const UObject* ContextObject, const TArray<uint32>& IndexListIn, TArray<uint32>& IndexListOut) const override;
+
+#if WITH_EDITOR
+	mutable float TestValue;
+	virtual bool EditorTestFilter(int32 RowIndex) const override
+	{
+		return RowValues.IsValidIndex(RowIndex) && TestValue >= RowValues[RowIndex].Min && TestValue <= RowValues[RowIndex].Max;
+	}
+#endif
 	
 	virtual void PostLoad() override
 	{
diff --git a/Engine/Plugins/Experimental/Chooser/Source/Chooser/Internal/GameplayTagColumn.h b/Engine/Plugins/Experimental/Chooser/Source/Chooser/Internal/GameplayTagColumn.h
--- a/Engine/Plugins/Experimental/Chooser/Source/Chooser/Internal/GameplayTagColumn.h
+++ b/Engine/Plugins/Experimental/Chooser/Source/Chooser/Internal/GameplayTagColumn.h
@@ -71,6 +71,7 @@ struct CHOOSER_API FGameplayTagColumn : public FChooserColumnBase
 #if WITH_EDITORONLY_DATA
 	UPROPERTY(EditAnywhere, Category=Runtime)
 	FGameplayTagContainer DefaultRowValue;
+
 #endif
 	
 	UPROPERTY(EditAnywhere, Category=Runtime)
@@ -78,7 +79,26 @@ struct CHOOSER_API FGameplayTagColumn : public FChooserColumnBase
 	// should match the length of the Results array 
 	TArray<FGameplayTagContainer> RowValues;
 	
-	virtual void Filter(const UObject* ContextObject, const TArray<uint32>& IndexListIn, TArray<uint32>& IndexListOut) const override;
+	virtual void Filter(FChooserDebuggingInfo& DebugInfo, const UObject* ContextObject, const TArray<uint32>& IndexListIn, TArray<uint32>& IndexListOut) const override;
+
+#if WITH_EDITOR
+	mutable FGameplayTagContainer TestValue;
+	virtual bool EditorTestFilter(int32 RowIndex) const override
+	{
+		if(RowValues.IsValidIndex(RowIndex))
+		{
+			if (TagMatchType == EGameplayContainerMatchType::All)
+			{
+				return TestValue.HasAll(RowValues[RowIndex]);
+			}
+			else
+			{
+				return TestValue.HasAny(RowValues[RowIndex]);
+			}
+		}
+		return false;
+	}
+#endif
 
 	virtual void PostLoad() override
 	{
diff --git a/Engine/Plugins/Experimental/Chooser/Source/Chooser/Internal/ObjectColumn.h b/Engine/Plugins/Experimental/Chooser/Source/Chooser/Internal/ObjectColumn.h
--- a/Engine/Plugins/Experimental/Chooser/Source/Chooser/Internal/ObjectColumn.h
+++ b/Engine/Plugins/Experimental/Chooser/Source/Chooser/Internal/ObjectColumn.h
@@ -21,7 +21,6 @@ struct CHOOSER_API FObjectContextProperty : public FChooserParameterObjectBase
 	virtual bool GetValue(const UObject* ContextObject, FSoftObjectPath& OutResult) const override;
 
 #if WITH_EDITOR
-
 	static bool CanBind(const FProperty& Property)
 	{
 		return Property.IsA<FObjectPropertyBase>() && !(Property.IsA<FClassProperty>() || Property.IsA<FSoftClassProperty>());
@@ -38,7 +37,6 @@ struct CHOOSER_API FObjectContextProperty : public FChooserParameterObjectBase
 	}
 
 	virtual UClass* GetAllowedClass() const override { return Binding.AllowedClass; }
-
 #endif
 };
 
@@ -86,7 +84,15 @@ struct CHOOSER_API FObjectColumn : public FChooserColumnBase
 	// should match the length of the Results array
 	TArray<FChooserObjectRowData> RowValues;
 
-	virtual void Filter(const UObject* ContextObject, const TArray<uint32>& IndexListIn, TArray<uint32>& IndexListOut) const override;
+	virtual void Filter(FChooserDebuggingInfo& DebugInfo, const UObject* ContextObject, const TArray<uint32>& IndexListIn, TArray<uint32>& IndexListOut) const override;
+
+#if WITH_EDITOR
+	mutable FSoftObjectPath TestValue;
+	virtual bool EditorTestFilter(int32 RowIndex) const override
+	{
+		return RowValues.IsValidIndex(RowIndex) && RowValues[RowIndex].Evaluate(TestValue);
+	}
+#endif
 
 	CHOOSER_COLUMN_BOILERPLATE(FChooserParameterObjectBase);
 
diff --git a/Engine/Plugins/Experimental/Chooser/Source/Chooser/Internal/OutputBoolColumn.h b/Engine/Plugins/Experimental/Chooser/Source/Chooser/Internal/OutputBoolColumn.h
--- a/Engine/Plugins/Experimental/Chooser/Source/Chooser/Internal/OutputBoolColumn.h
+++ b/Engine/Plugins/Experimental/Chooser/Source/Chooser/Internal/OutputBoolColumn.h
@@ -16,10 +16,14 @@ struct CHOOSER_API FOutputBoolColumn : public FChooserColumnBase
 	FOutputBoolColumn();
 	virtual bool HasFilters() const override { return false; }
 	virtual bool HasOutputs() const override { return true; }
-	virtual void SetOutputs(UObject* ContextObject, int RowIndex) const override;
+	virtual void SetOutputs(FChooserDebuggingInfo& DebugInfo, UObject* ContextObject, int RowIndex) const override;
 	
 	UPROPERTY(EditAnywhere, Meta = (ExcludeBaseStruct, BaseStruct = "/Script/Chooser.ChooserParameterBoolBase"), Category = "Hidden")
 	FInstancedStruct InputValue;
+
+#if WITH_EDITOR
+	mutable bool TestValue=false;
+#endif
 	
 #if WITH_EDITORONLY_DATA
 	UPROPERTY(EditAnywhere, Category=Runtime);
diff --git a/Engine/Plugins/Experimental/Chooser/Source/Chooser/Internal/RandomizeColumn.h b/Engine/Plugins/Experimental/Chooser/Source/Chooser/Internal/RandomizeColumn.h
--- a/Engine/Plugins/Experimental/Chooser/Source/Chooser/Internal/RandomizeColumn.h
+++ b/Engine/Plugins/Experimental/Chooser/Source/Chooser/Internal/RandomizeColumn.h
@@ -52,7 +52,7 @@ struct CHOOSER_API FRandomizeColumn : public FChooserColumnBase
 	FInstancedStruct InputValue;
 	
 	
-	UPROPERTY(EditAnywhere, Category= "Data", meta=(Tooltip="Multiplies the weight of the previous chosen result (set to 0 to never pick the same result twice in a row)"));
+	UPROPERTY(EditAnywhere, Category= "Data", meta=(ClampMin="0.0",Tooltip="Multiplies the weight of the previous chosen result (set to 0 to never pick the same result twice in a row)"));
 	float RepeatProbabilityMultiplier = 1.0f;
 
 #if WITH_EDITORONLY_DATA
@@ -63,12 +63,19 @@ struct CHOOSER_API FRandomizeColumn : public FChooserColumnBase
 	UPROPERTY(EditAnywhere, Category= "Data", DisplayName="RowValues");
 	TArray<float> RowValues; 
 	
-	virtual void Filter(const UObject* ContextObject, const TArray<uint32>& IndexListIn, TArray<uint32>& IndexListOut) const override;
-	virtual void SetOutputs(UObject* ContextObject, int RowIndex) const override;
+	virtual void Filter(FChooserDebuggingInfo& DebugInfo, const UObject* ContextObject, const TArray<uint32>& IndexListIn, TArray<uint32>& IndexListOut) const override;
+	virtual void SetOutputs(FChooserDebuggingInfo& DebugInfo, UObject* ContextObject, int RowIndex) const override;
 	
 	virtual bool HasFilters() const override { return true; }
 	virtual bool HasOutputs() const override { return true; }
 
+	#if WITH_EDITOR
+    	virtual bool EditorTestFilter(int32 RowIndex) const override
+    	{
+    		return true;
+    	}
+    #endif
+
 	virtual void PostLoad() override
 	{
 		if (InputValue.IsValid())
diff --git a/Engine/Plugins/Experimental/Chooser/Source/Chooser/Private/BoolColumn.cpp b/Engine/Plugins/Experimental/Chooser/Source/Chooser/Private/BoolColumn.cpp
--- a/Engine/Plugins/Experimental/Chooser/Source/Chooser/Private/BoolColumn.cpp
+++ b/Engine/Plugins/Experimental/Chooser/Source/Chooser/Private/BoolColumn.cpp
@@ -1,6 +1,7 @@
 // Copyright Epic Games, Inc. All Rights Reserved.
 #include "BoolColumn.h"
 #include "ChooserPropertyAccess.h"
+#include "Chooser.h"
 
 bool FBoolContextProperty::GetValue(const UObject* ContextObject, bool& OutResult) const
 {
@@ -60,12 +61,19 @@ FBoolColumn::FBoolColumn()
 	InputValue.InitializeAs(FBoolContextProperty::StaticStruct());
 }
 
-void FBoolColumn::Filter(const UObject* ContextObject, const TArray<uint32>& IndexListIn, TArray<uint32>& IndexListOut) const
+void FBoolColumn::Filter(FChooserDebuggingInfo& DebugInfo, const UObject* ContextObject, const TArray<uint32>& IndexListIn, TArray<uint32>& IndexListOut) const
 {
 	if (ContextObject && InputValue.IsValid())
 	{
 		bool Result = false;
 		InputValue.Get<FChooserParameterBoolBase>().GetValue(ContextObject,Result);
+
+	#if WITH_EDITOR
+		if (DebugInfo.bCurrentDebugTarget)
+		{
+			TestValue = Result;
+		}
+	#endif
 		
 		for (uint32 Index : IndexListIn)
 		{
diff --git a/Engine/Plugins/Experimental/Chooser/Source/Chooser/Private/Chooser.cpp b/Engine/Plugins/Experimental/Chooser/Source/Chooser/Private/Chooser.cpp
--- a/Engine/Plugins/Experimental/Chooser/Source/Chooser/Private/Chooser.cpp
+++ b/Engine/Plugins/Experimental/Chooser/Source/Chooser/Private/Chooser.cpp
@@ -101,6 +101,31 @@ void UChooserTable::PostLoad()
 	}
 		
 }
+
+void UChooserTable::IterateRecentContextObjects(TFunction<void(const UObject*)> Callback) const
+{
+	FScopeLock Lock(&DebugLock);
+	for(TWeakObjectPtr<const UObject>& Object : RecentContextObjects)
+	{
+		if (Object.IsValid())
+		{
+			Callback(Object.Get());
+		}
+	}
+}
+
+bool UChooserTable::UpdateDebugging(const UObject* ContextObject) const
+{
+	FScopeLock Lock(&DebugLock);
+	RecentContextObjects.Add(MakeWeakObjectPtr(ContextObject));
+	if (ContextObject == DebugTarget)
+	{
+		bDebugTestValuesValid = true;
+		return true;
+	}
+	return false;
+}
+
 #endif
 
 static FObjectChooserBase::EIteratorStatus StaticEvaluateChooser(const UObject* ContextObject, const UChooserTable* Chooser, FObjectChooserBase::FObjectChooserIteratorCallback Callback)
@@ -109,6 +134,12 @@ static FObjectChooserBase::EIteratorStatus StaticEvaluateChooser(const UObject*
 	{
 		return FObjectChooserBase::EIteratorStatus::Continue;
 	}
+	
+	FChooserDebuggingInfo DebugInfo;
+
+#if WITH_EDITOR
+	DebugInfo.bCurrentDebugTarget = Chooser->UpdateDebugging(ContextObject);
+#endif
 
 	TArray<uint32> Indices1;
 	TArray<uint32> Indices2;
@@ -129,7 +160,7 @@ static FObjectChooserBase::EIteratorStatus StaticEvaluateChooser(const UObject*
 		{
 			Swap(IndicesIn, IndicesOut);
 			IndicesOut->SetNum(0, false);
-			Column.Filter(ContextObject, *IndicesIn, *IndicesOut);
+			Column.Filter(DebugInfo, ContextObject, *IndicesIn, *IndicesOut);
 		}
 	}
 	
@@ -145,8 +176,14 @@ static FObjectChooserBase::EIteratorStatus StaticEvaluateChooser(const UObject*
 				for (const FInstancedStruct& ColumnData : Chooser->ColumnsStructs)
 				{
 					const FChooserColumnBase& Column = ColumnData.Get<FChooserColumnBase>();
-					Column.SetOutputs(const_cast<UObject*>(ContextObject), SelectedIndex);
+					Column.SetOutputs(DebugInfo, const_cast<UObject*>(ContextObject), SelectedIndex);
 				}
+#if WITH_EDITOR
+				if (DebugInfo.bCurrentDebugTarget)
+				{
+					Chooser->SetDebugSelectedRow(SelectedIndex);
+				}
+#endif
 				return FObjectChooserBase::EIteratorStatus::Stop;
 			}
 		}
diff --git a/Engine/Plugins/Experimental/Chooser/Source/Chooser/Private/EnumColumn.cpp b/Engine/Plugins/Experimental/Chooser/Source/Chooser/Private/EnumColumn.cpp
--- a/Engine/Plugins/Experimental/Chooser/Source/Chooser/Private/EnumColumn.cpp
+++ b/Engine/Plugins/Experimental/Chooser/Source/Chooser/Private/EnumColumn.cpp
@@ -67,13 +67,20 @@ bool FChooserEnumRowData::Evaluate(const uint8 LeftHandSide) const
 	return Equal ^ CompareNotEqual;
 }
 
-void FEnumColumn::Filter(const UObject* ContextObject, const TArray<uint32>& IndexListIn, TArray<uint32>& IndexListOut) const
+void FEnumColumn::Filter(FChooserDebuggingInfo& DebugInfo, const UObject* ContextObject, const TArray<uint32>& IndexListIn, TArray<uint32>& IndexListOut) const
 {
 	uint8 Result = 0;
 	if (ContextObject != nullptr &&
 		InputValue.IsValid() &&
 		InputValue.Get<FChooserParameterEnumBase>().GetValue(ContextObject, Result))
 	{
+#if WITH_EDITOR
+		if (DebugInfo.bCurrentDebugTarget)
+		{
+			TestValue = Result;
+		}
+#endif
+		
 		for (const uint32 Index : IndexListIn)
 		{
 			if (RowValues.IsValidIndex(Index))
diff --git a/Engine/Plugins/Experimental/Chooser/Source/Chooser/Private/FloatRangeColumn.cpp b/Engine/Plugins/Experimental/Chooser/Source/Chooser/Private/FloatRangeColumn.cpp
--- a/Engine/Plugins/Experimental/Chooser/Source/Chooser/Private/FloatRangeColumn.cpp
+++ b/Engine/Plugins/Experimental/Chooser/Source/Chooser/Private/FloatRangeColumn.cpp
@@ -1,5 +1,6 @@
 // Copyright Epic Games, Inc. All Rights Reserved.
 #include "FloatRangeColumn.h"
+
 #include "ChooserPropertyAccess.h"
 
 bool FFloatContextProperty::GetValue(const UObject* ContextObject, float& OutResult) const
@@ -68,13 +69,20 @@ FFloatRangeColumn::FFloatRangeColumn()
 	InputValue.InitializeAs(FFloatContextProperty::StaticStruct());
 }
 
-void FFloatRangeColumn::Filter(const UObject* ContextObject, const TArray<uint32>& IndexListIn, TArray<uint32>& IndexListOut) const
+void FFloatRangeColumn::Filter(FChooserDebuggingInfo& DebugInfo, const UObject* ContextObject, const TArray<uint32>& IndexListIn, TArray<uint32>& IndexListOut) const
 {
 	if (ContextObject && InputValue.IsValid())
 	{
 		float Result = 0.0f;
 		InputValue.Get<FChooserParameterFloatBase>().GetValue(ContextObject, Result);
 
+#if WITH_EDITOR
+		if (DebugInfo.bCurrentDebugTarget)
+		{
+			TestValue = Result;
+		}
+#endif
+
 		for(uint32 Index : IndexListIn)
 		{
 			if (RowValues.Num() > (int)Index)
diff --git a/Engine/Plugins/Experimental/Chooser/Source/Chooser/Private/GameplayTagColumn.cpp b/Engine/Plugins/Experimental/Chooser/Source/Chooser/Private/GameplayTagColumn.cpp
--- a/Engine/Plugins/Experimental/Chooser/Source/Chooser/Private/GameplayTagColumn.cpp
+++ b/Engine/Plugins/Experimental/Chooser/Source/Chooser/Private/GameplayTagColumn.cpp
@@ -24,11 +24,18 @@ FGameplayTagColumn::FGameplayTagColumn()
 	InputValue.InitializeAs(FGameplayTagContextProperty::StaticStruct());
 }
 
-void FGameplayTagColumn::Filter(const UObject* ContextObject, const TArray<uint32>& IndexListIn, TArray<uint32>& IndexListOut) const
+void FGameplayTagColumn::Filter(FChooserDebuggingInfo& DebugInfo, const UObject* ContextObject, const TArray<uint32>& IndexListIn, TArray<uint32>& IndexListOut) const
 {
 	const FGameplayTagContainer* Result = nullptr;
 	if (ContextObject && InputValue.IsValid() && InputValue.Get<FChooserParameterGameplayTagBase>().GetValue(ContextObject,Result))
 	{
+#if WITH_EDITOR
+		if (DebugInfo.bCurrentDebugTarget)
+		{
+			TestValue = *Result;
+		}
+#endif
+
 		for (uint32 Index : IndexListIn)
 		{
 		 
diff --git a/Engine/Plugins/Experimental/Chooser/Source/Chooser/Private/ObjectColumn.cpp b/Engine/Plugins/Experimental/Chooser/Source/Chooser/Private/ObjectColumn.cpp
--- a/Engine/Plugins/Experimental/Chooser/Source/Chooser/Private/ObjectColumn.cpp
+++ b/Engine/Plugins/Experimental/Chooser/Source/Chooser/Private/ObjectColumn.cpp
@@ -15,21 +15,18 @@ bool FObjectContextProperty::GetValue(const UObject* ContextObject, FSoftObjectP
 	{
 		if (const FObjectPropertyBase* ObjectProperty = FindFProperty<FObjectPropertyBase>(StructType, Binding.PropertyBindingChain.Last()))
 		{
-			// If the property has a value, then create a soft object path from it;
-			// otherwise, it could be a weak pointer to something that is not loaded.
-			const UObject* LoadedObject = ObjectProperty->GetObjectPropertyValue_InContainer(Container);
-			if (LoadedObject != nullptr)
-			{
-				OutResult = LoadedObject;
-				return true;
-			}
-
+			// if the property is a soft object property, get the path directly
 			if (ObjectProperty->IsA<FSoftObjectProperty>())
 			{
 				const FSoftObjectPtr& SoftObjectPtr = *ObjectProperty->ContainerPtrToValuePtr<FSoftObjectPtr>(Container);
 				OutResult = SoftObjectPtr.ToSoftObjectPath();
 				return true;
 			}
+			
+			// otherwise get the value from the object property and convert to a soft object path
+			const UObject* LoadedObject = ObjectProperty->GetObjectPropertyValue_InContainer(Container);
+			OutResult = LoadedObject;
+			return true;
 		}
 	}
 
@@ -79,13 +76,21 @@ bool FChooserObjectRowData::Evaluate(const FSoftObjectPath& LeftHandSide) const
 	}
 }
 
-void FObjectColumn::Filter(const UObject* ContextObject, const TArray<uint32>& IndexListIn, TArray<uint32>& IndexListOut) const
+void FObjectColumn::Filter(FChooserDebuggingInfo& DebugInfo, const UObject* ContextObject, const TArray<uint32>& IndexListIn, TArray<uint32>& IndexListOut) const
 {
 	FSoftObjectPath Result;
 	if (ContextObject != nullptr &&
 		InputValue.IsValid() &&
 		InputValue.Get<FChooserParameterObjectBase>().GetValue(ContextObject, Result))
 	{
+		
+#if WITH_EDITOR
+		if (DebugInfo.bCurrentDebugTarget)
+		{
+			TestValue = Result;
+		}
+#endif
+		
 		for (const uint32 Index : IndexListIn)
 		{
 			if (RowValues.IsValidIndex(Index))
diff --git a/Engine/Plugins/Experimental/Chooser/Source/Chooser/Private/OutputBoolColumn.cpp b/Engine/Plugins/Experimental/Chooser/Source/Chooser/Private/OutputBoolColumn.cpp
--- a/Engine/Plugins/Experimental/Chooser/Source/Chooser/Private/OutputBoolColumn.cpp
+++ b/Engine/Plugins/Experimental/Chooser/Source/Chooser/Private/OutputBoolColumn.cpp
@@ -8,7 +8,17 @@ FOutputBoolColumn::FOutputBoolColumn()
 	InputValue.InitializeAs(FBoolContextProperty::StaticStruct());
 }
 
-void FOutputBoolColumn::SetOutputs(UObject* ContextObject, int RowIndex) const
+void FOutputBoolColumn::SetOutputs(FChooserDebuggingInfo& DebugInfo, UObject* ContextObject, int RowIndex) const
 {
-	InputValue.Get<FChooserParameterBoolBase>().SetValue(ContextObject, RowValues[RowIndex]);
+	if (InputValue.IsValid())
+	{
+		InputValue.Get<FChooserParameterBoolBase>().SetValue(ContextObject, RowValues[RowIndex]);
+	}
+	
+#if WITH_EDITOR
+	if (DebugInfo.bCurrentDebugTarget)
+	{
+		TestValue = RowValues[RowIndex];
+	}
+#endif
 }
diff --git a/Engine/Plugins/Experimental/Chooser/Source/Chooser/Private/RandomizeColumn.cpp b/Engine/Plugins/Experimental/Chooser/Source/Chooser/Private/RandomizeColumn.cpp
--- a/Engine/Plugins/Experimental/Chooser/Source/Chooser/Private/RandomizeColumn.cpp
+++ b/Engine/Plugins/Experimental/Chooser/Source/Chooser/Private/RandomizeColumn.cpp
@@ -25,7 +25,7 @@ FRandomizeColumn::FRandomizeColumn()
 	InputValue.InitializeAs(FRandomizeContextProperty::StaticStruct());
 }
 
-void FRandomizeColumn::Filter(const UObject* ContextObject, const TArray<uint32>& IndexListIn, TArray<uint32>& IndexListOut) const
+void FRandomizeColumn::Filter(FChooserDebuggingInfo& DebuggingInfo, const UObject* ContextObject, const TArray<uint32>& IndexListIn, TArray<uint32>& IndexListOut) const
 {
 	int Count = IndexListIn.Num();
 	int Selection = 0;
@@ -99,7 +99,7 @@ void FRandomizeColumn::Filter(const UObject* ContextObject, const TArray<uint32>
 	}
 }
 
-void FRandomizeColumn::SetOutputs(UObject* ContextObject, int RowIndex) const
+void FRandomizeColumn::SetOutputs(FChooserDebuggingInfo& DebugInfo, UObject* ContextObject, int RowIndex) const
 {
 	if (ContextObject && InputValue.IsValid())
 	{
diff --git a/Engine/Plugins/Experimental/Chooser/Source/Chooser/Public/IChooserColumn.h b/Engine/Plugins/Experimental/Chooser/Source/Chooser/Public/IChooserColumn.h
--- a/Engine/Plugins/Experimental/Chooser/Source/Chooser/Public/IChooserColumn.h
+++ b/Engine/Plugins/Experimental/Chooser/Source/Chooser/Public/IChooserColumn.h
@@ -22,6 +22,13 @@ public:
 	virtual void ConvertToInstancedStruct(FInstancedStruct& OutInstancedStruct) const {}
 };
 
+class UChooserTable;
+
+struct FChooserDebuggingInfo
+{
+	bool bCurrentDebugTarget;
+};
+
 USTRUCT()
 struct CHOOSER_API FChooserColumnBase
 {
@@ -30,12 +37,12 @@ struct CHOOSER_API FChooserColumnBase
 public:
 	virtual ~FChooserColumnBase() {}
 	virtual void PostLoad() {};
-	virtual void Filter(const UObject* ContextObject, const TArray<uint32>& IndexListIn, TArray<uint32>& IndexListOut) const {}
+	virtual void Filter(FChooserDebuggingInfo& DebugInfo, const UObject* ContextObject, const TArray<uint32>& IndexListIn, TArray<uint32>& IndexListOut) const {}
 
 	virtual bool HasFilters() const { return true; }
 	virtual bool HasOutputs() const { return false; }
 	
-	virtual void SetOutputs(UObject* ContextObject, int RowIndex) const { }
+	virtual void SetOutputs(FChooserDebuggingInfo& DebugInfo, UObject* ContextObject, int RowIndex) const { }
 
 #if WITH_EDITOR
 	virtual FName RowValuesPropertyName() { return FName(); }
@@ -51,6 +58,8 @@ public:
 	// random columns must go last, and get a special icon
 	// using a virtual fucntion to identify them (rather than hard coding a specific type) to potentially support multiple varieties of randomization column.
 	virtual bool IsRandomizeColumn() const { return false; }
+
+	virtual bool EditorTestFilter(int32 RowIndex) const { return false; }
 #endif
 };
 
diff --git a/Engine/Plugins/Experimental/Chooser/Source/ChooserEditor/Private/BoolColumnEditor.cpp b/Engine/Plugins/Experimental/Chooser/Source/ChooserEditor/Private/BoolColumnEditor.cpp
--- a/Engine/Plugins/Experimental/Chooser/Source/ChooserEditor/Private/BoolColumnEditor.cpp
+++ b/Engine/Plugins/Experimental/Chooser/Source/ChooserEditor/Private/BoolColumnEditor.cpp
@@ -8,6 +8,7 @@
 #include "Widgets/SBoxPanel.h"
 #include "Widgets/Input/SButton.h"
 #include "Widgets/Input/SCheckBox.h"
+#include "Widgets/Images/SImage.h"
 #include "GraphEditorSettings.h"
 
 #define LOCTEXT_NAMESPACE "BoolColumnEditor"
@@ -19,6 +20,58 @@ TSharedRef<SWidget> CreateBoolColumnWidget(UChooserTable* Chooser, FChooserColum
 {
 	FBoolColumn* BoolColumn = static_cast<FBoolColumn*>(Column);
 	
+	if (Row < 0)
+	{
+		// create column header widget
+		TSharedPtr<SWidget> InputValueWidget = nullptr;
+		if (FChooserParameterBase* InputValue = Column->GetInputValue())
+		{
+			InputValueWidget = FObjectChooserWidgetFactories::CreateWidget(false, Chooser, InputValue, Column->GetInputType(), Chooser->ContextObjectType, Chooser->OutputObjectType);
+		}
+	
+		const FSlateBrush* ColumnIcon = FCoreStyle::Get().GetBrush("Icons.Filter");
+		
+		TSharedRef<SWidget> ColumnHeaderWidget = SNew(SHorizontalBox)
+			+ SHorizontalBox::Slot().AutoWidth()
+			[
+				SNew(SBorder)
+				.BorderBackgroundColor(FLinearColor(0,0,0,0))
+				.Content()
+				[
+					SNew(SImage).Image(ColumnIcon)
+				]
+			]
+			+ SHorizontalBox::Slot()
+			[
+				InputValueWidget ? InputValueWidget.ToSharedRef() : SNullWidget::NullWidget
+			];
+	
+		if (Chooser->bEnableDebugTesting)
+		{
+			ColumnHeaderWidget = SNew(SVerticalBox)
+			+ SVerticalBox::Slot()
+			[
+				ColumnHeaderWidget
+			]
+			+ SVerticalBox::Slot()
+			[
+				SNew(SHorizontalBox)
+				 + SHorizontalBox::Slot().FillWidth(1)
+				 + SHorizontalBox::Slot().AutoWidth().VAlign(VAlign_Center)
+				[
+					SNew(SCheckBox).IsEnabled_Lambda([Chooser](){ return Chooser->HasDebugTarget(); })
+					.IsChecked_Lambda([BoolColumn]() { return BoolColumn->TestValue ? ECheckBoxState::Checked : ECheckBoxState::Unchecked; })
+					.OnCheckStateChanged_Lambda([Chooser, BoolColumn](ECheckBoxState NewValue) { BoolColumn->TestValue = NewValue == ECheckBoxState::Checked; })
+				]
+				 + SHorizontalBox::Slot().FillWidth(1)
+			];
+		}
+
+		return ColumnHeaderWidget;
+	}
+
+	// create cell widget
+	
 	return
 	SNew(SHorizontalBox)
 		+ SHorizontalBox::Slot().FillWidth(1)
@@ -57,8 +110,56 @@ TSharedRef<SWidget> CreateOutputBoolColumnWidget(UChooserTable* Chooser, FChoose
 {
 	FOutputBoolColumn* BoolColumn = static_cast<FOutputBoolColumn*>(Column);
 
-	return
+	if (Row < 0)
+	{
+		// create column header widget
+		TSharedPtr<SWidget> InputValueWidget = nullptr;
+		if (FChooserParameterBase* InputValue = Column->GetInputValue())
+		{
+			InputValueWidget = FObjectChooserWidgetFactories::CreateWidget(false, Chooser, InputValue, Column->GetInputType(), Chooser->ContextObjectType, Chooser->OutputObjectType);
+		}
+		
+		const FSlateBrush* ColumnIcon = FCoreStyle::Get().GetBrush("Icons.ArrowRight");
+		
+		TSharedRef<SWidget> ColumnHeaderWidget = SNew(SHorizontalBox)
+			+ SHorizontalBox::Slot().AutoWidth()
+			[
+				SNew(SBorder)
+				.BorderBackgroundColor(FLinearColor(0,0,0,0))
+				.Content()
+				[
+					SNew(SImage).Image(ColumnIcon)
+				]
+			]
+			+ SHorizontalBox::Slot()
+			[
+				InputValueWidget ? InputValueWidget.ToSharedRef() : SNullWidget::NullWidget
+			];
 	
+		if (Chooser->bEnableDebugTesting)
+		{
+			ColumnHeaderWidget = SNew(SVerticalBox)
+			+ SVerticalBox::Slot()
+			[
+				ColumnHeaderWidget
+			]
+			+ SVerticalBox::Slot()
+			[
+				SNew(SHorizontalBox)
+             	+ SHorizontalBox::Slot().FillWidth(1)
+             	+ SHorizontalBox::Slot().AutoWidth().VAlign(VAlign_Center)
+             	[
+				SNew(SCheckBox).IsEnabled(false)
+				.IsChecked_Lambda([BoolColumn]() { return BoolColumn->TestValue ? ECheckBoxState::Checked : ECheckBoxState::Unchecked; })
+				]
+             	+ SHorizontalBox::Slot().FillWidth(1)
+			];
+		}
+
+		return ColumnHeaderWidget;
+	}
+
+	return
 	SNew(SHorizontalBox)
 	+ SHorizontalBox::Slot().FillWidth(1)
 	+ SHorizontalBox::Slot().AutoWidth().VAlign(VAlign_Center)
diff --git a/Engine/Plugins/Experimental/Chooser/Source/ChooserEditor/Private/ChooserTableEditor.cpp b/Engine/Plugins/Experimental/Chooser/Source/ChooserEditor/Private/ChooserTableEditor.cpp
--- a/Engine/Plugins/Experimental/Chooser/Source/ChooserEditor/Private/ChooserTableEditor.cpp
+++ b/Engine/Plugins/Experimental/Chooser/Source/ChooserEditor/Private/ChooserTableEditor.cpp
@@ -75,10 +75,104 @@ const FName FChooserTableEditor::ChooserEditorAppIdentifier( TEXT( "ChooserEdito
 
 FChooserTableEditor::~FChooserTableEditor()
 {
+	if (SelectedColumn)
+	{
+		SelectedColumn->ClearFlags(RF_Standalone);
+	}
+	for (UObject* SelectedRow : SelectedRows)
+	{
+		SelectedRow->ClearFlags(RF_Standalone);
+	}
+	
 	FCoreUObjectDelegates::OnObjectsReplaced.RemoveAll(this);
 
 	DetailsView.Reset();
 }
+	
+void FChooserTableEditor::MakeDebugTargetMenu(UToolMenu* InToolMenu) 
+{
+	static FName SectionName = "Select Debug Target";
+		
+	InToolMenu->AddMenuEntry(
+			SectionName,
+			FToolMenuEntry::InitMenuEntry(
+				"None",
+				LOCTEXT("None", "None"),
+				LOCTEXT("None Tooltip", "Clear selected debug target"),
+				FSlateIcon(),
+				FUIAction(
+					FExecuteAction::CreateLambda([this]()
+					{
+						UChooserTable* Chooser = GetChooser();
+						Chooser->ResetDebugTarget();
+						if (Chooser->bEnableDebugTesting)
+						{
+							Chooser->bEnableDebugTesting = false;
+							Chooser->bDebugTestValuesValid = false;
+							UpdateTableColumns();
+						}
+					}),
+					FCanExecuteAction()
+				)
+			));
+	
+	InToolMenu->AddMenuEntry(
+			SectionName,
+			FToolMenuEntry::InitMenuEntry(
+				"Manual",
+				LOCTEXT("Manual Testing", "Manual Testing"),
+				LOCTEXT("Manual Tooltip", "Test the chooser by manually entering values for each column"),
+				FSlateIcon(),
+				FUIAction(
+					FExecuteAction::CreateLambda([this]()
+					{
+						UChooserTable* Chooser = GetChooser();
+						Chooser->ResetDebugTarget();
+						if (!Chooser->bEnableDebugTesting)
+						{
+							Chooser->bEnableDebugTesting = true;
+							Chooser->bDebugTestValuesValid = true;
+							UpdateTableColumns();
+						}
+					}),
+					FCanExecuteAction()
+				)
+			));
+
+	const UChooserTable* Chooser = GetChooser();
+
+	Chooser->IterateRecentContextObjects([this, InToolMenu](const UObject* Object)
+		{
+			InToolMenu->AddMenuEntry(
+						SectionName,
+						FToolMenuEntry::InitMenuEntry(
+							Object->GetFName(),
+							FText::FromString(Object->GetName()),
+							LOCTEXT("Select Object ToolTip", "Select this object as the debug target"),
+							FSlateIcon(),
+							FUIAction(
+								FExecuteAction::CreateLambda([this, ObjectPtr = MakeWeakObjectPtr(Object)]()
+								{
+									if(ObjectPtr.IsValid())
+									{
+										UChooserTable* Chooser = GetChooser();
+										Chooser->SetDebugTarget(ObjectPtr);
+										Chooser->bDebugTestValuesValid = false;
+										if (!Chooser->bEnableDebugTesting)
+										{
+											Chooser->bEnableDebugTesting = true;
+											UpdateTableColumns();
+										}
+									}
+								}),
+								FCanExecuteAction()
+							)
+						));
+		}
+	);
+	
+
+}
 
 void FChooserTableEditor::RegisterToolbar()
 {
@@ -104,10 +198,48 @@ void FChooserTableEditor::RegisterToolbar()
 			TAttribute<FText>(),
 			TAttribute<FText>(),
 			FSlateIcon("EditorStyle", "FullBlueprintEditor.EditGlobalOptions")));
+
+
+		Section.AddDynamicEntry("DebuggingCommands", FNewToolMenuSectionDelegate::CreateLambda([](FToolMenuSection& InSection)
+		{
+			UChooserEditorToolMenuContext* Context = InSection.FindContext<UChooserEditorToolMenuContext>();
+
+			if (Context)
+			{
+				if (TSharedPtr<FChooserTableEditor> ChooserEditor = Context->ChooserEditor.Pin())
+				{
+					InSection.AddEntry(FToolMenuEntry::InitComboButton( "SelectDebugTarget",
+						FToolUIActionChoice(),
+					  FNewToolMenuDelegate::CreateSP(ChooserEditor.Get(), &FChooserTableEditor::MakeDebugTargetMenu),
+						TAttribute<FText>::CreateLambda([Chooser = ChooserEditor->GetChooser()]
+						{
+							if (Chooser->HasDebugTarget())
+							{
+								return  FText::FromString(Chooser->GetDebugTarget()->GetName());
+							}
+							else
+							{
+								return Chooser->bEnableDebugTesting ? LOCTEXT("Manual Testing", "Manual Testing") : LOCTEXT("Debug Target", "Debug Target");
+							}
+						}),
+						LOCTEXT("Debug Target Tooltip", "Select an object that has recently been the context object for this chooser to visualize the selection results")));
+				}
+			}
+		}));
+
 	}
 
 }
 
+void FChooserTableEditor::InitToolMenuContext(FToolMenuContext& MenuContext)
+{
+	FAssetEditorToolkit::InitToolMenuContext(MenuContext);
+
+	UChooserEditorToolMenuContext* Context = NewObject<UChooserEditorToolMenuContext>();
+	Context->ChooserEditor = SharedThis(this);
+	MenuContext.AddObject(Context);
+}
+
 void FChooserTableEditor::BindCommands()
 {
 	const FChooserTableEditorCommands& Commands = FChooserTableEditorCommands::Get();
@@ -369,10 +501,26 @@ public:
 
 		ChildSlot
 		[
-			SNew(SBox) .Padding(0.0f) .HAlign(HAlign_Center) .VAlign(VAlign_Center) .WidthOverride(16.0f)
+			SNew(SOverlay)
+			+ SOverlay::Slot()
 			[
-				SNew(SImage)
-				.Image(FCoreStyle::Get().GetBrush("VerticalBoxDragIndicatorShort"))
+				SNew(SBox) .Padding(0.0f) .HAlign(HAlign_Center) .VAlign(VAlign_Center) .WidthOverride(16.0f)
+				[
+					SNew(SImage)
+					.Image(FCoreStyle::Get().GetBrush("VerticalBoxDragIndicatorShort"))
+				]
+			]
+			+ SOverlay::Slot()
+			[
+				SNew(SBox).Padding(0.0f) .HAlign(HAlign_Center) .VAlign(VAlign_Center) .WidthOverride(16.0f)
+				[
+					SNew(SImage)
+					.Visibility_Lambda([this]()
+					{
+						return ChooserEditor->GetChooser()->bDebugTestValuesValid && RowIndex == ChooserEditor->GetChooser()->GetDebugSelectedRow() ? EVisibility::HitTestInvisible : EVisibility::Hidden;
+					})
+					.Image(FAppStyle::Get().GetBrush("Icons.ArrowRight"))
+				]
 			]
 		];
 	}
@@ -472,7 +620,30 @@ public:
 				
 					if (ColumnWidget.IsValid())
 					{
-						return ColumnWidget.ToSharedRef();
+						return SNew(SOverlay)
+						+ SOverlay::Slot()
+						[
+							ColumnWidget.ToSharedRef()
+						]
+						+ SOverlay::Slot()
+						[
+							SNew(SColorBlock).Visibility(EVisibility::HitTestInvisible).Color_Lambda(
+									[this,Column]()
+									{
+										if (Chooser->bDebugTestValuesValid && Column->HasFilters())
+										{
+											if (Column->EditorTestFilter(RowIndex->RowIndex))
+											{
+												return FLinearColor(0.0,1.0,0.0,0.30);
+											}
+											else
+											{
+												return FLinearColor(1.0,0.0,0.0,0.20);
+											}
+										}
+										return FLinearColor::Transparent;
+									})
+						];
 					}
 				}
 			}
@@ -606,28 +777,13 @@ void FChooserTableEditor::UpdateTableColumns()
 	for(int ColumnIndex = 0; ColumnIndex < NumColumns; ColumnIndex++)
 	{
 		FChooserColumnBase& Column = Chooser->ColumnsStructs[ColumnIndex].GetMutable<FChooserColumnBase>();
-		TSharedPtr<SWidget> HeaderWidget = nullptr;
-		if (FChooserParameterBase* InputValue = Column.GetInputValue())
-		{
-			HeaderWidget = FObjectChooserWidgetFactories::CreateWidget(false, Chooser, InputValue, Column.GetInputType(), Chooser->ContextObjectType, Chooser->OutputObjectType);
-		}
-
 
-		const ISlateStyle& CoreStyle = FCoreStyle::Get();
-		const FSlateBrush* ColumnIcon = nullptr;
-		if (Column.IsRandomizeColumn())
-		{
-			ColumnIcon = FAppStyle::Get().GetBrush("Icons.Help");
-		}
-		else if (Column.HasOutputs())
+		TSharedPtr<SWidget> HeaderWidget = FObjectChooserWidgetFactories::CreateColumnWidget(&Column, Chooser->ColumnsStructs[ColumnIndex].GetScriptStruct(), Chooser, -1);
+		if (!HeaderWidget.IsValid())
 		{
-			ColumnIcon = CoreStyle.GetBrush("Icons.ArrowRight");
+			HeaderWidget = SNullWidget::NullWidget;
 		}
-		else
-		{
-			ColumnIcon = CoreStyle.GetBrush("Icons.Filter");
-		}
-
+		
 		HeaderRow->AddColumn(SHeaderRow::FColumn::FArguments()
 			.ColumnId(ColumnId)
 			.ManualWidth(200)
@@ -676,7 +832,6 @@ void FChooserTableEditor::UpdateTableColumns()
 							));
 				}
 
-				// if (ColumnIndex < Chooser->ColumnsStructs.Num() - 1)
 
 				MenuBuilder.AddMenuEntry(LOCTEXT("Delete Column","Delete"),LOCTEXT("Delete Column ToolTip", "Remove this column and all its data from the table"),FSlateIcon(),
 					FUIAction(
@@ -712,8 +867,8 @@ void FChooserTableEditor::UpdateTableColumns()
 					}));
 			
 				return MenuBuilder.MakeWidget();
-			
 			})
+			.HeaderComboVisibility(EHeaderComboVisibility::Ghosted)
 			.HeaderContent()
 			[
 				SNew(SBorder)
@@ -733,20 +888,7 @@ void FChooserTableEditor::UpdateTableColumns()
 					return FReply::Handled();
 				})
 				[
-					SNew(SHorizontalBox)
-					+ SHorizontalBox::Slot().AutoWidth()
-					[
-						SNew(SBorder)
-						.BorderBackgroundColor(FLinearColor(0,0,0,0))
-						.Content()
-						[
-							SNew(SImage).Image(ColumnIcon)
-						]
-					]
-					+ SHorizontalBox::Slot()
-					[
-						HeaderWidget ? HeaderWidget.ToSharedRef() : SNullWidget::NullWidget
-					]
+					HeaderWidget.ToSharedRef()
 				]
 			
 			]);
@@ -915,16 +1057,25 @@ TSharedRef<SDockTab> FChooserTableEditor::SpawnTableTab( const FSpawnTabArgs& Ar
 				)
 				.OnSelectionChanged_Lambda([this](TSharedPtr<FChooserTableRow> SelectedItem,  ESelectInfo::Type SelectInfo)
 				{
-					SelectedColumn = nullptr;
+					if (SelectedColumn)
+					{
+						// deselect any selected column
+						SelectedColumn->Column = -1;
+					}
+					
 					if (SelectedItem)
 					{
+						for (UObject* SelectedRow : SelectedRows)
+						{
+							SelectedRow->ClearFlags(RF_Standalone);
+						}
 						SelectedRows.SetNum(0);
 						UChooserTable* Chooser = Cast<UChooserTable>(EditingObjects[0]);
 						// Get the list of objects to edit the details of
 						TObjectPtr<UChooserRowDetails> Selection = NewObject<UChooserRowDetails>();
 						Selection->Chooser = Chooser;
 						Selection->Row = SelectedItem->RowIndex;
-						Selection->SetFlags(RF_Transactional);
+						Selection->SetFlags(RF_Standalone);
 						SelectedRows.Add(Selection);
 											
 						TArray<UObject*> DetailsObjects;
@@ -1039,8 +1190,10 @@ void FChooserTableEditor::SelectColumn(int Index)
 		if (SelectedColumn == nullptr)
 		{
 			SelectedColumn = NewObject<UChooserColumnDetails>();
+			SelectedColumn->SetFlags(RF_Standalone);
 			SelectedColumn->Chooser = Chooser;
 		}
+
 		SelectedColumn->Column = Index;
 		DetailsView->SetObject(SelectedColumn, true);
 	}
diff --git a/Engine/Plugins/Experimental/Chooser/Source/ChooserEditor/Private/ChooserTableEditor.h b/Engine/Plugins/Experimental/Chooser/Source/ChooserEditor/Private/ChooserTableEditor.h
--- a/Engine/Plugins/Experimental/Chooser/Source/ChooserEditor/Private/ChooserTableEditor.h
+++ b/Engine/Plugins/Experimental/Chooser/Source/ChooserEditor/Private/ChooserTableEditor.h
@@ -39,6 +39,7 @@ public:
 	int Column;
 };
 
+
 namespace UE::ChooserEditor
 {
 	class FChooserTableEditor : public FAssetEditorToolkit, public FSelfRegisteringEditorUndoClient, public FNotifyHook
@@ -73,6 +74,7 @@ namespace UE::ChooserEditor
 		virtual FLinearColor GetWorldCentricTabColorScale() const override;
 		virtual bool IsPrimaryEditor() const override { return true; }
 		virtual bool IsSimpleAssetEditor() const override { return false; }
+		virtual void InitToolMenuContext(FToolMenuContext& MenuContext) override;
 		
 		/** FEditorUndoClient Interface */
 		virtual void PostUndo(bool bSuccess) override;
@@ -83,6 +85,7 @@ namespace UE::ChooserEditor
 		virtual void NotifyPostChange( const FPropertyChangedEvent& PropertyChangedEvent, FProperty* PropertyThatChanged) override;
 
 		UChooserTable* GetChooser() { return Cast<UChooserTable>(EditingObjects[0]); }
+		const UChooserTable* GetChooser() const { return Cast<UChooserTable>(EditingObjects[0]); }
 	
 		/** Used to show or hide certain properties */
 		void SetPropertyVisibilityDelegate(FIsPropertyVisible InVisibilityDelegate);
@@ -105,6 +108,7 @@ namespace UE::ChooserEditor
 		void SelectRootProperties();
 		void RegisterToolbar();
 		void BindCommands();
+		void MakeDebugTargetMenu(UToolMenu* InToolMenu);
 	
 		/** Create the properties tab and its content */
 		TSharedRef<SDockTab> SpawnPropertiesTab( const FSpawnTabArgs& Args );
@@ -129,7 +133,7 @@ namespace UE::ChooserEditor
 		/** The objects open within this editor */
 		TArray<UObject*> EditingObjects;
 
-		TObjectPtr<UChooserColumnDetails> SelectedColumn;
+		UChooserColumnDetails* SelectedColumn = nullptr;
 		TArray<TObjectPtr<UChooserRowDetails>> SelectedRows;
 
 		void UpdateTableColumns();
@@ -140,7 +144,6 @@ namespace UE::ChooserEditor
 		
 		TSharedPtr<SHeaderRow> HeaderRow;
 		TSharedPtr<SListView<TSharedPtr<FChooserTableRow>>> TableView;
-		
 	public:
 		TSharedPtr<SComboButton>& GetCreateRowComboButton() { return CreateRowComboButton; };
 
@@ -154,3 +157,12 @@ namespace UE::ChooserEditor
 		static void RegisterWidgets();
 	};
 }
+
+UCLASS()
+class UChooserEditorToolMenuContext : public UObject
+{
+	GENERATED_BODY()
+
+public:
+	TWeakPtr<UE::ChooserEditor::FChooserTableEditor> ChooserEditor;
+};
diff --git a/Engine/Plugins/Experimental/Chooser/Source/ChooserEditor/Private/EnumColumnEditor.cpp b/Engine/Plugins/Experimental/Chooser/Source/ChooserEditor/Private/EnumColumnEditor.cpp
--- a/Engine/Plugins/Experimental/Chooser/Source/ChooserEditor/Private/EnumColumnEditor.cpp
+++ b/Engine/Plugins/Experimental/Chooser/Source/ChooserEditor/Private/EnumColumnEditor.cpp
@@ -38,10 +38,26 @@ public:
 				if (const UEnum* Enum = EnumColumnPointer->InputValue.Get<FChooserParameterEnumBase>().GetEnum())
 				{
 					return SNew(SEnumComboBox, Enum)
+						.IsEnabled_Lambda([this]
+						{
+							if (UChooserTable* Chooser = Cast<UChooserTable>(TransactionObject))
+							{
+								// return false only for the column header (RowIndex = -1) and when there is a debug target object bound
+								return RowIndex.Get()>=0 || !Chooser->HasDebugTarget();
+							}
+							return true;
+						})
 						.CurrentValue_Lambda([this]()
 						{
 							int Row = RowIndex.Get();
-							return EnumColumn->RowValues.IsValidIndex(Row) ? static_cast<int32>(EnumColumn->RowValues[Row].Value) : 0;
+							if (EnumColumn->RowValues.IsValidIndex(Row))
+							{
+								return EnumColumn->RowValues.IsValidIndex(Row) ? static_cast<int32>(EnumColumn->RowValues[Row].Value) : 0;
+							}
+							else
+							{
+								return EnumColumn->TestValue;
+							}
 						})
 						.OnEnumSelectionChanged_Lambda([this](int32 EnumValue, ESelectInfo::Type)
 						{
@@ -52,6 +68,10 @@ public:
 								TransactionObject->Modify(true);
 								EnumColumn->RowValues[Row].Value = static_cast<uint8>(EnumValue);
 							}
+							else
+							{
+								EnumColumn->TestValue = EnumValue;
+							}
 						});
 				}
 			}
@@ -65,8 +85,24 @@ public:
 		ChildSlot[ CreateEnumComboBox()	];
 	}
 
+	virtual void Tick(const FGeometry& AllottedGeometry, const double InCurrentTime, const float InDeltaTime) override
+	{
+		const UEnum* CurrentEnumSource = nullptr;
+		if (EnumColumn->InputValue.IsValid())
+		{
+			CurrentEnumSource = EnumColumn->InputValue.Get<FChooserParameterEnumBase>().GetEnum(); 
+		}
+		if (EnumSource != CurrentEnumSource)
+		{
+			EnumComboBorder->SetContent(CreateEnumComboBox());
+			EnumSource = CurrentEnumSource;
+		}
+	}
+    					
+
 	void Construct( const FArguments& InArgs)
 	{
+		SetCanTick(true);
 		RowIndex = InArgs._RowIndex;
 		EnumColumn = InArgs._EnumColumn;
 		TransactionObject = InArgs._TransactionObject;
@@ -88,24 +124,13 @@ public:
 			SNew(SHorizontalBox)
 			+ SHorizontalBox::Slot().AutoWidth()
 			[
-				SNew(SBox).WidthOverride(45)
+				SNew(SBox).WidthOverride(Row < 0 ? 0 : 45)
 				[
 					SNew(SButton).ButtonStyle(FAppStyle::Get(),"FlatButton").TextStyle(FAppStyle::Get(),"RichTextBlock.Bold").HAlign(HAlign_Center)
+					.Visibility(Row < 0 ? EVisibility::Hidden : EVisibility::Visible)
 					.Text_Lambda([this, Row]()
 					{
-						// this is a bit of an odd thing to do, but as we are polling for the "=" / "!=" button state anyway,
-						// also poll to see if the currently bound Enum has changed, and if it has, recreate the enum combo box.
-						const UEnum* CurrentEnumSource = nullptr;
-						if (EnumColumn->InputValue.IsValid())
-						{
-							CurrentEnumSource = EnumColumn->InputValue.Get<FChooserParameterEnumBase>().GetEnum(); 
-						}
-						if (EnumSource != CurrentEnumSource)
-						{
-							EnumComboBorder->SetContent(CreateEnumComboBox());
-							EnumSource = EnumColumn->InputValue.Get<FChooserParameterEnumBase>().GetEnum();
-						}
-					
+
 						return (EnumColumn->RowValues.IsValidIndex(Row) && EnumColumn->RowValues[Row].CompareNotEqual ? LOCTEXT("Not Equal", "!=") : LOCTEXT("Equal", "="));
 					})
 					.OnClicked_Lambda([this, Row]()
@@ -148,24 +173,68 @@ TSharedRef<SWidget> CreateEnumColumnWidget(UChooserTable* Chooser, FChooserColum
 {
 	FEnumColumn* EnumColumn = static_cast<FEnumColumn*>(Column);
 	
+	if (Row < 0)
+	{
+		// create column header widget
+		TSharedPtr<SWidget> InputValueWidget = nullptr;
+		if (FChooserParameterBase* InputValue = Column->GetInputValue())
+		{
+			InputValueWidget = FObjectChooserWidgetFactories::CreateWidget(false, Chooser, InputValue, Column->GetInputType(), Chooser->ContextObjectType, Chooser->OutputObjectType);
+		}
+		
+		const FSlateBrush* ColumnIcon = FCoreStyle::Get().GetBrush("Icons.Filter");
+		
+		TSharedRef<SWidget> ColumnHeaderWidget = SNew(SHorizontalBox)
+			+ SHorizontalBox::Slot().AutoWidth()
+			[
+				SNew(SBorder)
+				.BorderBackgroundColor(FLinearColor(0,0,0,0))
+				.Content()
+				[
+					SNew(SImage).Image(ColumnIcon)
+				]
+			]
+			+ SHorizontalBox::Slot()
+			[
+				InputValueWidget ? InputValueWidget.ToSharedRef() : SNullWidget::NullWidget
+			];
+	
+		if (Chooser->bEnableDebugTesting)
+		{
+			ColumnHeaderWidget = SNew(SVerticalBox)
+			+ SVerticalBox::Slot()
+			[
+				ColumnHeaderWidget
+			]
+			+ SVerticalBox::Slot()
+			[
+				SNew(SEnumCell).TransactionObject(Chooser).EnumColumn(EnumColumn).RowIndex(Row)
+			];
+		}
+
+		return ColumnHeaderWidget;
+	}
+
+	// create cell widget
+	
 	return SNew(SEnumCell).TransactionObject(Chooser).EnumColumn(EnumColumn).RowIndex(Row);
 }
 
 TSharedRef<SWidget> CreateEnumPropertyWidget(bool bReadOnly, UObject* TransactionObject, void* Value, UClass* ContextClass, UClass* ResultBaseClass)
 {
-		IHasContextClass* HasContextClass = Cast<IHasContextClass>(TransactionObject);
-    
-    	FEnumContextProperty* ContextProperty = reinterpret_cast<FEnumContextProperty*>(Value);
-    
-    	return SNew(SPropertyAccessChainWidget).ContextClassOwner(HasContextClass).AllowFunctions(false).BindingColor("BytePinTypeColor").TypeFilter("enum")
-    	.PropertyBindingValue(&ContextProperty->Binding)
-    	.OnAddBinding_Lambda(
-    		[ContextProperty, TransactionObject](FName InPropertyName, const TArray<FBindingChainElement>& InBindingChain)
-    		{
-    			const FScopedTransaction Transaction(NSLOCTEXT("ContextPropertyWidget", "Change Property Binding", "Change Property Binding"));
-    			TransactionObject->Modify(true);
-    			ContextProperty->SetBinding(InBindingChain);	
-    		});
+	IHasContextClass* HasContextClass = Cast<IHasContextClass>(TransactionObject);
+
+	FEnumContextProperty* ContextProperty = reinterpret_cast<FEnumContextProperty*>(Value);
+
+	return SNew(SPropertyAccessChainWidget).ContextClassOwner(HasContextClass).AllowFunctions(false).BindingColor("BytePinTypeColor").TypeFilter("enum")
+	.PropertyBindingValue(&ContextProperty->Binding)
+	.OnAddBinding_Lambda(
+		[ContextProperty, TransactionObject](FName InPropertyName, const TArray<FBindingChainElement>& InBindingChain)
+		{
+			const FScopedTransaction Transaction(NSLOCTEXT("ContextPropertyWidget", "Change Property Binding", "Change Property Binding"));
+			TransactionObject->Modify(true);
+			ContextProperty->SetBinding(InBindingChain);	
+		});
 }
 	
 void RegisterEnumWidgets()
diff --git a/Engine/Plugins/Experimental/Chooser/Source/ChooserEditor/Private/FloatRangeColumnEditor.cpp b/Engine/Plugins/Experimental/Chooser/Source/ChooserEditor/Private/FloatRangeColumnEditor.cpp
--- a/Engine/Plugins/Experimental/Chooser/Source/ChooserEditor/Private/FloatRangeColumnEditor.cpp
+++ b/Engine/Plugins/Experimental/Chooser/Source/ChooserEditor/Private/FloatRangeColumnEditor.cpp
@@ -5,6 +5,7 @@
 #include "ObjectChooserWidgetFactories.h"
 #include "ChooserTableEditor.h"
 #include "Widgets/Input/SNumericEntryBox.h"
+#include "Widgets/Images/SImage.h"
 #include "GraphEditorSettings.h"
 #include "SPropertyAccessChainWidget.h"
 
@@ -33,6 +34,51 @@ TSharedRef<SWidget> CreateFloatRangeColumnWidget(UChooserTable* Chooser, FChoose
 {
 	FFloatRangeColumn* FloatRangeColumn = static_cast<FFloatRangeColumn*>(Column);
 
+	if (Row < 0)
+	{
+		// create column header widget
+		TSharedPtr<SWidget> InputValueWidget = nullptr;
+		if (FChooserParameterBase* InputValue = Column->GetInputValue())
+		{
+			InputValueWidget = FObjectChooserWidgetFactories::CreateWidget(false, Chooser, InputValue, Column->GetInputType(), Chooser->ContextObjectType, Chooser->OutputObjectType);
+		}
+		
+		const FSlateBrush* ColumnIcon = FCoreStyle::Get().GetBrush("Icons.Filter");
+		
+		TSharedRef<SWidget> ColumnHeaderWidget = SNew(SHorizontalBox)
+			+ SHorizontalBox::Slot().AutoWidth()
+			[
+				SNew(SBorder)
+				.BorderBackgroundColor(FLinearColor(0,0,0,0))
+				.Content()
+				[
+					SNew(SImage).Image(ColumnIcon)
+				]
+			]
+			+ SHorizontalBox::Slot()
+			[
+				InputValueWidget ? InputValueWidget.ToSharedRef() : SNullWidget::NullWidget
+			];
+	
+		if (Chooser->bEnableDebugTesting)
+		{
+			ColumnHeaderWidget = SNew(SVerticalBox)
+			+ SVerticalBox::Slot()
+			[
+				ColumnHeaderWidget
+			]
+			+ SVerticalBox::Slot()
+			[
+				SNew(SNumericEntryBox<float>).IsEnabled_Lambda([Chooser](){ return !Chooser->HasDebugTarget(); })
+				.Value_Lambda([FloatRangeColumn]() { return FloatRangeColumn->TestValue; })
+				.OnValueCommitted_Lambda([Chooser, FloatRangeColumn](float NewValue, ETextCommit::Type CommitType) { FloatRangeColumn->TestValue = NewValue; })
+			];
+		}
+
+		return ColumnHeaderWidget;
+	}
+
+	// create cell widget
 	return SNew(SHorizontalBox)
 	+ SHorizontalBox::Slot().AutoWidth().VAlign(VAlign_Center)
 	[
diff --git a/Engine/Plugins/Experimental/Chooser/Source/ChooserEditor/Private/GameplayTagColumnEditor.cpp b/Engine/Plugins/Experimental/Chooser/Source/ChooserEditor/Private/GameplayTagColumnEditor.cpp
--- a/Engine/Plugins/Experimental/Chooser/Source/ChooserEditor/Private/GameplayTagColumnEditor.cpp
+++ b/Engine/Plugins/Experimental/Chooser/Source/ChooserEditor/Private/GameplayTagColumnEditor.cpp
@@ -7,6 +7,7 @@
 #include "SGameplayTagWidget.h"
 #include "SSimpleComboButton.h"
 #include "GraphEditorSettings.h"
+#include "Widgets/Images/SImage.h"
 
 #define LOCTEXT_NAMESPACE "FGameplayTagColumnEditor"
 
@@ -16,7 +17,69 @@ namespace UE::ChooserEditor
 TSharedRef<SWidget> CreateGameplayTagColumnWidget(UChooserTable* Chooser, FChooserColumnBase* Column, int Row)
 {
 	FGameplayTagColumn* GameplayTagColumn = static_cast<struct FGameplayTagColumn*>(Column);
+	
+	if (Row < 0)
+	{
+		// create column header widget
+		TSharedPtr<SWidget> InputValueWidget = nullptr;
+		if (FChooserParameterBase* InputValue = Column->GetInputValue())
+		{
+			InputValueWidget = FObjectChooserWidgetFactories::CreateWidget(false, Chooser, InputValue, Column->GetInputType(), Chooser->ContextObjectType, Chooser->OutputObjectType);
+		}
+		
+		const FSlateBrush* ColumnIcon = FCoreStyle::Get().GetBrush("Icons.Filter");
+		
+		TSharedRef<SWidget> ColumnHeaderWidget = SNew(SHorizontalBox)
+			+ SHorizontalBox::Slot().AutoWidth()
+			[
+				SNew(SBorder)
+				.BorderBackgroundColor(FLinearColor(0,0,0,0))
+				.Content()
+				[
+					SNew(SImage).Image(ColumnIcon)
+				]
+			]
+			+ SHorizontalBox::Slot()
+			[
+				InputValueWidget ? InputValueWidget.ToSharedRef() : SNullWidget::NullWidget
+			];
+	
+		if (Chooser->bEnableDebugTesting)
+		{
+			ColumnHeaderWidget = SNew(SVerticalBox)
+			+ SVerticalBox::Slot()
+			[
+				ColumnHeaderWidget
+			]
+			+ SVerticalBox::Slot()
+			[
+				SNew(SSimpleComboButton)
+					.IsEnabled_Lambda([Chooser]()
+					{
+						 return !Chooser->HasDebugTarget();
+					})
+					.Text_Lambda([GameplayTagColumn]()
+					{
+						FText Text = FText::FromString(GameplayTagColumn->TestValue.ToStringSimple(false));
+						if (Text.IsEmpty())
+						{
+							Text = LOCTEXT("None", "None");
+						}
+						return Text;
+					})	
+					.OnGetMenuContent_Lambda([Chooser, GameplayTagColumn, Row]()
+					{
+						TArray<SGameplayTagWidget::FEditableGameplayTagContainerDatum> EditableContainers;
+						EditableContainers.Emplace(Chooser, &(GameplayTagColumn->TestValue));
+						return TSharedRef<SWidget>(SNew(SGameplayTagWidget, EditableContainers));
+					})
+			];
+		}
+
+		return ColumnHeaderWidget;
+	}
 
+	// create cell widget
 	return SNew(SSimpleComboButton)
 		.Text_Lambda([GameplayTagColumn, Row]()
 		{
diff --git a/Engine/Plugins/Experimental/Chooser/Source/ChooserEditor/Private/ObjectColumnEditor.cpp b/Engine/Plugins/Experimental/Chooser/Source/ChooserEditor/Private/ObjectColumnEditor.cpp
--- a/Engine/Plugins/Experimental/Chooser/Source/ChooserEditor/Private/ObjectColumnEditor.cpp
+++ b/Engine/Plugins/Experimental/Chooser/Source/ChooserEditor/Private/ObjectColumnEditor.cpp
@@ -7,6 +7,7 @@
 #include "ObjectColumn.h"
 #include "PropertyCustomizationHelpers.h"
 #include "TransactionCommon.h"
+#include "Widgets/Images/SImage.h"
 #include "Widgets/Input/SButton.h"
 #include "Widgets/Layout/SBox.h"
 
@@ -30,11 +31,11 @@ namespace UE::ChooserEditor
 		return AllowedClass;
 	}
 
-	static TSharedRef<SWidget> CreateObjectPicker(UObject* TransactionObject, FObjectColumn* ObjectColumn, int32 Row)
+	static TSharedRef<SObjectPropertyEntryBox> CreateObjectPicker(UObject* TransactionObject, FObjectColumn* ObjectColumn, int32 Row)
 	{
-		return SNew(SObjectPropertyEntryBox)
+		TSharedRef<SObjectPropertyEntryBox> ObjectPicker = SNew(SObjectPropertyEntryBox)
 			.ObjectPath_Lambda([ObjectColumn, Row]() {
-				return ObjectColumn->RowValues.IsValidIndex(Row) ? ObjectColumn->RowValues[Row].Value.ToString() : FString();
+				return ObjectColumn->RowValues.IsValidIndex(Row) ? ObjectColumn->RowValues[Row].Value.ToString() : ObjectColumn->TestValue.ToString();
 			})
 			.AllowedClass(GetAllowedClass(ObjectColumn))
 			.OnObjectChanged_Lambda([TransactionObject, ObjectColumn, Row](const FAssetData& AssetData) {
@@ -44,23 +45,33 @@ namespace UE::ChooserEditor
 					TransactionObject->Modify(true);
 					ObjectColumn->RowValues[Row].Value = AssetData.ToSoftObjectPath();
 				}
+				else
+				{
+					ObjectColumn->TestValue = AssetData.ToSoftObjectPath();
+				}
 			})
 			.DisplayUseSelected(false)
 			.DisplayBrowse(false)
-			.DisplayThumbnail(false)
-			.Visibility_Lambda([ObjectColumn, Row]() {
-				return (ObjectColumn->RowValues.IsValidIndex(Row) &&
-						ObjectColumn->RowValues[Row].Comparison == EObjectColumnCellValueComparison::MatchAny)
-						   ? EVisibility::Collapsed
-						   : EVisibility::Visible;
-			});
+			.DisplayThumbnail(false);
+
+		if (Row >= 0)
+		{
+			ObjectPicker->SetVisibility(TAttribute<EVisibility>::CreateLambda([ObjectColumn, Row]() {
+						return (ObjectColumn->RowValues.IsValidIndex(Row) &&
+								ObjectColumn->RowValues[Row].Comparison == EObjectColumnCellValueComparison::MatchAny)
+								   ? EVisibility::Collapsed
+								   : EVisibility::Visible;
+					}));
+		}
+		return ObjectPicker;
 	}
 
 	static TSharedRef<SWidget> CreateObjectColumnWidget(UChooserTable* Chooser, FChooserColumnBase* Column, int32 Row)
 	{
 		// Extend `SHorizontalBox` a little bit so we can poll for changes & recreate the object picker if necessary.
-		struct SHorizontalBoxEx : public SHorizontalBox
+		class SHorizontalBoxEx : public SHorizontalBox
 		{
+			public:
 			virtual void Tick(const FGeometry& AllottedGeometry, const double InCurrentTime, const float InDeltaTime) override
 			{
 				SHorizontalBox::Tick(AllottedGeometry, InCurrentTime, InDeltaTime);
@@ -69,17 +80,82 @@ namespace UE::ChooserEditor
 				if (AllowedClass != CurrentAllowedClass)
 				{
 					AllowedClass = CurrentAllowedClass;
-					GetSlot(1)[ObjectPickerFactory()];
+					GetSlot(ObjectPickerSlot)[ObjectPickerFactory()];
 				}
 			}
 
 			FObjectColumn* ObjectColumn = nullptr;
 			UClass* AllowedClass = nullptr;
 			TFunction<TSharedRef<SWidget>()> ObjectPickerFactory;
+			int ObjectPickerSlot = 0;
 		};
 
 		FObjectColumn* ObjectColumn = static_cast<FObjectColumn*>(Column);
 
+		if (Row < 0)
+		{
+			// create column header widget
+			TSharedPtr<SWidget> InputValueWidget = nullptr;
+			if (FChooserParameterBase* InputValue = Column->GetInputValue())
+			{
+				InputValueWidget = FObjectChooserWidgetFactories::CreateWidget(false, Chooser, InputValue, Column->GetInputType(), Chooser->ContextObjectType, Chooser->OutputObjectType);
+			}
+			
+			const FSlateBrush* ColumnIcon = FCoreStyle::Get().GetBrush("Icons.Filter");
+			// ColumnIcon = FCoreStyle::Get().GetBrush("Icons.ArrowRight");
+			// ColumnIcon = FAppStyle::Get().GetBrush("Icons.Help");
+			
+			TSharedRef<SWidget> ColumnHeaderWidget = SNew(SHorizontalBox)
+				+ SHorizontalBox::Slot().AutoWidth()
+				[
+					SNew(SBorder)
+					.BorderBackgroundColor(FLinearColor(0,0,0,0))
+					.Content()
+					[
+						SNew(SImage).Image(ColumnIcon)
+					]
+				]
+				+ SHorizontalBox::Slot()
+				[
+					InputValueWidget ? InputValueWidget.ToSharedRef() : SNullWidget::NullWidget
+				];
+		
+			if (Chooser->bEnableDebugTesting)
+			{
+				TSharedRef<SWidget> ObjectPicker = CreateObjectPicker(Chooser, ObjectColumn, Row);
+				ObjectPicker->SetEnabled(TAttribute<bool>::CreateLambda([Chooser]() { return !Chooser->HasDebugTarget(); }));
+				
+				// create widget for test value object picker
+				TSharedRef<SHorizontalBoxEx> CellWidget = SNew(SHorizontalBoxEx)
+					+SHorizontalBox::Slot()
+					.FillWidth(1)
+					[
+						ObjectPicker
+					];
+				
+				CellWidget->SetCanTick(true);
+				CellWidget->ObjectColumn = ObjectColumn;
+				CellWidget->AllowedClass = GetAllowedClass(ObjectColumn);
+				CellWidget->ObjectPickerFactory = [Chooser, ObjectColumn, Row]() { return CreateObjectPicker(Chooser, ObjectColumn, Row); };
+				CellWidget->ObjectPickerSlot = 0;
+				
+				ColumnHeaderWidget = SNew(SVerticalBox)
+				+ SVerticalBox::Slot()
+				[
+					ColumnHeaderWidget
+				]
+				+ SVerticalBox::Slot()
+				[
+					CellWidget
+				];
+			}
+	
+			return ColumnHeaderWidget;
+		}
+
+
+		// create widget for cell
+
 		// clang-format off
 		TSharedRef<SHorizontalBoxEx> CellWidget = SNew(SHorizontalBoxEx)
 			+SHorizontalBox::Slot()
@@ -128,15 +204,13 @@ namespace UE::ChooserEditor
 			[
 				CreateObjectPicker(Chooser, ObjectColumn, Row)
 			];
-		// clang-format on
-
-		// Setup the things we need to check for changes to the allowed class.
-		{
-			CellWidget->SetCanTick(true);
-			CellWidget->ObjectColumn = ObjectColumn;
-			CellWidget->AllowedClass = GetAllowedClass(ObjectColumn);
-			CellWidget->ObjectPickerFactory = [Chooser, ObjectColumn, Row]() { return CreateObjectPicker(Chooser, ObjectColumn, Row); };
-		}
+			// clang-format on
+		
+		CellWidget->SetCanTick(true);
+		CellWidget->ObjectColumn = ObjectColumn;
+		CellWidget->AllowedClass = GetAllowedClass(ObjectColumn);
+		CellWidget->ObjectPickerFactory = [Chooser, ObjectColumn, Row]() { return CreateObjectPicker(Chooser, ObjectColumn, Row); };
+		CellWidget->ObjectPickerSlot = 1;
 
 		return CellWidget;
 	}
diff --git a/Engine/Plugins/Experimental/Chooser/Source/ChooserEditor/Private/RandomizeColumnEditor.cpp b/Engine/Plugins/Experimental/Chooser/Source/ChooserEditor/Private/RandomizeColumnEditor.cpp
--- a/Engine/Plugins/Experimental/Chooser/Source/ChooserEditor/Private/RandomizeColumnEditor.cpp
+++ b/Engine/Plugins/Experimental/Chooser/Source/ChooserEditor/Private/RandomizeColumnEditor.cpp
@@ -6,6 +6,7 @@
 #include "ObjectChooserWidgetFactories.h"
 #include "Widgets/SBoxPanel.h"
 #include "Widgets/Layout/SBox.h"
+#include "Widgets/Layout/SSpacer.h"
 #include "Widgets/Input/SNumericEntryBox.h"
 #include "GraphEditorSettings.h"
 
@@ -17,7 +18,33 @@ namespace UE::ChooserEditor
 TSharedRef<SWidget> CreateRandomizeColumnWidget(UChooserTable* Chooser, FChooserColumnBase* Column, int Row)
 {
 	FRandomizeColumn* RandomizeColumn = static_cast<FRandomizeColumn*>(Column);
-	
+	if (Row < 0)
+	{
+		const FSlateBrush* ColumnIcon = ColumnIcon = FAppStyle::Get().GetBrush("Icons.Help");
+		
+		TSharedRef<SWidget> ColumnHeaderWidget = SNew(SHorizontalBox)
+			+ SHorizontalBox::Slot().AutoWidth()
+			[
+				SNew(SBorder)
+				.BorderBackgroundColor(FLinearColor(0,0,0,0))
+				.Content()
+				[
+					SNew(SImage).Image(ColumnIcon)
+				]
+			]
+			+ SHorizontalBox::Slot().AutoWidth()
+			[
+				SNew(SSpacer).Size(FVector2D(10,10))
+			]
+			+ SHorizontalBox::Slot().VAlign(VAlign_Center)
+			[
+				SNew(STextBlock).Text(LOCTEXT("Randomize","Randomize"))
+			];
+
+		return ColumnHeaderWidget;
+	}
+
+	// create cell widget
 	return
 	SNew(SHorizontalBox)
 		+ SHorizontalBox::Slot().FillWidth(1)
@@ -40,7 +67,7 @@ TSharedRef<SWidget> CreateRandomizeColumnWidget(UChooserTable* Chooser, FChooser
     					{
     						const FScopedTransaction Transaction(LOCTEXT("Edit Randomize Cell Data", "Edit Randomize Cell Data"));
     						Chooser->Modify(true);
-    						RandomizeColumn->RowValues[Row] = Value;
+    						RandomizeColumn->RowValues[Row] = FMath::Max(Value, 0.0);
     					}
     				})
     		]
