Fix Static Analysis bug

diff --git a/Engine/Plugins/Experimental/Chooser/Source/ChooserEditor/Private/RandomizeColumnEditor.cpp b/Engine/Plugins/Experimental/Chooser/Source/ChooserEditor/Private/RandomizeColumnEditor.cpp
--- a/Engine/Plugins/Experimental/Chooser/Source/ChooserEditor/Private/RandomizeColumnEditor.cpp
+++ b/Engine/Plugins/Experimental/Chooser/Source/ChooserEditor/Private/RandomizeColumnEditor.cpp
@@ -20,7 +20,7 @@ TSharedRef<SWidget> CreateRandomizeColumnWidget(UChooserTable* Chooser, FChooser
 	FRandomizeColumn* RandomizeColumn = static_cast<FRandomizeColumn*>(Column);
 	if (Row < 0)
 	{
-		const FSlateBrush* ColumnIcon = ColumnIcon = FAppStyle::Get().GetBrush("Icons.Help");
+		const FSlateBrush* ColumnIcon = FAppStyle::Get().GetBrush("Icons.Help");
 		
 		TSharedRef<SWidget> ColumnHeaderWidget = SNew(SHorizontalBox)
 			+ SHorizontalBox::Slot().AutoWidth()
