Chooser multiple context object and struct support in chooser editor and chooser runtime.

diff --git a/Engine/Plugins/Experimental/Chooser/Source/Chooser/Internal/BoolColumn.h b/Engine/Plugins/Experimental/Chooser/Source/Chooser/Internal/BoolColumn.h
--- a/Engine/Plugins/Experimental/Chooser/Source/Chooser/Internal/BoolColumn.h
+++ b/Engine/Plugins/Experimental/Chooser/Source/Chooser/Internal/BoolColumn.h
@@ -28,8 +28,8 @@ public:
 	UPROPERTY(EditAnywhere, Meta = (BindingType = "bool", BindingAllowFunctions = "true", BindingColor = "BooleanPinTypeColor"), Category = "Binding")
 	FChooserPropertyBinding Binding;
 	
-	virtual bool GetValue(const UObject* ContextObject, bool& OutResult) const override;
-	virtual bool SetValue(UObject* ContextObject, bool InValue) const override;
+	virtual bool GetValue(FChooserEvaluationContext& Context, bool& OutResult) const override;
+	virtual bool SetValue(FChooserEvaluationContext& Context, bool InValue) const override;
 
 	virtual void PostLoad() override
 	{
@@ -51,7 +51,7 @@ public:
 
 	void SetBinding(const TArray<FBindingChainElement>& InBindingChain)
 	{
-		UE::Chooser::CopyPropertyChain(InBindingChain, Binding.PropertyBindingChain);
+		UE::Chooser::CopyPropertyChain(InBindingChain, Binding);
 	}
 
 	virtual void GetDisplayName(FText& OutName) const override
@@ -85,7 +85,7 @@ struct CHOOSER_API FBoolColumn : public FChooserColumnBase
 	UPROPERTY(EditAnywhere, Category= "Data", DisplayName="RowValues");
 	TArray<EBoolColumnCellValue> RowValuesWithAny; 
 	
-	virtual void Filter(FChooserDebuggingInfo& DebugInfo, const UObject* ContextObject, const TArray<uint32>& IndexListIn, TArray<uint32>& IndexListOut) const override;
+	virtual void Filter(FChooserEvaluationContext& Context, const TArray<uint32>& IndexListIn, TArray<uint32>& IndexListOut) const override;
 
 #if WITH_EDITOR
 	mutable bool TestValue;
diff --git a/Engine/Plugins/Experimental/Chooser/Source/Chooser/Internal/Chooser.h b/Engine/Plugins/Experimental/Chooser/Source/Chooser/Internal/Chooser.h
--- a/Engine/Plugins/Experimental/Chooser/Source/Chooser/Internal/Chooser.h
+++ b/Engine/Plugins/Experimental/Chooser/Source/Chooser/Internal/Chooser.h
@@ -10,8 +10,6 @@
 
 #include "Chooser.generated.h"
 
-DECLARE_MULTICAST_DELEGATE_OneParam(FChooserOutputObjectTypeChanged, const UClass* OutputObjectType);
-
 UCLASS(BlueprintType)
 class CHOOSER_API UChooserTable : public UObject, public IHasContextClass
 {
@@ -33,7 +31,7 @@ public:
 	void SetDebugTarget(TWeakObjectPtr<const UObject> Target) { DebugTarget = Target; }
 	void ResetDebugTarget() { DebugTarget.Reset(); }
 	void IterateRecentContextObjects(TFunction<void(const UObject*)> Callback) const;
-	bool UpdateDebugging(const UObject* ContextObject) const;
+	void UpdateDebugging(FChooserEvaluationContext& Context) const;
 	
 	// enable display of which cells pass/fail based on current TestValue for each column
 	bool bEnableDebugTesting = false;
@@ -42,7 +40,6 @@ public:
 private: 
 	// caching the OutputObjectType and ContextObjectType so that on Undo, we can tell if we should fire the changed delegate
 	UClass* CachedPreviousOutputObjectType;
-	UClass* CachedPreviousContextObjectType;
 	
 	// objects this chooser has been recently evaluated on
 	mutable TSet<TWeakObjectPtr<const UObject>> RecentContextObjects;
@@ -58,6 +55,10 @@ public:
 	// deprecated UObject Results
 	UPROPERTY()
 	TArray<TScriptInterface<IObjectChooser>> Results_DEPRECATED;
+
+	// deprecated single context object
+	UPROPERTY()
+	TObjectPtr<UClass> ContextObjectType_DEPRECATED;
 	
 	// deprecated UObject Columns
 	UPROPERTY()
@@ -71,14 +72,14 @@ public:
 	// Columns which filter Results
 	UPROPERTY(EditAnywhere, DisplayName = "Columns", Category = Hidden, meta = (ExcludeBaseStruct, BaseStruct = "/Script/Chooser.ChooserColumnBase"))
 	TArray<FInstancedStruct> ColumnsStructs;
-
-	UPROPERTY(EditAnywhere, Category="Input", Meta = (AllowAbstract=true))
-	TObjectPtr<UClass> ContextObjectType;
 	
+	UPROPERTY(EditAnywhere, Meta = (ExcludeBaseStruct, BaseStruct = "/Script/Chooser.ContextObjectTypeBase"), Category = "Input")
+	TArray<FInstancedStruct> ContextData;
+
 	UPROPERTY(EditAnywhere, Category="Output", Meta = (AllowAbstract=true))
 	TObjectPtr<UClass> OutputObjectType;
 
-	virtual UClass* GetContextClass() override { return ContextObjectType; }
+	virtual TConstArrayView<FInstancedStruct> GetContextData() const override { return ContextData; }
 };
 
 
@@ -87,8 +88,8 @@ struct CHOOSER_API FEvaluateChooser : public FObjectChooserBase
 {
 	GENERATED_BODY()
 
-	virtual UObject* ChooseObject(const UObject* ContextObject) const final override;
-	virtual EIteratorStatus ChooseMulti(const UObject* ContextObject, FObjectChooserIteratorCallback Callback) const final override;
+	virtual UObject* ChooseObject(FChooserEvaluationContext& Context) const final override;
+	virtual EIteratorStatus ChooseMulti(FChooserEvaluationContext &Context, FObjectChooserIteratorCallback Callback) const final override;
 	public:
 	
 	UPROPERTY(EditAnywhere, Category="Parameters")
diff --git a/Engine/Plugins/Experimental/Chooser/Source/Chooser/Internal/EnumColumn.h b/Engine/Plugins/Experimental/Chooser/Source/Chooser/Internal/EnumColumn.h
--- a/Engine/Plugins/Experimental/Chooser/Source/Chooser/Internal/EnumColumn.h
+++ b/Engine/Plugins/Experimental/Chooser/Source/Chooser/Internal/EnumColumn.h
@@ -21,7 +21,7 @@ struct CHOOSER_API FEnumContextProperty : public FChooserParameterEnumBase
 	UPROPERTY(EditAnywhere, Meta = (BindingType = "enum", BindingColor = "BytePinTypeColor"), Category = "Binding")
 	FChooserEnumPropertyBinding Binding;
 
-	virtual bool GetValue(const UObject* ContextObject, uint8& OutResult) const override;
+	virtual bool GetValue(FChooserEvaluationContext& Context, uint8& OutResult) const override;
 
 	virtual void PostLoad() override
 	{
@@ -108,7 +108,7 @@ public:
 	// should match the length of the Results array
 	TArray<FChooserEnumRowData> RowValues;
 
-	virtual void Filter(FChooserDebuggingInfo& DebugInfo, const UObject* ContextObject, const TArray<uint32>& IndexListIn, TArray<uint32>& IndexListOut) const override;
+	virtual void Filter(FChooserEvaluationContext& Context, const TArray<uint32>& IndexListIn, TArray<uint32>& IndexListOut) const override;
 	
 #if WITH_EDITOR
 	mutable int32 TestValue;
diff --git a/Engine/Plugins/Experimental/Chooser/Source/Chooser/Internal/FloatRangeColumn.h b/Engine/Plugins/Experimental/Chooser/Source/Chooser/Internal/FloatRangeColumn.h
--- a/Engine/Plugins/Experimental/Chooser/Source/Chooser/Internal/FloatRangeColumn.h
+++ b/Engine/Plugins/Experimental/Chooser/Source/Chooser/Internal/FloatRangeColumn.h
@@ -13,7 +13,7 @@ struct CHOOSER_API FFloatContextProperty :  public FChooserParameterFloatBase
 {
 	GENERATED_BODY()
 	
-	virtual bool GetValue(const UObject* ContextObject, float& OutResult) const override;
+	virtual bool GetValue(FChooserEvaluationContext& Context, float& OutResult) const override;
 
 	UPROPERTY()
 	TArray<FName> PropertyBindingChain_DEPRECATED;
@@ -41,7 +41,7 @@ struct CHOOSER_API FFloatContextProperty :  public FChooserParameterFloatBase
 	
 	void SetBinding(const TArray<FBindingChainElement>& InBindingChain)
 	{
-		UE::Chooser::CopyPropertyChain(InBindingChain, Binding.PropertyBindingChain);
+		UE::Chooser::CopyPropertyChain(InBindingChain, Binding);
 	}
 
 	virtual void GetDisplayName(FText& OutName) const override
@@ -87,7 +87,7 @@ struct CHOOSER_API FFloatRangeColumn : public FChooserColumnBase
 	// should match the length of the Results array 
 	TArray<FChooserFloatRangeRowData> RowValues;
 	
-	virtual void Filter(FChooserDebuggingInfo& DebugInfo, const UObject* ContextObject, const TArray<uint32>& IndexListIn, TArray<uint32>& IndexListOut) const override;
+	virtual void Filter(FChooserEvaluationContext& Context, const TArray<uint32>& IndexListIn, TArray<uint32>& IndexListOut) const override;
 
 #if WITH_EDITOR
 	mutable float TestValue;
diff --git a/Engine/Plugins/Experimental/Chooser/Source/Chooser/Internal/GameplayTagColumn.h b/Engine/Plugins/Experimental/Chooser/Source/Chooser/Internal/GameplayTagColumn.h
--- a/Engine/Plugins/Experimental/Chooser/Source/Chooser/Internal/GameplayTagColumn.h
+++ b/Engine/Plugins/Experimental/Chooser/Source/Chooser/Internal/GameplayTagColumn.h
@@ -23,7 +23,7 @@ public:
 	UPROPERTY(EditAnywhere, Meta = (BindingType = "FGameplayTagContainer", BindingColor = "StructPinTypeColor"), Category = "Binding")
 	FChooserPropertyBinding Binding;
 	
-	virtual bool GetValue(const UObject* ContextObject, const FGameplayTagContainer*& OutResult) const override;
+	virtual bool GetValue(FChooserEvaluationContext& Context, const FGameplayTagContainer*& OutResult) const override;
 	
 	virtual void PostLoad() override
 	{
@@ -42,7 +42,7 @@ public:
 
 	void SetBinding(const TArray<FBindingChainElement>& InBindingChain)
 	{
-		UE::Chooser::CopyPropertyChain(InBindingChain, Binding.PropertyBindingChain);
+		UE::Chooser::CopyPropertyChain(InBindingChain, Binding);
 	}
 
 	virtual void GetDisplayName(FText& OutName) const override
@@ -79,7 +79,7 @@ struct CHOOSER_API FGameplayTagColumn : public FChooserColumnBase
 	// should match the length of the Results array 
 	TArray<FGameplayTagContainer> RowValues;
 	
-	virtual void Filter(FChooserDebuggingInfo& DebugInfo, const UObject* ContextObject, const TArray<uint32>& IndexListIn, TArray<uint32>& IndexListOut) const override;
+	virtual void Filter(FChooserEvaluationContext& Context, const TArray<uint32>& IndexListIn, TArray<uint32>& IndexListOut) const override;
 
 #if WITH_EDITOR
 	mutable FGameplayTagContainer TestValue;
diff --git a/Engine/Plugins/Experimental/Chooser/Source/Chooser/Internal/ObjectChooser_Asset.h b/Engine/Plugins/Experimental/Chooser/Source/Chooser/Internal/ObjectChooser_Asset.h
--- a/Engine/Plugins/Experimental/Chooser/Source/Chooser/Internal/ObjectChooser_Asset.h
+++ b/Engine/Plugins/Experimental/Chooser/Source/Chooser/Internal/ObjectChooser_Asset.h
@@ -12,7 +12,7 @@ struct CHOOSER_API FAssetChooser : public FObjectChooserBase
 	GENERATED_BODY()
 	
 	// FObjectChooserBase interface
-	virtual UObject* ChooseObject(const UObject* ContextObject) const final override;
+	virtual UObject* ChooseObject(FChooserEvaluationContext& Context) const final override;
 public: 
 	UPROPERTY(EditAnywhere, Category = "Parameters")
 	TObjectPtr<UObject> Asset;
diff --git a/Engine/Plugins/Experimental/Chooser/Source/Chooser/Internal/ObjectChooser_Class.h b/Engine/Plugins/Experimental/Chooser/Source/Chooser/Internal/ObjectChooser_Class.h
--- a/Engine/Plugins/Experimental/Chooser/Source/Chooser/Internal/ObjectChooser_Class.h
+++ b/Engine/Plugins/Experimental/Chooser/Source/Chooser/Internal/ObjectChooser_Class.h
@@ -12,7 +12,7 @@ struct CHOOSER_API FClassChooser : public FObjectChooserBase
 	GENERATED_BODY()
 	
 	// FObjectChooserBase interface
-	virtual UObject* ChooseObject(const UObject* ContextObject) const final override;
+	virtual UObject* ChooseObject(FChooserEvaluationContext& Context) const final override;
 public: 
 	UPROPERTY(EditAnywhere, Category = "Parameters")
 	TObjectPtr<UClass> Class;
diff --git a/Engine/Plugins/Experimental/Chooser/Source/Chooser/Internal/ObjectColumn.h b/Engine/Plugins/Experimental/Chooser/Source/Chooser/Internal/ObjectColumn.h
--- a/Engine/Plugins/Experimental/Chooser/Source/Chooser/Internal/ObjectColumn.h
+++ b/Engine/Plugins/Experimental/Chooser/Source/Chooser/Internal/ObjectColumn.h
@@ -18,7 +18,7 @@ struct CHOOSER_API FObjectContextProperty : public FChooserParameterObjectBase
 	UPROPERTY(EditAnywhere, Meta = (BindingType = "object", BindingColor = "ObjectPinTypeColor"), Category = "Binding")
 	FChooserObjectPropertyBinding Binding;
 
-	virtual bool GetValue(const UObject* ContextObject, FSoftObjectPath& OutResult) const override;
+	virtual bool GetValue(FChooserEvaluationContext& Context, FSoftObjectPath& OutResult) const override;
 
 #if WITH_EDITOR
 	static bool CanBind(const FProperty& Property)
@@ -84,7 +84,7 @@ struct CHOOSER_API FObjectColumn : public FChooserColumnBase
 	// should match the length of the Results array
 	TArray<FChooserObjectRowData> RowValues;
 
-	virtual void Filter(FChooserDebuggingInfo& DebugInfo, const UObject* ContextObject, const TArray<uint32>& IndexListIn, TArray<uint32>& IndexListOut) const override;
+	virtual void Filter(FChooserEvaluationContext& Context, const TArray<uint32>& IndexListIn, TArray<uint32>& IndexListOut) const override;
 
 #if WITH_EDITOR
 	mutable FSoftObjectPath TestValue;
diff --git a/Engine/Plugins/Experimental/Chooser/Source/Chooser/Internal/OutputBoolColumn.h b/Engine/Plugins/Experimental/Chooser/Source/Chooser/Internal/OutputBoolColumn.h
--- a/Engine/Plugins/Experimental/Chooser/Source/Chooser/Internal/OutputBoolColumn.h
+++ b/Engine/Plugins/Experimental/Chooser/Source/Chooser/Internal/OutputBoolColumn.h
@@ -16,7 +16,7 @@ struct CHOOSER_API FOutputBoolColumn : public FChooserColumnBase
 	FOutputBoolColumn();
 	virtual bool HasFilters() const override { return false; }
 	virtual bool HasOutputs() const override { return true; }
-	virtual void SetOutputs(FChooserDebuggingInfo& DebugInfo, UObject* ContextObject, int RowIndex) const override;
+	virtual void SetOutputs(FChooserEvaluationContext& Context, int RowIndex) const override;
 	
 	UPROPERTY(EditAnywhere, Meta = (ExcludeBaseStruct, BaseStruct = "/Script/Chooser.ChooserParameterBoolBase"), Category = "Hidden")
 	FInstancedStruct InputValue;
diff --git a/Engine/Plugins/Experimental/Chooser/Source/Chooser/Internal/RandomizeColumn.h b/Engine/Plugins/Experimental/Chooser/Source/Chooser/Internal/RandomizeColumn.h
--- a/Engine/Plugins/Experimental/Chooser/Source/Chooser/Internal/RandomizeColumn.h
+++ b/Engine/Plugins/Experimental/Chooser/Source/Chooser/Internal/RandomizeColumn.h
@@ -17,7 +17,7 @@ public:
 	UPROPERTY(EditAnywhere, Meta = (BindingType = "FChooserRandomizationContext", BindingAllowFunctions = "false", BindingColor = "StructPinTypeColor"), Category = "Binding")
 	FChooserPropertyBinding Binding;
 	
-	virtual bool GetValue(const UObject* ContextObject, const FChooserRandomizationContext*& OutResult) const override;
+	virtual bool GetValue(FChooserEvaluationContext& Context, const FChooserRandomizationContext*& OutResult) const override;
 
 #if WITH_EDITOR
 	static bool CanBind(const FProperty& Property)
@@ -28,7 +28,7 @@ public:
 
 	void SetBinding(const TArray<FBindingChainElement>& InBindingChain)
 	{
-		UE::Chooser::CopyPropertyChain(InBindingChain, Binding.PropertyBindingChain);
+		UE::Chooser::CopyPropertyChain(InBindingChain, Binding);
 	}
 
 	virtual void GetDisplayName(FText& OutName) const override
@@ -63,8 +63,8 @@ struct CHOOSER_API FRandomizeColumn : public FChooserColumnBase
 	UPROPERTY(EditAnywhere, Category= "Data", DisplayName="RowValues");
 	TArray<float> RowValues; 
 	
-	virtual void Filter(FChooserDebuggingInfo& DebugInfo, const UObject* ContextObject, const TArray<uint32>& IndexListIn, TArray<uint32>& IndexListOut) const override;
-	virtual void SetOutputs(FChooserDebuggingInfo& DebugInfo, UObject* ContextObject, int RowIndex) const override;
+	virtual void Filter(FChooserEvaluationContext& Context, const TArray<uint32>& IndexListIn, TArray<uint32>& IndexListOut) const override;
+	virtual void SetOutputs(FChooserEvaluationContext& Context, int RowIndex) const override;
 	
 	virtual bool HasFilters() const override { return true; }
 	virtual bool HasOutputs() const override { return true; }
diff --git a/Engine/Plugins/Experimental/Chooser/Source/Chooser/Private/BoolColumn.cpp b/Engine/Plugins/Experimental/Chooser/Source/Chooser/Private/BoolColumn.cpp
--- a/Engine/Plugins/Experimental/Chooser/Source/Chooser/Private/BoolColumn.cpp
+++ b/Engine/Plugins/Experimental/Chooser/Source/Chooser/Private/BoolColumn.cpp
@@ -3,13 +3,13 @@
 #include "ChooserPropertyAccess.h"
 #include "Chooser.h"
 
-bool FBoolContextProperty::GetValue(const UObject* ContextObject, bool& OutResult) const
+bool FBoolContextProperty::GetValue(FChooserEvaluationContext& Context, bool& OutResult) const
 {
 	
-	UStruct* StructType = ContextObject->GetClass();
-	const void* Container = ContextObject;
+	const UStruct* StructType = nullptr;
+	const void* Container = nullptr;
 	
-	if (UE::Chooser::ResolvePropertyChain(Container, StructType, Binding.PropertyBindingChain))
+	if (UE::Chooser::ResolvePropertyChain(Context, Binding, Container, StructType))
 	{
 		if (const FBoolProperty* Property = FindFProperty<FBoolProperty>(StructType, Binding.PropertyBindingChain.Last()))
 		{
@@ -17,7 +17,7 @@ bool FBoolContextProperty::GetValue(const UObject* ContextObject, bool& OutResul
 			return true;
 		}
 		
-	    if (UClass* ClassType = Cast<UClass>(StructType))
+	    if (const UClass* ClassType = Cast<const UClass>(StructType))
 	    {
 			if (UFunction* Function = ClassType->FindFunctionByName(Binding.PropertyBindingChain.Last()))
 			{
@@ -38,12 +38,12 @@ bool FBoolContextProperty::GetValue(const UObject* ContextObject, bool& OutResul
 	return false;
 }
 
-bool FBoolContextProperty::SetValue(UObject* ContextObject, bool InValue) const
+bool FBoolContextProperty::SetValue(FChooserEvaluationContext& Context, bool InValue) const
 {
-	UStruct* StructType = ContextObject->GetClass();
-	const void* Container = ContextObject;
+	const UStruct* StructType = nullptr;
+	const void* Container = nullptr;
 	
-	if (UE::Chooser::ResolvePropertyChain(Container, StructType, Binding.PropertyBindingChain))
+	if (UE::Chooser::ResolvePropertyChain(Context, Binding, Container, StructType))
 	{
 		if (FBoolProperty* Property = FindFProperty<FBoolProperty>(StructType, Binding.PropertyBindingChain.Last()))
 		{
@@ -61,15 +61,15 @@ FBoolColumn::FBoolColumn()
 	InputValue.InitializeAs(FBoolContextProperty::StaticStruct());
 }
 
-void FBoolColumn::Filter(FChooserDebuggingInfo& DebugInfo, const UObject* ContextObject, const TArray<uint32>& IndexListIn, TArray<uint32>& IndexListOut) const
+void FBoolColumn::Filter(FChooserEvaluationContext& Context, const TArray<uint32>& IndexListIn, TArray<uint32>& IndexListOut) const
 {
-	if (ContextObject && InputValue.IsValid())
+	if (InputValue.IsValid())
 	{
 		bool Result = false;
-		InputValue.Get<FChooserParameterBoolBase>().GetValue(ContextObject,Result);
+		InputValue.Get<FChooserParameterBoolBase>().GetValue(Context,Result);
 
 	#if WITH_EDITOR
-		if (DebugInfo.bCurrentDebugTarget)
+		if (Context.DebuggingInfo.bCurrentDebugTarget)
 		{
 			TestValue = Result;
 		}
diff --git a/Engine/Plugins/Experimental/Chooser/Source/Chooser/Private/Chooser.cpp b/Engine/Plugins/Experimental/Chooser/Source/Chooser/Private/Chooser.cpp
--- a/Engine/Plugins/Experimental/Chooser/Source/Chooser/Private/Chooser.cpp
+++ b/Engine/Plugins/Experimental/Chooser/Source/Chooser/Private/Chooser.cpp
@@ -16,16 +16,13 @@ void UChooserTable::PostEditUndo()
 {
 	UObject::PostEditUndo();
 
+
 	if (CachedPreviousOutputObjectType != OutputObjectType)
 	{
 		OnOutputObjectTypeChanged.Broadcast(OutputObjectType);
 		CachedPreviousOutputObjectType = OutputObjectType;
 	}
-	if (CachedPreviousContextObjectType != ContextObjectType)
-	{
-		OnContextClassChanged.Broadcast(ContextObjectType);
-		CachedPreviousContextObjectType = ContextObjectType;
-	}
+	OnContextClassChanged.Broadcast();
 }
 
 void UChooserTable::PostEditChangeProperty(FPropertyChangedEvent& PropertyChangedEvent)
@@ -42,14 +39,10 @@ void UChooserTable::PostEditChangeProperty(FPropertyChangedEvent& PropertyChange
 		}
 	}
 	
-	static FName ContextPropertyName = "ContextObjectType";
+	static FName ContextPropertyName = "ContextData";
 	if (PropertyChangedEvent.Property->GetName() == ContextPropertyName)
 	{
-		if (CachedPreviousContextObjectType != ContextObjectType)
-		{
-			OnContextClassChanged.Broadcast(ContextObjectType);
-			CachedPreviousContextObjectType = ContextObjectType;
-		}
+		OnContextClassChanged.Broadcast();
 	}
 }
 
@@ -61,6 +54,16 @@ void UChooserTable::PostLoad()
 
 	// convert old data if it exists
 
+	if (ContextObjectType_DEPRECATED)
+	{
+		ContextData.SetNum(1);
+		ContextData[0].InitializeAs<FContextObjectTypeClass>();
+		FContextObjectTypeClass& Context = ContextData[0].GetMutable<FContextObjectTypeClass>();
+		Context.Class = ContextObjectType_DEPRECATED;
+		Context.Direction = EContextObjectDirection::ReadWrite;
+		ContextObjectType_DEPRECATED = nullptr;
+	}
+
 	if (Results_DEPRECATED.Num() > 0 || Columns_DEPRECATED.Num() > 0)
 	{
 		ResultsStructs.Reserve(Results_DEPRECATED.Num());
@@ -114,31 +117,39 @@ void UChooserTable::IterateRecentContextObjects(TFunction<void(const UObject*)>
 	}
 }
 
-bool UChooserTable::UpdateDebugging(const UObject* ContextObject) const
+void UChooserTable::UpdateDebugging(FChooserEvaluationContext& Context) const
 {
 	FScopeLock Lock(&DebugLock);
-	RecentContextObjects.Add(MakeWeakObjectPtr(ContextObject));
-	if (ContextObject == DebugTarget)
+	for (const FContextEntry& Entry : Context.ContextData)
 	{
-		bDebugTestValuesValid = true;
-		return true;
+		if (Entry.Data && Entry.Type && Entry.Type->IsA(UObject::StaticClass()))
+		{
+			UObject* ContextObject = static_cast<UObject*>(Entry.Data);
+			RecentContextObjects.Add(MakeWeakObjectPtr(ContextObject));
+			if (ContextObject == DebugTarget)
+			{
+				bDebugTestValuesValid = true;
+				Context.DebuggingInfo.bCurrentDebugTarget = true;
+				return;
+			}
+		}
 	}
-	return false;
+	Context.DebuggingInfo.bCurrentDebugTarget = false;
 }
 
 #endif
 
-static FObjectChooserBase::EIteratorStatus StaticEvaluateChooser(const UObject* ContextObject, const UChooserTable* Chooser, FObjectChooserBase::FObjectChooserIteratorCallback Callback)
+static FObjectChooserBase::EIteratorStatus StaticEvaluateChooser(FChooserEvaluationContext& Context, const UChooserTable* Chooser, FObjectChooserBase::FObjectChooserIteratorCallback Callback)
 {
 	if (Chooser == nullptr)
 	{
 		return FObjectChooserBase::EIteratorStatus::Continue;
 	}
-	
-	FChooserDebuggingInfo DebugInfo;
+
+	// todo validate that parameter types in context data match
 
 #if WITH_EDITOR
-	DebugInfo.bCurrentDebugTarget = Chooser->UpdateDebugging(ContextObject);
+	Chooser->UpdateDebugging(Context);
 #endif
 
 	TArray<uint32> Indices1;
@@ -160,7 +171,7 @@ static FObjectChooserBase::EIteratorStatus StaticEvaluateChooser(const UObject*
 		{
 			Swap(IndicesIn, IndicesOut);
 			IndicesOut->SetNum(0, false);
-			Column.Filter(DebugInfo, ContextObject, *IndicesIn, *IndicesOut);
+			Column.Filter(Context, *IndicesIn, *IndicesOut);
 		}
 	}
 	
@@ -170,16 +181,16 @@ static FObjectChooserBase::EIteratorStatus StaticEvaluateChooser(const UObject*
 		if (Chooser->ResultsStructs.Num() > (int32)SelectedIndex)
 		{
 			const FObjectChooserBase& SelectedResult = Chooser->ResultsStructs[SelectedIndex].Get<FObjectChooserBase>();
-			if (SelectedResult.ChooseMulti(ContextObject, Callback) == FObjectChooserBase::EIteratorStatus::Stop)
+			if (SelectedResult.ChooseMulti(Context, Callback) == FObjectChooserBase::EIteratorStatus::Stop)
 			{
 				// trigger all output columns
 				for (const FInstancedStruct& ColumnData : Chooser->ColumnsStructs)
 				{
 					const FChooserColumnBase& Column = ColumnData.Get<FChooserColumnBase>();
-					Column.SetOutputs(DebugInfo, const_cast<UObject*>(ContextObject), SelectedIndex);
+					Column.SetOutputs(Context, SelectedIndex);
 				}
 #if WITH_EDITOR
-				if (DebugInfo.bCurrentDebugTarget)
+				if (Context.DebuggingInfo.bCurrentDebugTarget)
 				{
 					Chooser->SetDebugSelectedRow(SelectedIndex);
 				}
@@ -192,10 +203,10 @@ static FObjectChooserBase::EIteratorStatus StaticEvaluateChooser(const UObject*
 	return FObjectChooserBase::EIteratorStatus::Continue;
 }
 
-UObject* FEvaluateChooser::ChooseObject(const UObject* ContextObject) const
+UObject* FEvaluateChooser::ChooseObject(FChooserEvaluationContext& Context) const
 {
 	UObject* Result = nullptr;
-	StaticEvaluateChooser(ContextObject, Chooser, FObjectChooserIteratorCallback::CreateLambda([&Result](UObject* InResult)
+	StaticEvaluateChooser(Context, Chooser, FObjectChooserIteratorCallback::CreateLambda([&Result](UObject* InResult)
 	{
 		Result = InResult;
 		return FObjectChooserBase::EIteratorStatus::Stop;
@@ -204,9 +215,9 @@ UObject* FEvaluateChooser::ChooseObject(const UObject* ContextObject) const
 	return Result;
 }
 
-FObjectChooserBase::EIteratorStatus FEvaluateChooser::ChooseMulti(const UObject* ContextObject, FObjectChooserIteratorCallback Callback) const
+FObjectChooserBase::EIteratorStatus FEvaluateChooser::ChooseMulti(FChooserEvaluationContext& Context, FObjectChooserIteratorCallback Callback) const
 {
-	return StaticEvaluateChooser(ContextObject, Chooser, Callback);
+	return StaticEvaluateChooser(Context, Chooser, Callback);
 }
 
 UChooserFunctionLibrary::UChooserFunctionLibrary(const FObjectInitializer& ObjectInitializer)
@@ -214,6 +225,14 @@ UChooserFunctionLibrary::UChooserFunctionLibrary(const FObjectInitializer& Objec
 {
 }
 
+static FObjectChooserBase::EIteratorStatus StaticEvaluateChooser(const UObject* ContextObject, const UChooserTable* Chooser, FObjectChooserBase::FObjectChooserIteratorCallback Callback)
+{
+	// Fallback single context object version
+	FChooserEvaluationContext Context;
+	Context.ContextData.Push(FContextEntry{ContextObject->GetClass(), const_cast<void*>(static_cast<const void*>(ContextObject))});
+	return StaticEvaluateChooser(Context, Chooser, Callback);
+}
+
 UObject* UChooserFunctionLibrary::EvaluateChooser(const UObject* ContextObject, const UChooserTable* Chooser, TSubclassOf<UObject> ObjectClass)
 {
 	UObject* Result = nullptr;
diff --git a/Engine/Plugins/Experimental/Chooser/Source/Chooser/Private/ChooserPropertyAccess.cpp b/Engine/Plugins/Experimental/Chooser/Source/Chooser/Private/ChooserPropertyAccess.cpp
--- a/Engine/Plugins/Experimental/Chooser/Source/Chooser/Private/ChooserPropertyAccess.cpp
+++ b/Engine/Plugins/Experimental/Chooser/Source/Chooser/Private/ChooserPropertyAccess.cpp
@@ -1,6 +1,8 @@
 // Copyright Epic Games, Inc. All Rights Reserved.
 
 #include "ChooserPropertyAccess.h"
+
+#include "IObjectChooser.h"
 #include "UObject/UnrealType.h"
 
 #if WITH_EDITOR
@@ -9,7 +11,20 @@
 
 namespace UE::Chooser
 {
-	bool ResolvePropertyChain(const void*& Container, UStruct*& StructType, const TArray<FName>& PropertyBindingChain)
+	bool ResolvePropertyChain(FChooserEvaluationContext& Context, const FChooserPropertyBinding& PropertyBinding, const void*& OutContainer, const UStruct*& OutStructType)
+	{
+		if (PropertyBinding.ContextIndex < Context.ContextData.Num())
+		{
+			OutContainer = Context.ContextData[PropertyBinding.ContextIndex].Data;
+			OutStructType = Context.ContextData[PropertyBinding.ContextIndex].Type;
+
+			return ResolvePropertyChain(OutContainer, OutStructType, PropertyBinding.PropertyBindingChain);
+		}
+
+		return false;
+	}
+	
+	bool ResolvePropertyChain(const void*& Container, const UStruct*& StructType, const TArray<FName>& PropertyBindingChain)
 	{
 		if (PropertyBindingChain.Num() == 0)
 		{
@@ -36,7 +51,7 @@ namespace UE::Chooser
 			else
 			{
 				// check if it's a member function
-				if (UClass* ClassType = Cast<UClass>(StructType))
+				if (const UClass* ClassType = Cast<const UClass>(StructType))
 				{
 					if (UFunction* Function = ClassType->FindFunctionByName(PropertyBindingChain[PropertyChainIndex]))
 					{
@@ -74,16 +89,20 @@ namespace UE::Chooser
 	
 		return true;
 	}
-
-
+	
 #if WITH_EDITOR
-	void CopyPropertyChain(const TArray<FBindingChainElement>& InBindingChain, TArray<FName>& OutPropertyBindingChain)
+	void CopyPropertyChain(const TArray<FBindingChainElement>& InBindingChain, FChooserPropertyBinding& OutPropertyBinding)
 	{
-		OutPropertyBindingChain.Empty();
+		OutPropertyBinding.PropertyBindingChain.Empty();
+
+		if (InBindingChain.Num() > 0)
+		{
+			OutPropertyBinding.ContextIndex = InBindingChain[0].ArrayIndex;
+		}
 
 		for (int32 i = 1; i < InBindingChain.Num(); ++i)
 		{
-			OutPropertyBindingChain.Emplace(InBindingChain[i].Field.GetFName());
+			OutPropertyBinding.PropertyBindingChain.Emplace(InBindingChain[i].Field.GetFName());
 		}
 	}
 #endif
diff --git a/Engine/Plugins/Experimental/Chooser/Source/Chooser/Private/EnumColumn.cpp b/Engine/Plugins/Experimental/Chooser/Source/Chooser/Private/EnumColumn.cpp
--- a/Engine/Plugins/Experimental/Chooser/Source/Chooser/Private/EnumColumn.cpp
+++ b/Engine/Plugins/Experimental/Chooser/Source/Chooser/Private/EnumColumn.cpp
@@ -6,12 +6,12 @@
 #include "IPropertyAccessEditor.h"
 #endif
 
-bool FEnumContextProperty::GetValue(const UObject* ContextObject, uint8& OutResult) const
+bool FEnumContextProperty::GetValue(FChooserEvaluationContext& Context, uint8& OutResult) const
 {
-	UStruct* StructType = ContextObject->GetClass();
-	const void* Container = ContextObject;
+	const UStruct* StructType = nullptr;
+	const void* Container = nullptr;
 
-	if (UE::Chooser::ResolvePropertyChain(Container, StructType, Binding.PropertyBindingChain))
+	if (UE::Chooser::ResolvePropertyChain(Context, Binding, Container, StructType))
 	{
 		if (const FEnumProperty* EnumProperty = FindFProperty<FEnumProperty>(StructType, Binding.PropertyBindingChain.Last()))
 		{
@@ -39,7 +39,7 @@ void FEnumContextProperty::SetBinding(const TArray<FBindingChainElement>& InBind
 	const UEnum* PreviousEnum = Binding.Enum;
 	Binding.Enum = nullptr;
 
-	UE::Chooser::CopyPropertyChain(InBindingChain, Binding.PropertyBindingChain);
+	UE::Chooser::CopyPropertyChain(InBindingChain, Binding);
 
 	const FField* Field = InBindingChain.Last().Field.ToField();
 	if (const FEnumProperty* EnumProperty = CastField<FEnumProperty>(Field))
@@ -67,15 +67,14 @@ bool FChooserEnumRowData::Evaluate(const uint8 LeftHandSide) const
 	return Equal ^ CompareNotEqual;
 }
 
-void FEnumColumn::Filter(FChooserDebuggingInfo& DebugInfo, const UObject* ContextObject, const TArray<uint32>& IndexListIn, TArray<uint32>& IndexListOut) const
+void FEnumColumn::Filter(FChooserEvaluationContext& Context, const TArray<uint32>& IndexListIn, TArray<uint32>& IndexListOut) const
 {
 	uint8 Result = 0;
-	if (ContextObject != nullptr &&
-		InputValue.IsValid() &&
-		InputValue.Get<FChooserParameterEnumBase>().GetValue(ContextObject, Result))
+	if (InputValue.IsValid() &&
+		InputValue.Get<FChooserParameterEnumBase>().GetValue(Context, Result))
 	{
 #if WITH_EDITOR
-		if (DebugInfo.bCurrentDebugTarget)
+		if (Context.DebuggingInfo.bCurrentDebugTarget)
 		{
 			TestValue = Result;
 		}
diff --git a/Engine/Plugins/Experimental/Chooser/Source/Chooser/Private/FloatRangeColumn.cpp b/Engine/Plugins/Experimental/Chooser/Source/Chooser/Private/FloatRangeColumn.cpp
--- a/Engine/Plugins/Experimental/Chooser/Source/Chooser/Private/FloatRangeColumn.cpp
+++ b/Engine/Plugins/Experimental/Chooser/Source/Chooser/Private/FloatRangeColumn.cpp
@@ -3,12 +3,12 @@
 
 #include "ChooserPropertyAccess.h"
 
-bool FFloatContextProperty::GetValue(const UObject* ContextObject, float& OutResult) const
+bool FFloatContextProperty::GetValue(FChooserEvaluationContext& Context, float& OutResult) const
 {
-	UStruct* StructType = ContextObject->GetClass();
-	const void* Container = ContextObject;
+	const UStruct* StructType = nullptr;
+	const void* Container = nullptr;
 
-	if (UE::Chooser::ResolvePropertyChain(Container, StructType, Binding.PropertyBindingChain))
+	if (UE::Chooser::ResolvePropertyChain(Context, Binding, Container, StructType))
 	{
 		if (const FDoubleProperty* DoubleProperty = FindFProperty<FDoubleProperty>(StructType, Binding.PropertyBindingChain.Last()))
 		{
@@ -22,7 +22,7 @@ bool FFloatContextProperty::GetValue(const UObject* ContextObject, float& OutRes
 			return true;
 		}
 
-	    if (UClass* ClassType = Cast<UClass>(StructType))
+	    if (const UClass* ClassType = Cast<const UClass>(StructType))
 	    {
 			if (UFunction* Function = ClassType->FindFunctionByName(Binding.PropertyBindingChain.Last()))
 			{
@@ -69,15 +69,15 @@ FFloatRangeColumn::FFloatRangeColumn()
 	InputValue.InitializeAs(FFloatContextProperty::StaticStruct());
 }
 
-void FFloatRangeColumn::Filter(FChooserDebuggingInfo& DebugInfo, const UObject* ContextObject, const TArray<uint32>& IndexListIn, TArray<uint32>& IndexListOut) const
+void FFloatRangeColumn::Filter(FChooserEvaluationContext& Context, const TArray<uint32>& IndexListIn, TArray<uint32>& IndexListOut) const
 {
-	if (ContextObject && InputValue.IsValid())
+	if (InputValue.IsValid())
 	{
 		float Result = 0.0f;
-		InputValue.Get<FChooserParameterFloatBase>().GetValue(ContextObject, Result);
+		InputValue.Get<FChooserParameterFloatBase>().GetValue(Context, Result);
 
 #if WITH_EDITOR
-		if (DebugInfo.bCurrentDebugTarget)
+		if (Context.DebuggingInfo.bCurrentDebugTarget)
 		{
 			TestValue = Result;
 		}
diff --git a/Engine/Plugins/Experimental/Chooser/Source/Chooser/Private/GameplayTagColumn.cpp b/Engine/Plugins/Experimental/Chooser/Source/Chooser/Private/GameplayTagColumn.cpp
--- a/Engine/Plugins/Experimental/Chooser/Source/Chooser/Private/GameplayTagColumn.cpp
+++ b/Engine/Plugins/Experimental/Chooser/Source/Chooser/Private/GameplayTagColumn.cpp
@@ -2,12 +2,12 @@
 #include "GameplayTagColumn.h"
 #include "ChooserPropertyAccess.h"
 
-bool FGameplayTagContextProperty::GetValue(const UObject* ContextObject, const FGameplayTagContainer*& OutResult) const
+bool FGameplayTagContextProperty::GetValue(FChooserEvaluationContext& Context, const FGameplayTagContainer*& OutResult) const
 {
-	UStruct* StructType = ContextObject->GetClass();
-	const void* Container = ContextObject;
+	const UStruct* StructType = nullptr;
+	const void* Container = nullptr;
 	
-	if (UE::Chooser::ResolvePropertyChain(Container, StructType, Binding.PropertyBindingChain))
+	if (UE::Chooser::ResolvePropertyChain(Context, Binding, Container, StructType))
 	{
 		if (const FStructProperty* Property = FindFProperty<FStructProperty>(StructType, Binding.PropertyBindingChain.Last()))
 		{
@@ -24,13 +24,13 @@ FGameplayTagColumn::FGameplayTagColumn()
 	InputValue.InitializeAs(FGameplayTagContextProperty::StaticStruct());
 }
 
-void FGameplayTagColumn::Filter(FChooserDebuggingInfo& DebugInfo, const UObject* ContextObject, const TArray<uint32>& IndexListIn, TArray<uint32>& IndexListOut) const
+void FGameplayTagColumn::Filter(FChooserEvaluationContext& Context, const TArray<uint32>& IndexListIn, TArray<uint32>& IndexListOut) const
 {
 	const FGameplayTagContainer* Result = nullptr;
-	if (ContextObject && InputValue.IsValid() && InputValue.Get<FChooserParameterGameplayTagBase>().GetValue(ContextObject,Result))
+	if (InputValue.IsValid() && InputValue.Get<FChooserParameterGameplayTagBase>().GetValue(Context,Result))
 	{
 #if WITH_EDITOR
-		if (DebugInfo.bCurrentDebugTarget)
+		if (Context.DebuggingInfo.bCurrentDebugTarget)
 		{
 			TestValue = *Result;
 		}
diff --git a/Engine/Plugins/Experimental/Chooser/Source/Chooser/Private/ObjectChooser_Asset.cpp b/Engine/Plugins/Experimental/Chooser/Source/Chooser/Private/ObjectChooser_Asset.cpp
--- a/Engine/Plugins/Experimental/Chooser/Source/Chooser/Private/ObjectChooser_Asset.cpp
+++ b/Engine/Plugins/Experimental/Chooser/Source/Chooser/Private/ObjectChooser_Asset.cpp
@@ -2,7 +2,7 @@
 
 #include "ObjectChooser_Asset.h"
 
-UObject* FAssetChooser::ChooseObject(const UObject* ContextObject) const
+UObject* FAssetChooser::ChooseObject(FChooserEvaluationContext& Context) const
 {
 	return Asset;
 }
\ No newline at end of file
diff --git a/Engine/Plugins/Experimental/Chooser/Source/Chooser/Private/ObjectChooser_Class.cpp b/Engine/Plugins/Experimental/Chooser/Source/Chooser/Private/ObjectChooser_Class.cpp
--- a/Engine/Plugins/Experimental/Chooser/Source/Chooser/Private/ObjectChooser_Class.cpp
+++ b/Engine/Plugins/Experimental/Chooser/Source/Chooser/Private/ObjectChooser_Class.cpp
@@ -2,7 +2,7 @@
 
 #include "ObjectChooser_Class.h"
 
-UObject* FClassChooser::ChooseObject(const UObject* ContextObject) const
+UObject* FClassChooser::ChooseObject(FChooserEvaluationContext& Context) const
 {
 	return Class;
 }
\ No newline at end of file
diff --git a/Engine/Plugins/Experimental/Chooser/Source/Chooser/Private/ObjectColumn.cpp b/Engine/Plugins/Experimental/Chooser/Source/Chooser/Private/ObjectColumn.cpp
--- a/Engine/Plugins/Experimental/Chooser/Source/Chooser/Private/ObjectColumn.cpp
+++ b/Engine/Plugins/Experimental/Chooser/Source/Chooser/Private/ObjectColumn.cpp
@@ -6,12 +6,12 @@
 	#include "IPropertyAccessEditor.h"
 #endif
 
-bool FObjectContextProperty::GetValue(const UObject* ContextObject, FSoftObjectPath& OutResult) const
+bool FObjectContextProperty::GetValue(FChooserEvaluationContext& Context, FSoftObjectPath& OutResult) const
 {
-	UStruct* StructType = ContextObject->GetClass();
-	const void* Container = ContextObject;
+	const UStruct* StructType = nullptr;
+	const void* Container = nullptr;
 
-	if (UE::Chooser::ResolvePropertyChain(Container, StructType, Binding.PropertyBindingChain))
+	if (UE::Chooser::ResolvePropertyChain(Context, Binding, Container, StructType))
 	{
 		if (const FObjectPropertyBase* ObjectProperty = FindFProperty<FObjectPropertyBase>(StructType, Binding.PropertyBindingChain.Last()))
 		{
@@ -40,7 +40,7 @@ void FObjectContextProperty::SetBinding(const TArray<FBindingChainElement>& InBi
 	const UClass* PreviousClass = Binding.AllowedClass;
 	Binding.AllowedClass = nullptr;
 
-	UE::Chooser::CopyPropertyChain(InBindingChain, Binding.PropertyBindingChain);
+	UE::Chooser::CopyPropertyChain(InBindingChain, Binding);
 
 	const FField* Field = InBindingChain.Last().Field.ToField();
 	if (const FObjectPropertyBase* ObjectProperty = CastField<FObjectPropertyBase>(Field))
@@ -76,16 +76,15 @@ bool FChooserObjectRowData::Evaluate(const FSoftObjectPath& LeftHandSide) const
 	}
 }
 
-void FObjectColumn::Filter(FChooserDebuggingInfo& DebugInfo, const UObject* ContextObject, const TArray<uint32>& IndexListIn, TArray<uint32>& IndexListOut) const
+void FObjectColumn::Filter(FChooserEvaluationContext& Context, const TArray<uint32>& IndexListIn, TArray<uint32>& IndexListOut) const
 {
 	FSoftObjectPath Result;
-	if (ContextObject != nullptr &&
-		InputValue.IsValid() &&
-		InputValue.Get<FChooserParameterObjectBase>().GetValue(ContextObject, Result))
+	if (InputValue.IsValid() &&
+		InputValue.Get<FChooserParameterObjectBase>().GetValue(Context, Result))
 	{
 		
 #if WITH_EDITOR
-		if (DebugInfo.bCurrentDebugTarget)
+		if (Context.DebuggingInfo.bCurrentDebugTarget)
 		{
 			TestValue = Result;
 		}
diff --git a/Engine/Plugins/Experimental/Chooser/Source/Chooser/Private/OutputBoolColumn.cpp b/Engine/Plugins/Experimental/Chooser/Source/Chooser/Private/OutputBoolColumn.cpp
--- a/Engine/Plugins/Experimental/Chooser/Source/Chooser/Private/OutputBoolColumn.cpp
+++ b/Engine/Plugins/Experimental/Chooser/Source/Chooser/Private/OutputBoolColumn.cpp
@@ -8,15 +8,15 @@ FOutputBoolColumn::FOutputBoolColumn()
 	InputValue.InitializeAs(FBoolContextProperty::StaticStruct());
 }
 
-void FOutputBoolColumn::SetOutputs(FChooserDebuggingInfo& DebugInfo, UObject* ContextObject, int RowIndex) const
+void FOutputBoolColumn::SetOutputs(FChooserEvaluationContext& Context, int RowIndex) const
 {
 	if (InputValue.IsValid())
 	{
-		InputValue.Get<FChooserParameterBoolBase>().SetValue(ContextObject, RowValues[RowIndex]);
+		InputValue.Get<FChooserParameterBoolBase>().SetValue(Context, RowValues[RowIndex]);
 	}
 	
 #if WITH_EDITOR
-	if (DebugInfo.bCurrentDebugTarget)
+	if (Context.DebuggingInfo.bCurrentDebugTarget)
 	{
 		TestValue = RowValues[RowIndex];
 	}
diff --git a/Engine/Plugins/Experimental/Chooser/Source/Chooser/Private/RandomizeColumn.cpp b/Engine/Plugins/Experimental/Chooser/Source/Chooser/Private/RandomizeColumn.cpp
--- a/Engine/Plugins/Experimental/Chooser/Source/Chooser/Private/RandomizeColumn.cpp
+++ b/Engine/Plugins/Experimental/Chooser/Source/Chooser/Private/RandomizeColumn.cpp
@@ -2,13 +2,13 @@
 #include "RandomizeColumn.h"
 #include "ChooserPropertyAccess.h"
 
-bool FRandomizeContextProperty::GetValue(const UObject* ContextObject, const FChooserRandomizationContext*& OutResult) const
+bool FRandomizeContextProperty::GetValue(FChooserEvaluationContext& Context, const FChooserRandomizationContext*& OutResult) const
 {
 	
-	UStruct* StructType = ContextObject->GetClass();
-	const void* Container = ContextObject;
+	const UStruct* StructType = nullptr;
+	const void* Container = nullptr;
 	
-	if (UE::Chooser::ResolvePropertyChain(Container, StructType, Binding.PropertyBindingChain))
+	if (UE::Chooser::ResolvePropertyChain(Context, Binding, Container, StructType))
 	{
 		if (FStructProperty* Property = FindFProperty<FStructProperty>(StructType, Binding.PropertyBindingChain.Last()))
 		{
@@ -25,15 +25,15 @@ FRandomizeColumn::FRandomizeColumn()
 	InputValue.InitializeAs(FRandomizeContextProperty::StaticStruct());
 }
 
-void FRandomizeColumn::Filter(FChooserDebuggingInfo& DebuggingInfo, const UObject* ContextObject, const TArray<uint32>& IndexListIn, TArray<uint32>& IndexListOut) const
+void FRandomizeColumn::Filter(FChooserEvaluationContext& Context, const TArray<uint32>& IndexListIn, TArray<uint32>& IndexListOut) const
 {
 	int Count = IndexListIn.Num();
 	int Selection = 0;
 
 	const FChooserRandomizationContext* RandomizationContext = nullptr;
-	if (ContextObject && InputValue.IsValid())
+	if (InputValue.IsValid())
 	{
-		InputValue.Get<FChooserParameterRandomizeBase>().GetValue(ContextObject,RandomizationContext);
+		InputValue.Get<FChooserParameterRandomizeBase>().GetValue(Context,RandomizationContext);
 	}
 	
 	if (Count > 1)
@@ -99,12 +99,12 @@ void FRandomizeColumn::Filter(FChooserDebuggingInfo& DebuggingInfo, const UObjec
 	}
 }
 
-void FRandomizeColumn::SetOutputs(FChooserDebuggingInfo& DebugInfo, UObject* ContextObject, int RowIndex) const
+void FRandomizeColumn::SetOutputs(FChooserEvaluationContext& Context, int RowIndex) const
 {
-	if (ContextObject && InputValue.IsValid())
+	if (InputValue.IsValid())
 	{
 		const FChooserRandomizationContext* ConstRandomizationContext = nullptr;
-		InputValue.Get<FChooserParameterRandomizeBase>().GetValue(ContextObject,ConstRandomizationContext);
+		InputValue.Get<FChooserParameterRandomizeBase>().GetValue(Context,ConstRandomizationContext);
 		if(ConstRandomizationContext)
 		{
 			FChooserRandomizationContext* RandomizationContext = const_cast<FChooserRandomizationContext*>(ConstRandomizationContext);
diff --git a/Engine/Plugins/Experimental/Chooser/Source/Chooser/Public/ChooserPropertyAccess.h b/Engine/Plugins/Experimental/Chooser/Source/Chooser/Public/ChooserPropertyAccess.h
--- a/Engine/Plugins/Experimental/Chooser/Source/Chooser/Public/ChooserPropertyAccess.h
+++ b/Engine/Plugins/Experimental/Chooser/Source/Chooser/Public/ChooserPropertyAccess.h
@@ -4,6 +4,8 @@
 
 #include "CoreMinimal.h"
 #include "UObject/Interface.h"
+#include "InstancedStruct.h"
+#include "IObjectChooser.h"
 
 #include "ChooserPropertyAccess.generated.h"
 
@@ -17,14 +19,14 @@ class UHasContextClass : public UInterface
 	GENERATED_BODY()
 };
 
-DECLARE_MULTICAST_DELEGATE_OneParam(FContextClassChanged, UClass*)
+DECLARE_MULTICAST_DELEGATE(FContextClassChanged)
 
 class IHasContextClass
 {
 	GENERATED_BODY()
 public:
 	FContextClassChanged OnContextClassChanged;
-	virtual UClass* GetContextClass() { return nullptr; };
+	virtual TConstArrayView<FInstancedStruct> GetContextData() const { return TConstArrayView<FInstancedStruct>(); }
 };
 
 USTRUCT()
@@ -34,6 +36,9 @@ struct FChooserPropertyBinding
 	
 	UPROPERTY()
 	TArray<FName> PropertyBindingChain;
+	
+	UPROPERTY()
+	int ContextIndex = 0;
 };
 
 USTRUCT()
@@ -58,11 +63,50 @@ struct FChooserObjectPropertyBinding : public FChooserPropertyBinding
 #endif
 };
 
+
+DECLARE_MULTICAST_DELEGATE_OneParam(FChooserOutputObjectTypeChanged, const UClass* OutputObjectType);
+UENUM()
+enum class EContextObjectDirection
+{
+	Read,
+	Write,
+	ReadWrite
+};
+
+USTRUCT()
+struct FContextObjectTypeBase
+{
+	GENERATED_BODY()
+	
+	UPROPERTY(EditAnywhere, Category="Type")
+	EContextObjectDirection Direction;
+};
+
+
+USTRUCT()
+struct FContextObjectTypeClass : public FContextObjectTypeBase
+{
+	GENERATED_BODY()
+	UPROPERTY(EditAnywhere, Category="Type")
+	TObjectPtr<UClass> Class;
+};
+
+USTRUCT()
+struct FContextObjectTypeStruct : public FContextObjectTypeBase
+{
+	GENERATED_BODY()
+	
+	UPROPERTY(EditAnywhere, Category="Type")
+	TObjectPtr<UStruct> Struct;
+};
+
+
 namespace UE::Chooser
 {
-	CHOOSER_API bool ResolvePropertyChain(const void*& Container, UStruct*& StructType, const TArray<FName>& PropertyBindingChain);
+	CHOOSER_API bool ResolvePropertyChain(const void*& Container, const UStruct*& StructType, const TArray<FName>& PropertyBindingChain);
+	CHOOSER_API bool ResolvePropertyChain(FChooserEvaluationContext& Context, const FChooserPropertyBinding& Binding, const void*& OutContainer, const UStruct*& OutStructType);
 
 #if WITH_EDITOR
-	CHOOSER_API void CopyPropertyChain(const TArray<FBindingChainElement>& InBindingChain, TArray<FName>& OutPropertyBindingChain);
+	CHOOSER_API void CopyPropertyChain(const TArray<FBindingChainElement>& InBindingChain, FChooserPropertyBinding& OutPropertyBinding);
 #endif
 }
\ No newline at end of file
diff --git a/Engine/Plugins/Experimental/Chooser/Source/Chooser/Public/IChooserColumn.h b/Engine/Plugins/Experimental/Chooser/Source/Chooser/Public/IChooserColumn.h
--- a/Engine/Plugins/Experimental/Chooser/Source/Chooser/Public/IChooserColumn.h
+++ b/Engine/Plugins/Experimental/Chooser/Source/Chooser/Public/IChooserColumn.h
@@ -5,6 +5,7 @@
 #include "CoreMinimal.h"
 #include "UObject/Object.h"
 #include "UObject/Interface.h"
+#include "IObjectChooser.h"
 #include "IChooserParameterBase.h"
 #include "InstancedStruct.h"
 #include "IChooserColumn.generated.h"
@@ -24,11 +25,6 @@ public:
 
 class UChooserTable;
 
-struct FChooserDebuggingInfo
-{
-	bool bCurrentDebugTarget;
-};
-
 USTRUCT()
 struct CHOOSER_API FChooserColumnBase
 {
@@ -37,12 +33,12 @@ struct CHOOSER_API FChooserColumnBase
 public:
 	virtual ~FChooserColumnBase() {}
 	virtual void PostLoad() {};
-	virtual void Filter(FChooserDebuggingInfo& DebugInfo, const UObject* ContextObject, const TArray<uint32>& IndexListIn, TArray<uint32>& IndexListOut) const {}
+	virtual void Filter(FChooserEvaluationContext& Context, const TArray<uint32>& IndexListIn, TArray<uint32>& IndexListOut) const {}
 
 	virtual bool HasFilters() const { return true; }
 	virtual bool HasOutputs() const { return false; }
 	
-	virtual void SetOutputs(FChooserDebuggingInfo& DebugInfo, UObject* ContextObject, int RowIndex) const { }
+	virtual void SetOutputs(FChooserEvaluationContext& Context, int RowIndex) const { }
 
 #if WITH_EDITOR
 	virtual FName RowValuesPropertyName() { return FName(); }
diff --git a/Engine/Plugins/Experimental/Chooser/Source/Chooser/Public/IChooserParameterBase.h b/Engine/Plugins/Experimental/Chooser/Source/Chooser/Public/IChooserParameterBase.h
--- a/Engine/Plugins/Experimental/Chooser/Source/Chooser/Public/IChooserParameterBase.h
+++ b/Engine/Plugins/Experimental/Chooser/Source/Chooser/Public/IChooserParameterBase.h
@@ -3,8 +3,6 @@
 #pragma once
 
 #include "CoreMinimal.h"
-#include "UObject/Object.h"
-#include "UObject/Interface.h"
 #include "IChooserParameterBase.generated.h"
 
 USTRUCT()
diff --git a/Engine/Plugins/Experimental/Chooser/Source/Chooser/Public/IChooserParameterBool.h b/Engine/Plugins/Experimental/Chooser/Source/Chooser/Public/IChooserParameterBool.h
--- a/Engine/Plugins/Experimental/Chooser/Source/Chooser/Public/IChooserParameterBool.h
+++ b/Engine/Plugins/Experimental/Chooser/Source/Chooser/Public/IChooserParameterBool.h
@@ -28,6 +28,6 @@ struct FChooserParameterBoolBase : public FChooserParameterBase
 	GENERATED_BODY()
     
     public:
-    	virtual bool GetValue(const UObject* ContextObject, bool& OutResult) const { return false; }
-    	virtual bool SetValue(UObject* ContextObject, bool InValue) const { return false; }
+    	virtual bool GetValue(FChooserEvaluationContext& Context, bool& OutResult) const { return false; }
+    	virtual bool SetValue(FChooserEvaluationContext& Context, bool InValue) const { return false; }
 };
\ No newline at end of file
diff --git a/Engine/Plugins/Experimental/Chooser/Source/Chooser/Public/IChooserParameterEnum.h b/Engine/Plugins/Experimental/Chooser/Source/Chooser/Public/IChooserParameterEnum.h
--- a/Engine/Plugins/Experimental/Chooser/Source/Chooser/Public/IChooserParameterEnum.h
+++ b/Engine/Plugins/Experimental/Chooser/Source/Chooser/Public/IChooserParameterEnum.h
@@ -3,7 +3,6 @@
 #pragma once
 
 #include "CoreMinimal.h"
-#include "UObject/Object.h"
 #include "UObject/Interface.h"
 #include "InstancedStruct.h"
 #include "IChooserParameterBase.h"
@@ -27,7 +26,7 @@ struct FChooserParameterEnumBase : public FChooserParameterBase
 {
 	GENERATED_BODY()
 	
-		virtual bool GetValue(const UObject* ContextObject, uint8& OutResult) const { return false; }
+		virtual bool GetValue(FChooserEvaluationContext& Context, uint8& OutResult) const { return false; }
 
 	#if WITH_EDITOR
     	virtual const UEnum* GetEnum() const { return nullptr; }
diff --git a/Engine/Plugins/Experimental/Chooser/Source/Chooser/Public/IChooserParameterFloat.h b/Engine/Plugins/Experimental/Chooser/Source/Chooser/Public/IChooserParameterFloat.h
--- a/Engine/Plugins/Experimental/Chooser/Source/Chooser/Public/IChooserParameterFloat.h
+++ b/Engine/Plugins/Experimental/Chooser/Source/Chooser/Public/IChooserParameterFloat.h
@@ -28,5 +28,5 @@ struct FChooserParameterFloatBase : public FChooserParameterBase
 	GENERATED_BODY()
     
 public:
-	virtual bool GetValue(const UObject* ContextObject, float& OutResult) const { return false; }
+	virtual bool GetValue(FChooserEvaluationContext& Context, float& OutResult) const { return false; }
 };
\ No newline at end of file
diff --git a/Engine/Plugins/Experimental/Chooser/Source/Chooser/Public/IChooserParameterGameplayTag.h b/Engine/Plugins/Experimental/Chooser/Source/Chooser/Public/IChooserParameterGameplayTag.h
--- a/Engine/Plugins/Experimental/Chooser/Source/Chooser/Public/IChooserParameterGameplayTag.h
+++ b/Engine/Plugins/Experimental/Chooser/Source/Chooser/Public/IChooserParameterGameplayTag.h
@@ -28,5 +28,5 @@ USTRUCT()
 struct FChooserParameterGameplayTagBase : public FChooserParameterBase
 {
 	GENERATED_BODY()
-	virtual bool GetValue(const UObject* ContextObject, const FGameplayTagContainer*& OutResult) const { return false; }
+	virtual bool GetValue(FChooserEvaluationContext& Context, const FGameplayTagContainer*& OutResult) const { return false; }
 };
diff --git a/Engine/Plugins/Experimental/Chooser/Source/Chooser/Public/IChooserParameterObject.h b/Engine/Plugins/Experimental/Chooser/Source/Chooser/Public/IChooserParameterObject.h
--- a/Engine/Plugins/Experimental/Chooser/Source/Chooser/Public/IChooserParameterObject.h
+++ b/Engine/Plugins/Experimental/Chooser/Source/Chooser/Public/IChooserParameterObject.h
@@ -13,7 +13,7 @@ struct FChooserParameterObjectBase : public FChooserParameterBase
 {
 	GENERATED_BODY()
 
-	virtual bool GetValue(const UObject* ContextObject, FSoftObjectPath& OutResult) const { return false; }
+	virtual bool GetValue(FChooserEvaluationContext& Context, FSoftObjectPath& OutResult) const { return false; }
 
 #if WITH_EDITOR
 	virtual UClass* GetAllowedClass() const { return UObject::StaticClass(); }
diff --git a/Engine/Plugins/Experimental/Chooser/Source/Chooser/Public/IChooserParameterRandomize.h b/Engine/Plugins/Experimental/Chooser/Source/Chooser/Public/IChooserParameterRandomize.h
--- a/Engine/Plugins/Experimental/Chooser/Source/Chooser/Public/IChooserParameterRandomize.h
+++ b/Engine/Plugins/Experimental/Chooser/Source/Chooser/Public/IChooserParameterRandomize.h
@@ -26,5 +26,5 @@ USTRUCT()
 struct FChooserParameterRandomizeBase : public FChooserParameterBase
 {
 	GENERATED_BODY()
-	virtual bool GetValue(const UObject* ContextObject, const FChooserRandomizationContext*& OutResult) const { return false; }
+	virtual bool GetValue(FChooserEvaluationContext& Context, const FChooserRandomizationContext*& OutResult) const { return false; }
 };
diff --git a/Engine/Plugins/Experimental/Chooser/Source/Chooser/Public/IObjectChooser.h b/Engine/Plugins/Experimental/Chooser/Source/Chooser/Public/IObjectChooser.h
--- a/Engine/Plugins/Experimental/Chooser/Source/Chooser/Public/IObjectChooser.h
+++ b/Engine/Plugins/Experimental/Chooser/Source/Chooser/Public/IObjectChooser.h
@@ -21,6 +21,27 @@ public:
 	virtual void ConvertToInstancedStruct(FInstancedStruct& OutInstancedStruct) const { }
 };
 
+struct FContextEntry
+{
+	const UStruct* Type = nullptr;
+	void* Data = nullptr;
+};
+
+#if WITH_EDITOR
+struct FChooserDebuggingInfo
+{
+	bool bCurrentDebugTarget = false;
+};
+#endif
+
+struct FChooserEvaluationContext
+{
+#if WITH_EDITOR
+	FChooserDebuggingInfo DebuggingInfo;
+#endif
+
+	TArray<FContextEntry, TFixedAllocator<4>> ContextData;
+};
 
 USTRUCT()
 struct CHOOSER_API FObjectChooserBase
@@ -30,18 +51,18 @@ struct CHOOSER_API FObjectChooserBase
 public:
 	virtual ~FObjectChooserBase() {}
 
-	virtual UObject* ChooseObject(const UObject* ContextObject) const { return nullptr; };
-
 	enum class EIteratorStatus { Continue,Stop };
 
 	DECLARE_DELEGATE_RetVal_OneParam( EIteratorStatus, FObjectChooserIteratorCallback, UObject*);
-	virtual EIteratorStatus ChooseMulti(const UObject* ContextObject, FObjectChooserIteratorCallback Callback) const
+
+	virtual UObject* ChooseObject(FChooserEvaluationContext& Context) const { return nullptr; };
+	virtual EIteratorStatus ChooseMulti(FChooserEvaluationContext& ContextData, FObjectChooserIteratorCallback Callback) const
 	{
 		// fallback implementation just calls the single version
-		if (UObject* Result = ChooseObject(ContextObject))
+		if (UObject* Result = ChooseObject(ContextData))
 		{
 			return Callback.Execute(Result);
 		}
 		return EIteratorStatus::Continue;
-	};
+	}
 };
\ No newline at end of file
diff --git a/Engine/Plugins/Experimental/Chooser/Source/ChooserEditor/Private/BoolColumnEditor.cpp b/Engine/Plugins/Experimental/Chooser/Source/ChooserEditor/Private/BoolColumnEditor.cpp
--- a/Engine/Plugins/Experimental/Chooser/Source/ChooserEditor/Private/BoolColumnEditor.cpp
+++ b/Engine/Plugins/Experimental/Chooser/Source/ChooserEditor/Private/BoolColumnEditor.cpp
@@ -26,7 +26,7 @@ TSharedRef<SWidget> CreateBoolColumnWidget(UChooserTable* Chooser, FChooserColum
 		TSharedPtr<SWidget> InputValueWidget = nullptr;
 		if (FChooserParameterBase* InputValue = Column->GetInputValue())
 		{
-			InputValueWidget = FObjectChooserWidgetFactories::CreateWidget(false, Chooser, InputValue, Column->GetInputType(), Chooser->ContextObjectType, Chooser->OutputObjectType);
+			InputValueWidget = FObjectChooserWidgetFactories::CreateWidget(false, Chooser, InputValue, Column->GetInputType(), Chooser->OutputObjectType);
 		}
 	
 		const FSlateBrush* ColumnIcon = FCoreStyle::Get().GetBrush("Icons.Filter");
@@ -116,7 +116,7 @@ TSharedRef<SWidget> CreateOutputBoolColumnWidget(UChooserTable* Chooser, FChoose
 		TSharedPtr<SWidget> InputValueWidget = nullptr;
 		if (FChooserParameterBase* InputValue = Column->GetInputValue())
 		{
-			InputValueWidget = FObjectChooserWidgetFactories::CreateWidget(false, Chooser, InputValue, Column->GetInputType(), Chooser->ContextObjectType, Chooser->OutputObjectType);
+			InputValueWidget = FObjectChooserWidgetFactories::CreateWidget(false, Chooser, InputValue, Column->GetInputType(), Chooser->OutputObjectType);
 		}
 		
 		const FSlateBrush* ColumnIcon = FCoreStyle::Get().GetBrush("Icons.ArrowRight");
@@ -183,7 +183,7 @@ TSharedRef<SWidget> CreateOutputBoolColumnWidget(UChooserTable* Chooser, FChoose
 	+ SHorizontalBox::Slot().FillWidth(1);
 }
 	
-TSharedRef<SWidget> CreateBoolPropertyWidget(bool bReadOnly, UObject* TransactionObject, void* Value, UClass* ContextClass, UClass* ResultBaseClass)
+TSharedRef<SWidget> CreateBoolPropertyWidget(bool bReadOnly, UObject* TransactionObject, void* Value, UClass* ResultBaseClass)
 {
 	IHasContextClass* HasContextClass = Cast<IHasContextClass>(TransactionObject);
 
diff --git a/Engine/Plugins/Experimental/Chooser/Source/ChooserEditor/Private/ChooserTableEditor.cpp b/Engine/Plugins/Experimental/Chooser/Source/ChooserEditor/Private/ChooserTableEditor.cpp
--- a/Engine/Plugins/Experimental/Chooser/Source/ChooserEditor/Private/ChooserTableEditor.cpp
+++ b/Engine/Plugins/Experimental/Chooser/Source/ChooserEditor/Private/ChooserTableEditor.cpp
@@ -581,13 +581,13 @@ public:
 			}
 			else if (ColumnName == Result) 
 			{
-				TSharedPtr<SWidget> ResultWidget = FObjectChooserWidgetFactories::CreateWidget(false, Chooser, FObjectChooserBase::StaticStruct(), Chooser->ResultsStructs[RowIndex->RowIndex].GetMutableMemory(), Chooser->ResultsStructs[RowIndex->RowIndex].GetScriptStruct(),Chooser->ContextObjectType, Chooser->OutputObjectType,
+				TSharedPtr<SWidget> ResultWidget = FObjectChooserWidgetFactories::CreateWidget(false, Chooser, FObjectChooserBase::StaticStruct(), Chooser->ResultsStructs[RowIndex->RowIndex].GetMutableMemory(), Chooser->ResultsStructs[RowIndex->RowIndex].GetScriptStruct(), Chooser->OutputObjectType,
 				FOnStructPicked::CreateLambda([this, RowIndex=RowIndex->RowIndex](const UScriptStruct* ChosenStruct)
 				{
 					const FScopedTransaction Transaction(LOCTEXT("Change Row Result Type", "Change Row Result Type"));
 					Chooser->Modify(true);
 					Chooser->ResultsStructs[RowIndex].InitializeAs(ChosenStruct);
-					FObjectChooserWidgetFactories::CreateWidget(false, Chooser, FObjectChooserBase::StaticStruct(), Chooser->ResultsStructs[RowIndex].GetMutableMemory(), ChosenStruct, Chooser->ContextObjectType, Chooser->OutputObjectType, FOnStructPicked(), &CacheBorder);
+					FObjectChooserWidgetFactories::CreateWidget(false, Chooser, FObjectChooserBase::StaticStruct(), Chooser->ResultsStructs[RowIndex].GetMutableMemory(), ChosenStruct, Chooser->OutputObjectType, FOnStructPicked(), &CacheBorder);
 				}),
 				&CacheBorder
 				);
@@ -1209,7 +1209,7 @@ void FChooserTableEditor::DeleteColumn(int Index)
 	}
 }
 
-TSharedRef<SWidget> CreateAssetWidget(bool bReadOnly, UObject* TransactionObject, void* Value, UClass* ContextClass, UClass* ResultBaseClass)
+TSharedRef<SWidget> CreateAssetWidget(bool bReadOnly, UObject* TransactionObject, void* Value, UClass* ResultBaseClass)
 {
 	FAssetChooser* DIAsset = static_cast<FAssetChooser*>(Value);
 
@@ -1227,7 +1227,7 @@ TSharedRef<SWidget> CreateAssetWidget(bool bReadOnly, UObject* TransactionObject
 		});
 }
 	
-TSharedRef<SWidget> CreateClassWidget(bool bReadOnly, UObject* TransactionObject, void* Value, UClass* ContextClas, UClass* ResultBaseClass)
+TSharedRef<SWidget> CreateClassWidget(bool bReadOnly, UObject* TransactionObject, void* Value, UClass* ResultBaseClass)
 {
 	FClassChooser* ClassChooser = static_cast<FClassChooser*>(Value);
 
@@ -1248,7 +1248,7 @@ TSharedRef<SWidget> CreateClassWidget(bool bReadOnly, UObject* TransactionObject
 			});
 }
 
-TSharedRef<SWidget> CreateEvaluateChooserWidget(bool bReadOnly, UObject* TransactionObject, void* Value, UClass* ContextObject, UClass* ResultBaseClass)
+TSharedRef<SWidget> CreateEvaluateChooserWidget(bool bReadOnly, UObject* TransactionObject, void* Value, UClass* ResultBaseClass)
 {
 	FEvaluateChooser* EvaluateChooser = static_cast<FEvaluateChooser*>(Value);
 	
diff --git a/Engine/Plugins/Experimental/Chooser/Source/ChooserEditor/Private/EnumColumnEditor.cpp b/Engine/Plugins/Experimental/Chooser/Source/ChooserEditor/Private/EnumColumnEditor.cpp
--- a/Engine/Plugins/Experimental/Chooser/Source/ChooserEditor/Private/EnumColumnEditor.cpp
+++ b/Engine/Plugins/Experimental/Chooser/Source/ChooserEditor/Private/EnumColumnEditor.cpp
@@ -179,7 +179,7 @@ TSharedRef<SWidget> CreateEnumColumnWidget(UChooserTable* Chooser, FChooserColum
 		TSharedPtr<SWidget> InputValueWidget = nullptr;
 		if (FChooserParameterBase* InputValue = Column->GetInputValue())
 		{
-			InputValueWidget = FObjectChooserWidgetFactories::CreateWidget(false, Chooser, InputValue, Column->GetInputType(), Chooser->ContextObjectType, Chooser->OutputObjectType);
+			InputValueWidget = FObjectChooserWidgetFactories::CreateWidget(false, Chooser, InputValue, Column->GetInputType(), Chooser->OutputObjectType);
 		}
 		
 		const FSlateBrush* ColumnIcon = FCoreStyle::Get().GetBrush("Icons.Filter");
@@ -220,7 +220,7 @@ TSharedRef<SWidget> CreateEnumColumnWidget(UChooserTable* Chooser, FChooserColum
 	return SNew(SEnumCell).TransactionObject(Chooser).EnumColumn(EnumColumn).RowIndex(Row);
 }
 
-TSharedRef<SWidget> CreateEnumPropertyWidget(bool bReadOnly, UObject* TransactionObject, void* Value, UClass* ContextClass, UClass* ResultBaseClass)
+TSharedRef<SWidget> CreateEnumPropertyWidget(bool bReadOnly, UObject* TransactionObject, void* Value, UClass* ResultBaseClass)
 {
 	IHasContextClass* HasContextClass = Cast<IHasContextClass>(TransactionObject);
 
diff --git a/Engine/Plugins/Experimental/Chooser/Source/ChooserEditor/Private/FloatRangeColumnEditor.cpp b/Engine/Plugins/Experimental/Chooser/Source/ChooserEditor/Private/FloatRangeColumnEditor.cpp
--- a/Engine/Plugins/Experimental/Chooser/Source/ChooserEditor/Private/FloatRangeColumnEditor.cpp
+++ b/Engine/Plugins/Experimental/Chooser/Source/ChooserEditor/Private/FloatRangeColumnEditor.cpp
@@ -13,7 +13,7 @@
 
 namespace UE::ChooserEditor
 {
-TSharedRef<SWidget> CreateFloatPropertyWidget(bool bReadOnly, UObject* TransactionObject, void* Value, UClass* ContextClass, UClass* ResultBaseClass)
+TSharedRef<SWidget> CreateFloatPropertyWidget(bool bReadOnly, UObject* TransactionObject, void* Value, UClass* ResultBaseClass)
 {
 	IHasContextClass* HasContextClass = Cast<IHasContextClass>(TransactionObject);
 
@@ -40,7 +40,7 @@ TSharedRef<SWidget> CreateFloatRangeColumnWidget(UChooserTable* Chooser, FChoose
 		TSharedPtr<SWidget> InputValueWidget = nullptr;
 		if (FChooserParameterBase* InputValue = Column->GetInputValue())
 		{
-			InputValueWidget = FObjectChooserWidgetFactories::CreateWidget(false, Chooser, InputValue, Column->GetInputType(), Chooser->ContextObjectType, Chooser->OutputObjectType);
+			InputValueWidget = FObjectChooserWidgetFactories::CreateWidget(false, Chooser, InputValue, Column->GetInputType(), Chooser->OutputObjectType);
 		}
 		
 		const FSlateBrush* ColumnIcon = FCoreStyle::Get().GetBrush("Icons.Filter");
diff --git a/Engine/Plugins/Experimental/Chooser/Source/ChooserEditor/Private/GameplayTagColumnEditor.cpp b/Engine/Plugins/Experimental/Chooser/Source/ChooserEditor/Private/GameplayTagColumnEditor.cpp
--- a/Engine/Plugins/Experimental/Chooser/Source/ChooserEditor/Private/GameplayTagColumnEditor.cpp
+++ b/Engine/Plugins/Experimental/Chooser/Source/ChooserEditor/Private/GameplayTagColumnEditor.cpp
@@ -24,7 +24,7 @@ TSharedRef<SWidget> CreateGameplayTagColumnWidget(UChooserTable* Chooser, FChoos
 		TSharedPtr<SWidget> InputValueWidget = nullptr;
 		if (FChooserParameterBase* InputValue = Column->GetInputValue())
 		{
-			InputValueWidget = FObjectChooserWidgetFactories::CreateWidget(false, Chooser, InputValue, Column->GetInputType(), Chooser->ContextObjectType, Chooser->OutputObjectType);
+			InputValueWidget = FObjectChooserWidgetFactories::CreateWidget(false, Chooser, InputValue, Column->GetInputType(), Chooser->OutputObjectType);
 		}
 		
 		const FSlateBrush* ColumnIcon = FCoreStyle::Get().GetBrush("Icons.Filter");
@@ -111,7 +111,7 @@ TSharedRef<SWidget> CreateGameplayTagColumnWidget(UChooserTable* Chooser, FChoos
 	);
 }
 
-TSharedRef<SWidget> CreateGameplayTagPropertyWidget(bool bReadOnly, UObject* TransactionObject, void* Value, UClass* ContextClass, UClass* ResultBaseClass)
+TSharedRef<SWidget> CreateGameplayTagPropertyWidget(bool bReadOnly, UObject* TransactionObject, void* Value, UClass* ResultBaseClass)
 {
 	IHasContextClass* HasContextClass = Cast<IHasContextClass>(TransactionObject);
 
diff --git a/Engine/Plugins/Experimental/Chooser/Source/ChooserEditor/Private/ObjectChooserWidgetFactories.cpp b/Engine/Plugins/Experimental/Chooser/Source/ChooserEditor/Private/ObjectChooserWidgetFactories.cpp
--- a/Engine/Plugins/Experimental/Chooser/Source/ChooserEditor/Private/ObjectChooserWidgetFactories.cpp
+++ b/Engine/Plugins/Experimental/Chooser/Source/ChooserEditor/Private/ObjectChooserWidgetFactories.cpp
@@ -46,7 +46,7 @@ TSharedPtr<SWidget> FObjectChooserWidgetFactories::CreateColumnWidget(FChooserCo
 }
 	
 	
-TSharedPtr<SWidget> FObjectChooserWidgetFactories::CreateWidget(bool bReadOnly, UObject* TransactionObject, void* Value, const UStruct* ValueType, UClass* ContextClass, UClass* ResultBaseClass)
+TSharedPtr<SWidget> FObjectChooserWidgetFactories::CreateWidget(bool bReadOnly, UObject* TransactionObject, void* Value, const UStruct* ValueType, UClass* ResultBaseClass)
 {
 	if (Value)
 	{
@@ -54,7 +54,7 @@ TSharedPtr<SWidget> FObjectChooserWidgetFactories::CreateWidget(bool bReadOnly,
 		{
 			if (FChooserWidgetCreator* Creator = ChooserWidgetCreators.Find(ValueType))
 			{
-				return (*Creator)(bReadOnly, TransactionObject, Value, ContextClass, ResultBaseClass);
+				return (*Creator)(bReadOnly, TransactionObject, Value, ResultBaseClass);
 			}
 			ValueType = ValueType->GetSuperStruct();
 		}
@@ -63,9 +63,9 @@ TSharedPtr<SWidget> FObjectChooserWidgetFactories::CreateWidget(bool bReadOnly,
 	return nullptr;
 }
 
-TSharedPtr<SWidget> FObjectChooserWidgetFactories::CreateWidget(bool bReadOnly, UObject* TransactionObject, const UScriptStruct* BaseType, void* Value, const UStruct* ValueType, UClass* ContextClass, UClass* ResultBaseClass, const FOnStructPicked& CreateClassCallback, TSharedPtr<SBorder>* InnerWidget)
+TSharedPtr<SWidget> FObjectChooserWidgetFactories::CreateWidget(bool bReadOnly, UObject* TransactionObject, const UScriptStruct* BaseType, void* Value, const UStruct* ValueType, UClass* ResultBaseClass, const FOnStructPicked& CreateClassCallback, TSharedPtr<SBorder>* InnerWidget)
 {
-	TSharedPtr<SWidget> LeftWidget = CreateWidget(bReadOnly, TransactionObject, Value, ValueType, ContextClass, ResultBaseClass);
+	TSharedPtr<SWidget> LeftWidget = CreateWidget(bReadOnly, TransactionObject, Value, ValueType, ResultBaseClass);
 	if (bReadOnly)
 	{
 		// don't need the type selector dropdown when read only
diff --git a/Engine/Plugins/Experimental/Chooser/Source/ChooserEditor/Private/ObjectColumnEditor.cpp b/Engine/Plugins/Experimental/Chooser/Source/ChooserEditor/Private/ObjectColumnEditor.cpp
--- a/Engine/Plugins/Experimental/Chooser/Source/ChooserEditor/Private/ObjectColumnEditor.cpp
+++ b/Engine/Plugins/Experimental/Chooser/Source/ChooserEditor/Private/ObjectColumnEditor.cpp
@@ -98,7 +98,7 @@ namespace UE::ChooserEditor
 			TSharedPtr<SWidget> InputValueWidget = nullptr;
 			if (FChooserParameterBase* InputValue = Column->GetInputValue())
 			{
-				InputValueWidget = FObjectChooserWidgetFactories::CreateWidget(false, Chooser, InputValue, Column->GetInputType(), Chooser->ContextObjectType, Chooser->OutputObjectType);
+				InputValueWidget = FObjectChooserWidgetFactories::CreateWidget(false, Chooser, InputValue, Column->GetInputType(), Chooser->OutputObjectType);
 			}
 			
 			const FSlateBrush* ColumnIcon = FCoreStyle::Get().GetBrush("Icons.Filter");
@@ -215,7 +215,7 @@ namespace UE::ChooserEditor
 		return CellWidget;
 	}
 
-	static TSharedRef<SWidget> CreateObjectPropertyWidget(bool bReadOnly, UObject* TransactionObject, void* Value, UClass* ContextClass, UClass* ResultBaseClass)
+	static TSharedRef<SWidget> CreateObjectPropertyWidget(bool bReadOnly, UObject* TransactionObject, void* Value, UClass* ResultBaseClass)
 	{
 		IHasContextClass* HasContextClass = Cast<IHasContextClass>(TransactionObject);
 
diff --git a/Engine/Plugins/Experimental/Chooser/Source/ChooserEditor/Private/PropertyAccessChainCustomization.cpp b/Engine/Plugins/Experimental/Chooser/Source/ChooserEditor/Private/PropertyAccessChainCustomization.cpp
--- a/Engine/Plugins/Experimental/Chooser/Source/ChooserEditor/Private/PropertyAccessChainCustomization.cpp
+++ b/Engine/Plugins/Experimental/Chooser/Source/ChooserEditor/Private/PropertyAccessChainCustomization.cpp
@@ -79,7 +79,7 @@ void FPropertyAccessChainCustomization::CustomizeHeader(TSharedRef<IPropertyHand
          	 				PropertyHandle->NotifyPreChange();
          	 				
          	 				OuterObjects[i]->Modify(true);
-         	 				Chooser::CopyPropertyChain(InBindingChain, PropertyValue->PropertyBindingChain);
+         	 				Chooser::CopyPropertyChain(InBindingChain, *PropertyValue);
          	
          	 				if (TypeFilter == "enum")
          	 				{
diff --git a/Engine/Plugins/Experimental/Chooser/Source/ChooserEditor/Private/RandomizeColumnEditor.cpp b/Engine/Plugins/Experimental/Chooser/Source/ChooserEditor/Private/RandomizeColumnEditor.cpp
--- a/Engine/Plugins/Experimental/Chooser/Source/ChooserEditor/Private/RandomizeColumnEditor.cpp
+++ b/Engine/Plugins/Experimental/Chooser/Source/ChooserEditor/Private/RandomizeColumnEditor.cpp
@@ -75,7 +75,7 @@ TSharedRef<SWidget> CreateRandomizeColumnWidget(UChooserTable* Chooser, FChooser
 		+ SHorizontalBox::Slot().FillWidth(1);
 }
 	
-TSharedRef<SWidget> CreateRandomizePropertyWidget(bool bReadOnly, UObject* TransactionObject, void* Value, UClass* ContextClass, UClass* ResultBaseClass)
+TSharedRef<SWidget> CreateRandomizePropertyWidget(bool bReadOnly, UObject* TransactionObject, void* Value, UClass* ResultBaseClass)
 {
 	IHasContextClass* HasContextClass = Cast<IHasContextClass>(TransactionObject);
 
diff --git a/Engine/Plugins/Experimental/Chooser/Source/ChooserEditor/Private/SPropertyAccessChainWidget.cpp b/Engine/Plugins/Experimental/Chooser/Source/ChooserEditor/Private/SPropertyAccessChainWidget.cpp
--- a/Engine/Plugins/Experimental/Chooser/Source/ChooserEditor/Private/SPropertyAccessChainWidget.cpp
+++ b/Engine/Plugins/Experimental/Chooser/Source/ChooserEditor/Private/SPropertyAccessChainWidget.cpp
@@ -23,12 +23,12 @@ TSharedRef<SWidget> SPropertyAccessChainWidget::CreatePropertyAccessWidget()
 	FPropertyBindingWidgetArgs Args;
 	Args.bAllowPropertyBindings = true;
 
-	UClass* ContextClass = UObject::StaticClass();
+	TConstArrayView<FInstancedStruct> ContextData;
 	if (ContextClassOwner)
 	{
-		ContextClass = ContextClassOwner->GetContextClass();
+		ContextData = ContextClassOwner->GetContextData();
 	}
-
+	
 	Args.bAllowUObjectFunctions = true;
 	Args.bAllowOnlyThreadSafeFunctions = true;
 
@@ -176,11 +176,24 @@ TSharedRef<SWidget> SPropertyAccessChainWidget::CreatePropertyAccessWidget()
 			static FName PropertyIcon(TEXT("Kismet.Tabs.Variables"));
 			return FAppStyle::GetBrush(PropertyIcon);
 		});
+	
+	TArray<FBindingContextStruct> ContextStructs;
+	for (const FInstancedStruct& ContextStruct : ContextData)
+	{
+		if (const FContextObjectTypeClass* ClassType = ContextStruct.GetPtr<FContextObjectTypeClass>())
+		{
+			ContextStructs.SetNum(ContextStructs.Num()+1);
+			ContextStructs.Last().Struct = ClassType->Class;
+		}
+		else if (const FContextObjectTypeStruct* StructType = ContextStruct.GetPtr<FContextObjectTypeStruct>())
+		{
+			ContextStructs.SetNum(ContextStructs.Num()+1);
+			ContextStructs.Last().Struct = StructType->Struct;
+		}
+	}	
 
 	const IPropertyAccessEditor& PropertyAccessEditor = IModularFeatures::Get().GetModularFeature<IPropertyAccessEditor>("PropertyAccessEditor");
-	FBindingContextStruct StructInfo;
-	StructInfo.Struct = ContextClass;
-	return  PropertyAccessEditor.MakePropertyBindingWidget({StructInfo}, Args);
+	return  PropertyAccessEditor.MakePropertyBindingWidget(ContextStructs, Args);
 }
 
 void SPropertyAccessChainWidget::UpdateWidget()
@@ -188,7 +201,7 @@ void SPropertyAccessChainWidget::UpdateWidget()
 	ChildSlot[ CreatePropertyAccessWidget() ];
 }
 
-void SPropertyAccessChainWidget::ContextClassChanged(UClass* NewContextClass)
+void SPropertyAccessChainWidget::ContextClassChanged()
 {
 	UpdateWidget();
 }
diff --git a/Engine/Plugins/Experimental/Chooser/Source/ChooserEditor/Public/ObjectChooserWidgetFactories.h b/Engine/Plugins/Experimental/Chooser/Source/ChooserEditor/Public/ObjectChooserWidgetFactories.h
--- a/Engine/Plugins/Experimental/Chooser/Source/ChooserEditor/Public/ObjectChooserWidgetFactories.h
+++ b/Engine/Plugins/Experimental/Chooser/Source/ChooserEditor/Public/ObjectChooserWidgetFactories.h
@@ -13,15 +13,15 @@ struct FChooserColumnBase;
 
 namespace UE::ChooserEditor
 {
-	typedef TFunction<TSharedRef<SWidget>(bool bReadOnly, UObject* TransactionObject, void* Value, UClass* ContextClass, UClass* ResultBaseClass)> FChooserWidgetCreator;
+	typedef TFunction<TSharedRef<SWidget>(bool bReadOnly, UObject* TransactionObject, void* Value, UClass* ResultBaseClass)> FChooserWidgetCreator;
 	typedef TFunction<TSharedRef<SWidget>(UChooserTable* Chooser, FChooserColumnBase* Column, int Row)> FColumnWidgetCreator;
 
 	class CHOOSEREDITOR_API FObjectChooserWidgetFactories
 	{
 	public:
 
-		static TSharedPtr<SWidget> CreateWidget(bool ReadOnly, UObject* TransactionObject, void* Value, const UStruct* ValueType, UClass* ContextClass, UClass* ResultBaseClass);
-		static TSharedPtr<SWidget> CreateWidget(bool ReadOnly, UObject* TransactionObject, const UScriptStruct* BaseType, void* Value, const UStruct* ValueType, UClass* ContextClass, UClass* ResultBaseClass, const FOnStructPicked& CreateClassCallback, TSharedPtr<SBorder>* InnerWidget = nullptr);
+		static TSharedPtr<SWidget> CreateWidget(bool ReadOnly, UObject* TransactionObject, void* Value, const UStruct* ValueType, UClass* ResultBaseClass);
+		static TSharedPtr<SWidget> CreateWidget(bool ReadOnly, UObject* TransactionObject, const UScriptStruct* BaseType, void* Value, const UStruct* ValueType, UClass* ResultBaseClass, const FOnStructPicked& CreateClassCallback, TSharedPtr<SBorder>* InnerWidget = nullptr);
 		
 		static TSharedPtr<SWidget> CreateColumnWidget(FChooserColumnBase* Column, const UStruct* ColumnType, UChooserTable* Chooser, int Row);
 		
diff --git a/Engine/Plugins/Experimental/Chooser/Source/ChooserEditor/Public/SPropertyAccessChainWidget.h b/Engine/Plugins/Experimental/Chooser/Source/ChooserEditor/Public/SPropertyAccessChainWidget.h
--- a/Engine/Plugins/Experimental/Chooser/Source/ChooserEditor/Public/SPropertyAccessChainWidget.h
+++ b/Engine/Plugins/Experimental/Chooser/Source/ChooserEditor/Public/SPropertyAccessChainWidget.h
@@ -34,7 +34,7 @@ public:
 private:
 	TSharedRef<SWidget> CreatePropertyAccessWidget();
 	void UpdateWidget();
-	void ContextClassChanged(UClass* NewContextClass);
+	void ContextClassChanged();
 		
 	FString TypeFilter;
 	FString BindingColor;
diff --git a/Engine/Plugins/Experimental/Chooser/Source/ProxyTable/Internal/ProxyTable.h b/Engine/Plugins/Experimental/Chooser/Source/ProxyTable/Internal/ProxyTable.h
--- a/Engine/Plugins/Experimental/Chooser/Source/ProxyTable/Internal/ProxyTable.h
+++ b/Engine/Plugins/Experimental/Chooser/Source/ProxyTable/Internal/ProxyTable.h
@@ -27,15 +27,15 @@ public:
 
 	// caching the type so that on Undo, we can tell if we should fire the changed delegate
 	UClass* CachedPreviousType = nullptr;
-	UClass* CachedPreviousContextClass = nullptr;
 #endif
 	
 
 	UPROPERTY(EditAnywhere, Category = "Proxy", Meta = (AllowAbstract=true))
 	TObjectPtr<UClass> Type;
 	
-	UPROPERTY(EditAnywhere, Category = "Proxy Table Reference")
-	TObjectPtr<UClass> ContextClass;
+	UPROPERTY(EditAnywhere, Meta = (ExcludeBaseStruct, BaseStruct = "/Script/Chooser.ContextObjectTypeBase"), Category = "Input")
+	TArray<FInstancedStruct> ContextData;
+
 	
 	UPROPERTY(EditAnywhere, Meta = (ExcludeBaseStruct, BaseStruct ="/Script/ProxyTable.ChooserParameterProxyTableBase"), Category = "Proxy Table Reference")
 	FInstancedStruct ProxyTable;
@@ -43,9 +43,12 @@ public:
 	UPROPERTY()
 	FGuid Guid;
 
-	virtual UClass* GetContextClass() override { return ContextClass; }
+	virtual TConstArrayView<FInstancedStruct> GetContextData() const override { return ContextData; }
 
 #if WITH_EDITORONLY_DATA
+	UPROPERTY()
+	TObjectPtr<UClass> ContextClass_DEPRECATED;
+	
 	virtual void PostLoad() override;
 	virtual void PostDuplicate(EDuplicateMode::Type DuplicateMode) override;
 #endif
@@ -91,7 +94,7 @@ public:
 	UPROPERTY()
 	FInstancedStructContainer Values;
 	
-	UObject* FindProxyObject(const FGuid& Key, const UObject* ContextObject) const;
+	UObject* FindProxyObject(const FGuid& Key, FChooserEvaluationContext& Context) const;
 
 #if WITH_EDITORONLY_DATA
 public:
@@ -126,7 +129,7 @@ public:
 	UPROPERTY(EditAnywhere, Meta = (BindingType = "UProxyTable*"), Category = "Binding")
 	FChooserPropertyBinding Binding;
 
-	virtual bool GetValue(const UObject* ContextObject, const UProxyTable*& OutResult) const override;
+	virtual bool GetValue(FChooserEvaluationContext& Context, const UProxyTable*& OutResult) const override;
 
 #if WITH_EDITOR
 	static bool CanBind(const FProperty& Property)
@@ -136,7 +139,7 @@ public:
 
 	void SetBinding(const TArray<FBindingChainElement>& InBindingChain)
 	{
-		UE::Chooser::CopyPropertyChain(InBindingChain, Binding.PropertyBindingChain);
+		UE::Chooser::CopyPropertyChain(InBindingChain, Binding);
 	}
 #endif
 };
@@ -145,7 +148,7 @@ USTRUCT()
 struct PROXYTABLE_API FLookupProxy : public FObjectChooserBase
 {
 	GENERATED_BODY()
-	virtual UObject* ChooseObject(const UObject* ContextObject) const final override;
+	virtual UObject* ChooseObject(FChooserEvaluationContext& Context) const final override;
 	
 	public:
 	FLookupProxy();
diff --git a/Engine/Plugins/Experimental/Chooser/Source/ProxyTable/Private/ProxyTable.cpp b/Engine/Plugins/Experimental/Chooser/Source/ProxyTable/Private/ProxyTable.cpp
--- a/Engine/Plugins/Experimental/Chooser/Source/ProxyTable/Private/ProxyTable.cpp
+++ b/Engine/Plugins/Experimental/Chooser/Source/ProxyTable/Private/ProxyTable.cpp
@@ -2,6 +2,7 @@
 #include "ProxyTable.h"
 #include "ProxyTableFunctionLibrary.h"
 #include "AssetRegistry/AssetRegistryModule.h"
+#include "ChooserPropertyAccess.h"
 #include "Logging/LogMacros.h"
 #include "UObject/Package.h"
 
@@ -18,11 +19,7 @@ void UProxyAsset::PostEditUndo()
 		CachedPreviousType = Type;
 	}
 	
-	if (CachedPreviousContextClass != ContextClass)
-	{
-		OnContextClassChanged.Broadcast(ContextClass);
-		CachedPreviousContextClass = ContextClass;
-	}
+	OnContextClassChanged.Broadcast();
 }
 
 void UProxyAsset::PostEditChangeProperty(FPropertyChangedEvent& PropertyChangedEvent)
@@ -30,7 +27,7 @@ void UProxyAsset::PostEditChangeProperty(FPropertyChangedEvent& PropertyChangedE
 	UObject::PostEditChangeProperty(PropertyChangedEvent);
 	
 	static FName TypeName = "Type";
-	static FName ContextClassName = "ContextClass";
+	static FName ContextClassName = "ContextData";
 	if (PropertyChangedEvent.Property->GetName() == TypeName)
 	{
 		if (CachedPreviousType != Type)
@@ -41,11 +38,7 @@ void UProxyAsset::PostEditChangeProperty(FPropertyChangedEvent& PropertyChangedE
 	}
 	else if (PropertyChangedEvent.Property->GetName() == ContextClassName)
 	{
-		if (CachedPreviousContextClass != ContextClass)
-		{
-			OnContextClassChanged.Broadcast(ContextClass);
-		}
-		CachedPreviousContextClass = ContextClass;
+		OnContextClassChanged.Broadcast();
 	}
 }
 
@@ -68,6 +61,19 @@ void UProxyAsset::PostLoad()
 #if WITH_EDITOR
 	CachedPreviousType = Type;
 #endif
+	
+#if WITH_EDITORONLY_DATA
+	if (ContextClass_DEPRECATED)
+	{
+		
+		ContextData.SetNum(1);
+		ContextData[0].InitializeAs<FContextObjectTypeClass>();
+		FContextObjectTypeClass& Context = ContextData[0].GetMutable<FContextObjectTypeClass>();
+		Context.Class = ContextClass_DEPRECATED;
+		Context.Direction = EContextObjectDirection::ReadWrite;
+		ContextClass_DEPRECATED = nullptr;
+	}
+#endif
 
 	if (!Guid.IsValid())
 	{
@@ -75,6 +81,7 @@ void UProxyAsset::PostLoad()
 		Guid.A = GetTypeHash(GetName());
 		Guid.B = GetTypeHash(GetPackage()->GetPathName());
 	}
+
 }
 
 void UProxyAsset::PostDuplicate(EDuplicateMode::Type DuplicateMode)
@@ -248,28 +255,28 @@ void UProxyTable::PostTransacted(const FTransactionObjectEvent& TransactionEvent
 
 #endif
 
-UObject* UProxyTable::FindProxyObject(const FGuid& Key, const UObject* ContextObject) const
+UObject* UProxyTable::FindProxyObject(const FGuid& Key, FChooserEvaluationContext& Context) const
 {
 	const int FoundIndex = Algo::BinarySearch(Keys, Key);
 	if (FoundIndex != INDEX_NONE)
 	{
 		const FObjectChooserBase &EntryValue = Values[FoundIndex].Get<const FObjectChooserBase>();
-		return EntryValue.ChooseObject(ContextObject);
+		return EntryValue.ChooseObject(Context);
 	}
 	
 	return nullptr;
 }
 
-static UObject* FindProxyObject(const UProxyAsset* Proxy, const UObject* ContextObject)
+static UObject* FindProxyObject(const UProxyAsset* Proxy, FChooserEvaluationContext Context)
 {
 	if (Proxy && Proxy->ProxyTable.IsValid())
 	{
 		const UProxyTable* Table;
-		if (Proxy->ProxyTable.Get<FChooserParameterProxyTableBase>().GetValue(ContextObject, Table))
+		if (Proxy->ProxyTable.Get<FChooserParameterProxyTableBase>().GetValue(Context, Table))
 		{
 			if(Table)
 			{
-				if (UObject* Value = Table->FindProxyObject(Proxy->Guid, ContextObject))
+				if (UObject* Value = Table->FindProxyObject(Proxy->Guid, Context))
 				{
 					return Value;
 				}
@@ -280,9 +287,9 @@ static UObject* FindProxyObject(const UProxyAsset* Proxy, const UObject* Context
 	return nullptr;
 }
 
-UObject* FLookupProxy::ChooseObject(const UObject* ContextObject) const
+UObject* FLookupProxy::ChooseObject(FChooserEvaluationContext& Context) const
 {
-	return FindProxyObject(Proxy, ContextObject);
+	return FindProxyObject(Proxy, Context);
 }
 
 UProxyTable::UProxyTable(const FObjectInitializer& Initializer)
@@ -297,12 +304,12 @@ UProxyAsset::UProxyAsset(const FObjectInitializer& Initializer)
 	ProxyTable.InitializeAs(FProxyTableContextProperty::StaticStruct());
 }
 
-bool FProxyTableContextProperty::GetValue(const UObject* ContextObject, const UProxyTable*& OutResult) const
+bool FProxyTableContextProperty::GetValue(FChooserEvaluationContext& Context, const UProxyTable*& OutResult) const
 {
-	UStruct* StructType = ContextObject->GetClass();
-	const void* Container = ContextObject;
+	const UStruct* StructType = nullptr;
+	const void* Container = nullptr;
 	
-	if (UE::Chooser::ResolvePropertyChain(Container, StructType, Binding.PropertyBindingChain))
+	if (UE::Chooser::ResolvePropertyChain(Context, Binding,Container, StructType))
 	{
 		if (const FObjectProperty* Property = FindFProperty<FObjectProperty>(StructType, Binding.PropertyBindingChain.Last()))
 		{
@@ -319,7 +326,10 @@ bool FProxyTableContextProperty::GetValue(const UObject* ContextObject, const UP
 
 UObject* UProxyTableFunctionLibrary::EvaluateProxyAsset(const UObject* ContextObject, const UProxyAsset* Proxy, TSubclassOf<UObject> ObjectClass)
 {
-	UObject* Result = FindProxyObject(Proxy, ContextObject);
+	FChooserEvaluationContext Context;
+	Context.ContextData.Add({ContextObject->GetClass(), const_cast<UObject*>(ContextObject)});
+	
+	UObject* Result = FindProxyObject(Proxy, Context);
 	if (ObjectClass && Result && !Result->IsA(ObjectClass))
 	{
 		return nullptr;
@@ -335,7 +345,9 @@ UObject* UProxyTableFunctionLibrary::EvaluateProxyTable(const UObject* ContextOb
 	{
 		FGuid Guid;
 		Guid.A = GetTypeHash(Key);
-		if (UObject* Value = ProxyTable->FindProxyObject(Guid, ContextObject))
+		FChooserEvaluationContext Context;
+		Context.ContextData.Add({ContextObject->GetClass(), const_cast<UObject*>(ContextObject)});
+		if (UObject* Value = ProxyTable->FindProxyObject(Guid, Context))
 		{
 			return Value;
 		}
diff --git a/Engine/Plugins/Experimental/Chooser/Source/ProxyTable/Public/IChooserParameterProxyTable.h b/Engine/Plugins/Experimental/Chooser/Source/ProxyTable/Public/IChooserParameterProxyTable.h
--- a/Engine/Plugins/Experimental/Chooser/Source/ProxyTable/Public/IChooserParameterProxyTable.h
+++ b/Engine/Plugins/Experimental/Chooser/Source/ProxyTable/Public/IChooserParameterProxyTable.h
@@ -3,10 +3,10 @@
 #pragma once
 
 #include "CoreMinimal.h"
-#include "UObject/Object.h"
 #include "UObject/Interface.h"
 #include "IChooserParameterBase.h"
 #include "InstancedStruct.h"
+#include "IObjectChooser.h"
 #include "IChooserParameterProxyTable.generated.h"
 
 class UProxyTable;
@@ -31,5 +31,5 @@ struct FChooserParameterProxyTableBase : public FChooserParameterBase
 	GENERATED_BODY()
     
 public:
-	virtual bool GetValue(const UObject* ContextObject, const UProxyTable*& OutResult) const { return false; }
+	virtual bool GetValue(FChooserEvaluationContext& Context, const UProxyTable*& OutResult) const { return false; }
 };
\ No newline at end of file
diff --git a/Engine/Plugins/Experimental/Chooser/Source/ProxyTableEditor/Private/ProxyTableEditor.cpp b/Engine/Plugins/Experimental/Chooser/Source/ProxyTableEditor/Private/ProxyTableEditor.cpp
--- a/Engine/Plugins/Experimental/Chooser/Source/ProxyTableEditor/Private/ProxyTableEditor.cpp
+++ b/Engine/Plugins/Experimental/Chooser/Source/ProxyTableEditor/Private/ProxyTableEditor.cpp
@@ -427,11 +427,9 @@ public:
 			}
 			else if (ColumnName == Value) 
 			{
-				UClass* ContextClass = nullptr;
 				UClass* ObjectType = nullptr;
 				if (Row->ProxyTable->Entries[Row->RowIndex].Proxy)
 				{
-					ContextClass = Row->ProxyTable->Entries[Row->RowIndex].Proxy->ContextClass;
 					ObjectType = Row->ProxyTable->Entries[Row->RowIndex].Proxy->Type;
 				}
 				bool bReadOnly = Row->ProxyTable != Editor->GetProxyTable();
@@ -439,7 +437,7 @@ public:
 					Row->ProxyTable, FObjectChooserBase::StaticStruct(),
 					Row->ProxyTable->Entries[Row->RowIndex].ValueStruct.GetMutableMemory(),
 					Row->ProxyTable->Entries[Row->RowIndex].ValueStruct.GetScriptStruct(),
-					ContextClass, ObjectType,
+					ObjectType,
 					FOnStructPicked::CreateLambda([this](const UScriptStruct* ChosenStruct)
 					{
 						{
@@ -450,7 +448,6 @@ public:
 						ChooserEditor::FObjectChooserWidgetFactories::CreateWidget(false, Row->ProxyTable, FObjectChooserBase::StaticStruct(),
 								Row->ProxyTable->Entries[Row->RowIndex].ValueStruct.GetMutableMemory(),
 								Row->ProxyTable->Entries[Row->RowIndex].ValueStruct.GetScriptStruct(),
-								Row->ProxyTable->Entries[Row->RowIndex].Proxy->ContextClass,
 								Row->ProxyTable->Entries[Row->RowIndex].Proxy->Type,
 								FOnStructPicked(), &CacheBorder);
 					}),
@@ -976,7 +973,7 @@ TSharedRef<FProxyTableEditor> FProxyTableEditor::CreateEditor( const EToolkitMod
 
 /// Result widgets
 
-TSharedRef<SWidget> CreateLookupProxyWidget(bool bReadOnly, UObject* TransactionObject, void* Value, UClass* ContextObject, UClass* ResultBaseClass)
+TSharedRef<SWidget> CreateLookupProxyWidget(bool bReadOnly, UObject* TransactionObject, void* Value, UClass* ResultBaseClass)
 {
 	FLookupProxy* LookupProxy = static_cast<FLookupProxy*>(Value);
 	
