Fix non-unitiy compile

diff --git a/Engine/Plugins/Experimental/Chooser/Source/Chooser/Public/IChooserParameterBase.h b/Engine/Plugins/Experimental/Chooser/Source/Chooser/Public/IChooserParameterBase.h
--- a/Engine/Plugins/Experimental/Chooser/Source/Chooser/Public/IChooserParameterBase.h
+++ b/Engine/Plugins/Experimental/Chooser/Source/Chooser/Public/IChooserParameterBase.h
@@ -3,6 +3,7 @@
 #pragma once
 
 #include "CoreMinimal.h"
+#include "IObjectChooser.h"
 #include "IChooserParameterBase.generated.h"
 
 USTRUCT()
