proxy table cleanup

diff --git a/Engine/Plugins/Experimental/Chooser/Source/ProxyTable/Internal/LookupProxy.h b/Engine/Plugins/Experimental/Chooser/Source/ProxyTable/Internal/LookupProxy.h
new file mode 100644
--- /dev/null
+++ b/Engine/Plugins/Experimental/Chooser/Source/ProxyTable/Internal/LookupProxy.h
@@ -0,0 +1,49 @@
+// Copyright Epic Games, Inc. All Rights Reserved.
+#pragma once
+
+#include "CoreMinimal.h"
+#include "IObjectChooser.h"
+#include "ChooserPropertyAccess.h"
+#include "IChooserParameterProxyTable.h"
+#include "InstancedStructContainer.h"
+#include "ProxyTable.h"
+#include "LookupProxy.generated.h"
+
+struct FBindingChainElement;
+
+USTRUCT()
+struct PROXYTABLE_API FProxyTableContextProperty :  public FChooserParameterProxyTableBase
+{
+	GENERATED_BODY()
+public:
+	
+	UPROPERTY(EditAnywhere, Meta = (BindingType = "UProxyTable*"), Category = "Binding")
+	FChooserPropertyBinding Binding;
+
+	virtual bool GetValue(FChooserEvaluationContext& Context, const UProxyTable*& OutResult) const override;
+
+#if WITH_EDITOR
+	static bool CanBind(const FProperty& Property)
+	{
+		return Property.GetCPPType() == "UProxyTable*";
+	}
+
+	void SetBinding(const TArray<FBindingChainElement>& InBindingChain)
+	{
+		UE::Chooser::CopyPropertyChain(InBindingChain, Binding);
+	}
+#endif
+};
+
+USTRUCT()
+struct PROXYTABLE_API FLookupProxy : public FObjectChooserBase
+{
+	GENERATED_BODY()
+	virtual UObject* ChooseObject(FChooserEvaluationContext& Context) const final override;
+	
+	public:
+	FLookupProxy();
+	
+	UPROPERTY(EditAnywhere, Category="Parameters")
+	TObjectPtr<UProxyAsset> Proxy;
+};
\ No newline at end of file
diff --git a/Engine/Plugins/Experimental/Chooser/Source/ProxyTable/Internal/ProxyTable.h b/Engine/Plugins/Experimental/Chooser/Source/ProxyTable/Internal/ProxyTable.h
deleted file mode 100644
diff --git a/Engine/Plugins/Experimental/Chooser/Source/ProxyTable/Private/LookupProxy.cpp b/Engine/Plugins/Experimental/Chooser/Source/ProxyTable/Private/LookupProxy.cpp
new file mode 100644
--- /dev/null
+++ b/Engine/Plugins/Experimental/Chooser/Source/ProxyTable/Private/LookupProxy.cpp
@@ -0,0 +1,36 @@
+// Copyright Epic Games, Inc. All Rights Reserved.
+#include "LookupProxy.h"
+#include "ProxyTableFunctionLibrary.h"
+#include "AssetRegistry/AssetRegistryModule.h"
+#include "ChooserPropertyAccess.h"
+#include "UObject/Package.h"
+
+FLookupProxy::FLookupProxy()
+{
+}
+
+UObject* FLookupProxy::ChooseObject(FChooserEvaluationContext& Context) const
+{
+	if (Proxy)
+	{
+		return Proxy->FindProxyObject(Context);
+	}
+	return nullptr;
+}
+
+bool FProxyTableContextProperty::GetValue(FChooserEvaluationContext& Context, const UProxyTable*& OutResult) const
+{
+	const UStruct* StructType = nullptr;
+	const void* Container = nullptr;
+	
+	if (UE::Chooser::ResolvePropertyChain(Context, Binding,Container, StructType))
+	{
+		if (const FObjectProperty* Property = FindFProperty<FObjectProperty>(StructType, Binding.PropertyBindingChain.Last()))
+		{
+			OutResult = *Property->ContainerPtrToValuePtr<UProxyTable*>(Container);
+			return true;
+		}
+	}
+
+	return false;
+}
diff --git a/Engine/Plugins/Experimental/Chooser/Source/ProxyTable/Private/ProxyAsset.cpp b/Engine/Plugins/Experimental/Chooser/Source/ProxyTable/Private/ProxyAsset.cpp
new file mode 100644
--- /dev/null
+++ b/Engine/Plugins/Experimental/Chooser/Source/ProxyTable/Private/ProxyAsset.cpp
@@ -0,0 +1,114 @@
+// Copyright Epic Games, Inc. All Rights Reserved.
+#include "ProxyAsset.h"
+#include "ProxyTableFunctionLibrary.h"
+#include "AssetRegistry/AssetRegistryModule.h"
+#include "ChooserPropertyAccess.h"
+#include "UObject/Package.h"
+#include "LookupProxy.h"
+
+#if WITH_EDITOR
+void UProxyAsset::PostEditUndo()
+{
+	UObject::PostEditUndo();
+
+	if (CachedPreviousType != Type)
+	{
+		OnTypeChanged.Broadcast(Type);
+		CachedPreviousType = Type;
+	}
+	
+	OnContextClassChanged.Broadcast();
+}
+
+void UProxyAsset::PostEditChangeProperty(FPropertyChangedEvent& PropertyChangedEvent)
+{
+	UObject::PostEditChangeProperty(PropertyChangedEvent);
+	
+	static FName TypeName = "Type";
+	static FName ContextClassName = "ContextData";
+	if (PropertyChangedEvent.Property->GetName() == TypeName)
+	{
+		if (CachedPreviousType != Type)
+		{
+			OnTypeChanged.Broadcast(Type);
+		}
+		CachedPreviousType = Type;
+	}
+	else if (PropertyChangedEvent.Property->GetName() == ContextClassName)
+	{
+		OnContextClassChanged.Broadcast();
+	}
+}
+
+#endif
+
+#if WITH_EDITORONLY_DATA
+
+/////////////////////////////////////////////////////////////////////////////////////////
+// Proxy Asset
+
+void UProxyAsset::PostLoad()
+{
+	Super::PostLoad();
+
+#if WITH_EDITOR
+	CachedPreviousType = Type;
+#endif
+	
+	if (ContextClass_DEPRECATED)
+	{
+		
+		ContextData.SetNum(1);
+		ContextData[0].InitializeAs<FContextObjectTypeClass>();
+		FContextObjectTypeClass& Context = ContextData[0].GetMutable<FContextObjectTypeClass>();
+		Context.Class = ContextClass_DEPRECATED;
+		Context.Direction = EContextObjectDirection::ReadWrite;
+		ContextClass_DEPRECATED = nullptr;
+	}
+
+	if (!Guid.IsValid())
+	{
+		// if we load a ProxyAsset that was created before the Guid, assign it a deterministic guid based on the name and path.
+		Guid.A = GetTypeHash(GetName());
+		Guid.B = GetTypeHash(GetPackage()->GetPathName());
+	}
+
+}
+
+void UProxyAsset::PostDuplicate(EDuplicateMode::Type DuplicateMode)
+{
+	UObject::PostDuplicate(DuplicateMode);
+	if (DuplicateMode == EDuplicateMode::Normal)
+	{
+		// create a new guid when duplicating
+		Guid = FGuid::NewGuid();
+	}
+}
+
+#endif
+
+UProxyAsset::UProxyAsset(const FObjectInitializer& Initializer)
+	: Super(Initializer)
+{
+	ProxyTable.InitializeAs(FProxyTableContextProperty::StaticStruct());
+}
+
+UObject* UProxyAsset::FindProxyObject(FChooserEvaluationContext Context) const
+{
+	if (ProxyTable.IsValid())
+	{
+		const UProxyTable* Table;
+		if (ProxyTable.Get<FChooserParameterProxyTableBase>().GetValue(Context, Table))
+		{
+			if(Table)
+			{
+				if (UObject* Value = Table->FindProxyObject(Guid, Context))
+				{
+					return Value;
+				}
+			}
+		}
+	}
+	
+	return nullptr;
+}
diff --git a/Engine/Plugins/Experimental/Chooser/Source/ProxyTable/Private/ProxyTable.cpp b/Engine/Plugins/Experimental/Chooser/Source/ProxyTable/Private/ProxyTable.cpp
--- a/Engine/Plugins/Experimental/Chooser/Source/ProxyTable/Private/ProxyTable.cpp
+++ b/Engine/Plugins/Experimental/Chooser/Source/ProxyTable/Private/ProxyTable.cpp
@@ -1,99 +1,12 @@
 // Copyright Epic Games, Inc. All Rights Reserved.
 #include "ProxyTable.h"
-#include "ProxyTableFunctionLibrary.h"
 #include "AssetRegistry/AssetRegistryModule.h"
-#include "ChooserPropertyAccess.h"
 #include "Logging/LogMacros.h"
-#include "UObject/Package.h"
 
 DEFINE_LOG_CATEGORY_STATIC(LogProxyTable,Log,All);
 
-#if WITH_EDITOR
-void UProxyAsset::PostEditUndo()
-{
-	UObject::PostEditUndo();
-
-	if (CachedPreviousType != Type)
-	{
-		OnTypeChanged.Broadcast(Type);
-		CachedPreviousType = Type;
-	}
-	
-	OnContextClassChanged.Broadcast();
-}
-
-void UProxyAsset::PostEditChangeProperty(FPropertyChangedEvent& PropertyChangedEvent)
-{
-	UObject::PostEditChangeProperty(PropertyChangedEvent);
-	
-	static FName TypeName = "Type";
-	static FName ContextClassName = "ContextData";
-	if (PropertyChangedEvent.Property->GetName() == TypeName)
-	{
-		if (CachedPreviousType != Type)
-		{
-			OnTypeChanged.Broadcast(Type);
-		}
-		CachedPreviousType = Type;
-	}
-	else if (PropertyChangedEvent.Property->GetName() == ContextClassName)
-	{
-		OnContextClassChanged.Broadcast();
-	}
-}
-
-#endif
-
-
-FLookupProxy::FLookupProxy()
-{
-}
-
 #if WITH_EDITORONLY_DATA
 
-/////////////////////////////////////////////////////////////////////////////////////////
-// Proxy Asset
-
-void UProxyAsset::PostLoad()
-{
-	Super::PostLoad();
-
-#if WITH_EDITOR
-	CachedPreviousType = Type;
-#endif
-	
-#if WITH_EDITORONLY_DATA
-	if (ContextClass_DEPRECATED)
-	{
-		
-		ContextData.SetNum(1);
-		ContextData[0].InitializeAs<FContextObjectTypeClass>();
-		FContextObjectTypeClass& Context = ContextData[0].GetMutable<FContextObjectTypeClass>();
-		Context.Class = ContextClass_DEPRECATED;
-		Context.Direction = EContextObjectDirection::ReadWrite;
-		ContextClass_DEPRECATED = nullptr;
-	}
-#endif
-
-	if (!Guid.IsValid())
-	{
-		// if we load a ProxyAsset that was created before the Guid, assign it a deterministic guid based on the name and path.
-		Guid.A = GetTypeHash(GetName());
-		Guid.B = GetTypeHash(GetPackage()->GetPathName());
-	}
-
-}
-
-void UProxyAsset::PostDuplicate(EDuplicateMode::Type DuplicateMode)
-{
-	UObject::PostDuplicate(DuplicateMode);
-	if (DuplicateMode == EDuplicateMode::Normal)
-	{
-		// create a new guid when duplicating
-		Guid = FGuid::NewGuid();
-	}
-}
-
 //////////////////////////////////////////////////////////////////////////////////////
 /// Proxy Entry
 
@@ -267,91 +180,8 @@ UObject* UProxyTable::FindProxyObject(const FGuid& Key, FChooserEvaluationContex
 	return nullptr;
 }
 
-static UObject* FindProxyObject(const UProxyAsset* Proxy, FChooserEvaluationContext Context)
-{
-	if (Proxy && Proxy->ProxyTable.IsValid())
-	{
-		const UProxyTable* Table;
-		if (Proxy->ProxyTable.Get<FChooserParameterProxyTableBase>().GetValue(Context, Table))
-		{
-			if(Table)
-			{
-				if (UObject* Value = Table->FindProxyObject(Proxy->Guid, Context))
-				{
-					return Value;
-				}
-			}
-		}
-	}
-	
-	return nullptr;
-}
-
-UObject* FLookupProxy::ChooseObject(FChooserEvaluationContext& Context) const
-{
-	return FindProxyObject(Proxy, Context);
-}
-
 UProxyTable::UProxyTable(const FObjectInitializer& Initializer)
 	:Super(Initializer)
 {
 
 }
-
-UProxyAsset::UProxyAsset(const FObjectInitializer& Initializer)
-	: Super(Initializer)
-{
-	ProxyTable.InitializeAs(FProxyTableContextProperty::StaticStruct());
-}
-
-bool FProxyTableContextProperty::GetValue(FChooserEvaluationContext& Context, const UProxyTable*& OutResult) const
-{
-	const UStruct* StructType = nullptr;
-	const void* Container = nullptr;
-	
-	if (UE::Chooser::ResolvePropertyChain(Context, Binding,Container, StructType))
-	{
-		if (const FObjectProperty* Property = FindFProperty<FObjectProperty>(StructType, Binding.PropertyBindingChain.Last()))
-		{
-			OutResult = *Property->ContainerPtrToValuePtr<UProxyTable*>(Container);
-			return true;
-		}
-	}
-
-	return false;
-}
-
-/////////////////////////////////////////////////////////////////////////////////////
-// Blueprint Library Functions
-
-UObject* UProxyTableFunctionLibrary::EvaluateProxyAsset(const UObject* ContextObject, const UProxyAsset* Proxy, TSubclassOf<UObject> ObjectClass)
-{
-	FChooserEvaluationContext Context;
-	Context.ContextData.Add({ContextObject->GetClass(), const_cast<UObject*>(ContextObject)});
-	
-	UObject* Result = FindProxyObject(Proxy, Context);
-	if (ObjectClass && Result && !Result->IsA(ObjectClass))
-	{
-		return nullptr;
-	}
-	return Result;
-	
-}
-
-// fallback for FName based Keys:
-UObject* UProxyTableFunctionLibrary::EvaluateProxyTable(const UObject* ContextObject, const UProxyTable* ProxyTable, FName Key)
-{
-	if (ProxyTable)
-	{
-		FGuid Guid;
-		Guid.A = GetTypeHash(Key);
-		FChooserEvaluationContext Context;
-		Context.ContextData.Add({ContextObject->GetClass(), const_cast<UObject*>(ContextObject)});
-		if (UObject* Value = ProxyTable->FindProxyObject(Guid, Context))
-		{
-			return Value;
-		}
-	}
-	
-	return nullptr;
-}
diff --git a/Engine/Plugins/Experimental/Chooser/Source/ProxyTable/Private/ProxyTableFunctionLibrary.cpp b/Engine/Plugins/Experimental/Chooser/Source/ProxyTable/Private/ProxyTableFunctionLibrary.cpp
new file mode 100644
--- /dev/null
+++ b/Engine/Plugins/Experimental/Chooser/Source/ProxyTable/Private/ProxyTableFunctionLibrary.cpp
@@ -0,0 +1,42 @@
+// Copyright Epic Games, Inc. All Rights Reserved.
+#include "ProxyTableFunctionLibrary.h"
+#include "ProxyTable.h"
+
+/////////////////////////////////////////////////////////////////////////////////////
+// Blueprint Library Functions
+
+UObject* UProxyTableFunctionLibrary::EvaluateProxyAsset(const UObject* ContextObject, const UProxyAsset* Proxy, TSubclassOf<UObject> ObjectClass)
+{
+	UObject* Result = nullptr;
+	if (Proxy)
+	{
+		FChooserEvaluationContext Context;
+		Context.ContextData.Add({ContextObject->GetClass(), const_cast<UObject*>(ContextObject)});
+		
+		Result = Proxy->FindProxyObject(Context);
+		if (ObjectClass && Result && !Result->IsA(ObjectClass))
+		{
+			return nullptr;
+		}
+	}
+	return Result;
+	
+}
+
+// fallback for FName based Keys:
+UObject* UProxyTableFunctionLibrary::EvaluateProxyTable(const UObject* ContextObject, const UProxyTable* ProxyTable, FName Key)
+{
+	if (ProxyTable)
+	{
+		FGuid Guid;
+		Guid.A = GetTypeHash(Key);
+		FChooserEvaluationContext Context;
+		Context.ContextData.Add({ContextObject->GetClass(), const_cast<UObject*>(ContextObject)});
+		if (UObject* Value = ProxyTable->FindProxyObject(Guid, Context))
+		{
+			return Value;
+		}
+	}
+	
+	return nullptr;
+}
diff --git a/Engine/Plugins/Experimental/Chooser/Source/ProxyTable/Public/ProxyAsset.h b/Engine/Plugins/Experimental/Chooser/Source/ProxyTable/Public/ProxyAsset.h
new file mode 100644
--- /dev/null
+++ b/Engine/Plugins/Experimental/Chooser/Source/ProxyTable/Public/ProxyAsset.h
@@ -0,0 +1,55 @@
+// Copyright Epic Games, Inc. All Rights Reserved.
+#pragma once
+
+#include "CoreMinimal.h"
+#include "UObject/Object.h"
+#include "ChooserPropertyAccess.h"
+#include "IChooserParameterProxyTable.h"
+#include "InstancedStruct.h"
+#include "InstancedStructContainer.h"
+#include "Misc/Guid.h"
+#include "ProxyAsset.generated.h"
+
+DECLARE_MULTICAST_DELEGATE_OneParam(FProxyTypeChanged, const UClass* OutputObjectType);
+
+UCLASS(MinimalAPI,BlueprintType)
+class UProxyAsset : public UObject, public IHasContextClass
+{
+	GENERATED_UCLASS_BODY()
+public:
+	UProxyAsset() {}
+
+#if WITH_EDITOR
+	FProxyTypeChanged OnTypeChanged;
+	virtual void PostEditUndo() override;
+	virtual void PostEditChangeProperty(FPropertyChangedEvent& PropertyChangedEvent) override;
+
+	// caching the type so that on Undo, we can tell if we should fire the changed delegate
+	UClass* CachedPreviousType = nullptr;
+#endif
+	
+
+	UPROPERTY(EditAnywhere, Category = "Proxy", Meta = (AllowAbstract=true))
+	TObjectPtr<UClass> Type;
+	
+	UPROPERTY(EditAnywhere, Meta = (ExcludeBaseStruct, BaseStruct = "/Script/Chooser.ContextObjectTypeBase"), Category = "Input")
+	TArray<FInstancedStruct> ContextData;
+
+	
+	UPROPERTY(EditAnywhere, Meta = (ExcludeBaseStruct, BaseStruct ="/Script/ProxyTable.ChooserParameterProxyTableBase"), Category = "Proxy Table Reference")
+	FInstancedStruct ProxyTable;
+
+	UPROPERTY()
+	FGuid Guid;
+
+	virtual TConstArrayView<FInstancedStruct> GetContextData() const override { return ContextData; }
+	UObject* FindProxyObject(FChooserEvaluationContext Context) const;
+
+#if WITH_EDITORONLY_DATA
+	UPROPERTY()
+	TObjectPtr<UClass> ContextClass_DEPRECATED;
+	
+	virtual void PostLoad() override;
+	virtual void PostDuplicate(EDuplicateMode::Type DuplicateMode) override;
+#endif
+};
\ No newline at end of file
diff --git a/Engine/Plugins/Experimental/Chooser/Source/ProxyTable/Public/ProxyTable.h b/Engine/Plugins/Experimental/Chooser/Source/ProxyTable/Public/ProxyTable.h
new file mode 100644
--- /dev/null
+++ b/Engine/Plugins/Experimental/Chooser/Source/ProxyTable/Public/ProxyTable.h
@@ -0,0 +1,74 @@
+// Copyright Epic Games, Inc. All Rights Reserved.
+#pragma once
+
+#include "CoreMinimal.h"
+#include "UObject/Object.h"
+#include "ChooserPropertyAccess.h"
+#include "IChooserParameterProxyTable.h"
+#include "InstancedStruct.h"
+#include "InstancedStructContainer.h"
+#include "Misc/Guid.h"
+#include "ProxyAsset.h"
+#include "ProxyTable.generated.h"
+
+USTRUCT()
+struct PROXYTABLE_API FProxyEntry
+{
+	GENERATED_BODY()
+
+	bool operator== (const FProxyEntry& Other) const;
+	bool operator< (const FProxyEntry& Other) const;
+	
+	UPROPERTY(EditAnywhere, Category = "Data")
+	TObjectPtr<UProxyAsset> Proxy;
+
+	// temporarily leaving this property for backwards compatibility with old content which used FNames rather than UProxyAsset
+	UPROPERTY(EditAnywhere, Category = "Data")
+	FName Key;
+	
+	UPROPERTY(DisplayName="Value", EditAnywhere, Meta = (ExcludeBaseStruct, BaseStruct = "/Script/Chooser.ObjectChooserBase"), Category = "Data")
+	FInstancedStruct ValueStruct;
+
+	const FGuid GetGuid() const;
+};
+
+#if WITH_EDITORONLY_DATA
+inline uint32 GetTypeHash(const FProxyEntry& Entry) { return GetTypeHash(Entry.GetGuid()); }
+#endif
+
+DECLARE_MULTICAST_DELEGATE(FProxyTableChanged);
+
+UCLASS(MinimalAPI,BlueprintType)
+class UProxyTable : public UObject
+{
+	GENERATED_UCLASS_BODY()
+public:
+	UProxyTable() {}
+
+	UPROPERTY()
+	TArray<FGuid> Keys;
+
+	UPROPERTY()
+	FInstancedStructContainer Values;
+	
+	UObject* FindProxyObject(const FGuid& Key, FChooserEvaluationContext& Context) const;
+
+#if WITH_EDITORONLY_DATA
+public:
+	FProxyTableChanged OnProxyTableChanged;
+	
+	UPROPERTY(EditAnywhere, Category = "Hidden")
+	TArray<FProxyEntry> Entries;
+	
+	UPROPERTY(EditAnywhere, Category = "Inheritance")
+	TArray<TObjectPtr<UProxyTable>> InheritEntriesFrom;
+
+	virtual void PostLoad() override;
+	virtual void PostTransacted(const FTransactionObjectEvent& TransactionEvent) override;
+private:
+	void BuildRuntimeData();
+	TArray<TWeakObjectPtr<UProxyTable>> TableDependencies;
+	TArray<TWeakObjectPtr<UProxyAsset>> ProxyDependencies;
+#endif
+
+};
diff --git a/Engine/Plugins/Experimental/Chooser/Source/ProxyTableEditor/Private/ProxyTableEditor.cpp b/Engine/Plugins/Experimental/Chooser/Source/ProxyTableEditor/Private/ProxyTableEditor.cpp
--- a/Engine/Plugins/Experimental/Chooser/Source/ProxyTableEditor/Private/ProxyTableEditor.cpp
+++ b/Engine/Plugins/Experimental/Chooser/Source/ProxyTableEditor/Private/ProxyTableEditor.cpp
@@ -31,6 +31,7 @@
 #include "Widgets/Layout/SSeparator.h"
 #include "PropertyCustomizationHelpers.h"
 #include "ProxyTableEditorCommands.h"
+#include "LookupProxy.h"
 #include "Widgets/Layout/SScrollBox.h"
 
 #define LOCTEXT_NAMESPACE "ProxyTableEditor"
