Set output columns for first result when querying a chooser for multiple results.

diff --git a/Engine/Plugins/Experimental/Chooser/Source/Chooser/Private/Chooser.cpp b/Engine/Plugins/Experimental/Chooser/Source/Chooser/Private/Chooser.cpp
--- a/Engine/Plugins/Experimental/Chooser/Source/Chooser/Private/Chooser.cpp
+++ b/Engine/Plugins/Experimental/Chooser/Source/Chooser/Private/Chooser.cpp
@@ -16,7 +16,6 @@ void UChooserTable::PostEditUndo()
 {
 	UObject::PostEditUndo();
 
-
 	if (CachedPreviousOutputObjectType != OutputObjectType)
 	{
 		OnOutputObjectTypeChanged.Broadcast(OutputObjectType);
@@ -175,14 +174,17 @@ static FObjectChooserBase::EIteratorStatus StaticEvaluateChooser(FChooserEvaluat
 		}
 	}
 	
-	// of the rows that passed all column filters, return the first one for which the result row succeeds (could fail eg for a nexted chooser where no rows passed)
+	// of the rows that passed all column filters, iterate through them calling the callback until it returns Stop
+	bool bSetOutputs = false;
 	for (uint32 SelectedIndex : *IndicesOut)
 	{
 		if (Chooser->ResultsStructs.Num() > (int32)SelectedIndex)
 		{
 			const FObjectChooserBase& SelectedResult = Chooser->ResultsStructs[SelectedIndex].Get<FObjectChooserBase>();
-			if (SelectedResult.ChooseMulti(Context, Callback) == FObjectChooserBase::EIteratorStatus::Stop)
+			FObjectChooserBase::EIteratorStatus Status = SelectedResult.ChooseMulti(Context, Callback);
+			if (Status != FObjectChooserBase::EIteratorStatus::Continue)
 			{
+				bSetOutputs = true;
 				// trigger all output columns
 				for (const FInstancedStruct& ColumnData : Chooser->ColumnsStructs)
 				{
@@ -195,12 +197,16 @@ static FObjectChooserBase::EIteratorStatus StaticEvaluateChooser(FChooserEvaluat
 					Chooser->SetDebugSelectedRow(SelectedIndex);
 				}
 #endif
+			}
+			if (Status == FObjectChooserBase::EIteratorStatus::Stop)
+			{
 				return FObjectChooserBase::EIteratorStatus::Stop;
 			}
 		}
 	}
-	
-	return FObjectChooserBase::EIteratorStatus::Continue;
+
+	// If this is a nested chooser make sure the parent also sets the output vales from the row that contained this chooser
+	return bSetOutputs ? FObjectChooserBase::EIteratorStatus::ContinueWithOutputs : FObjectChooserBase::EIteratorStatus::Continue;
 }
 
 UObject* FEvaluateChooser::ChooseObject(FChooserEvaluationContext& Context) const
@@ -258,7 +264,8 @@ TArray<UObject*> UChooserFunctionLibrary::EvaluateChooserMulti(const UObject* Co
 		{
 			Result.Add(InResult);
 		}
-		return FObjectChooserBase::EIteratorStatus::Continue;
+		// trigger output columns only on the first result
+		return Result.Num() == 1 ? FObjectChooserBase::EIteratorStatus::ContinueWithOutputs : FObjectChooserBase::EIteratorStatus::Continue;
 	}));
 
 	for (int Index = 0; Index < Result.Num(); ++Index)
@@ -270,4 +277,4 @@ TArray<UObject*> UChooserFunctionLibrary::EvaluateChooserMulti(const UObject* Co
 	}
 
 	return Result;
-}
\ No newline at end of file
+}
diff --git a/Engine/Plugins/Experimental/Chooser/Source/Chooser/Public/IObjectChooser.h b/Engine/Plugins/Experimental/Chooser/Source/Chooser/Public/IObjectChooser.h
--- a/Engine/Plugins/Experimental/Chooser/Source/Chooser/Public/IObjectChooser.h
+++ b/Engine/Plugins/Experimental/Chooser/Source/Chooser/Public/IObjectChooser.h
@@ -51,7 +51,7 @@ struct CHOOSER_API FObjectChooserBase
 public:
 	virtual ~FObjectChooserBase() {}
 
-	enum class EIteratorStatus { Continue,Stop };
+	enum class EIteratorStatus { Continue, ContinueWithOutputs, Stop };
 
 	DECLARE_DELEGATE_RetVal_OneParam( EIteratorStatus, FObjectChooserIteratorCallback, UObject*);
 
