chooser struct support part 2: blueprint support for passing in/out structs from chooser evaluation

diff --git a/Engine/Plugins/Experimental/Chooser/Source/Chooser/Internal/Chooser.h b/Engine/Plugins/Experimental/Chooser/Source/Chooser/Internal/Chooser.h
--- a/Engine/Plugins/Experimental/Chooser/Source/Chooser/Internal/Chooser.h
+++ b/Engine/Plugins/Experimental/Chooser/Source/Chooser/Internal/Chooser.h
@@ -80,10 +80,12 @@ public:
 	TObjectPtr<UClass> OutputObjectType;
 
 	virtual TConstArrayView<FInstancedStruct> GetContextData() const override { return ContextData; }
+	
+	static FObjectChooserBase::EIteratorStatus EvaluateChooser(FChooserEvaluationContext& Context, const UChooserTable* Chooser, FObjectChooserBase::FObjectChooserIteratorCallback Callback);
 };
 
 
-USTRUCT(DisplayName = "Evaluate Chooser")
+USTRUCT(BlueprintType, DisplayName = "Evaluate Chooser")
 struct CHOOSER_API FEvaluateChooser : public FObjectChooserBase
 {
 	GENERATED_BODY()
@@ -91,6 +93,9 @@ struct CHOOSER_API FEvaluateChooser : public FObjectChooserBase
 	virtual UObject* ChooseObject(FChooserEvaluationContext& Context) const final override;
 	virtual EIteratorStatus ChooseMulti(FChooserEvaluationContext &Context, FObjectChooserIteratorCallback Callback) const final override;
 	public:
+
+	FEvaluateChooser() : Chooser(nullptr) {}
+	FEvaluateChooser(TObjectPtr<UChooserTable> Table) : Chooser(Table) {}
 	
 	UPROPERTY(EditAnywhere, Category="Parameters")
 	TObjectPtr<UChooserTable> Chooser;
diff --git a/Engine/Plugins/Experimental/Chooser/Source/Chooser/Private/Chooser.cpp b/Engine/Plugins/Experimental/Chooser/Source/Chooser/Private/Chooser.cpp
--- a/Engine/Plugins/Experimental/Chooser/Source/Chooser/Private/Chooser.cpp
+++ b/Engine/Plugins/Experimental/Chooser/Source/Chooser/Private/Chooser.cpp
@@ -37,9 +37,7 @@ void UChooserTable::PostEditChangeProperty(FPropertyChangedEvent& PropertyChange
 			CachedPreviousOutputObjectType = OutputObjectType;
 		}
 	}
-	
-	static FName ContextPropertyName = "ContextData";
-	if (PropertyChangedEvent.Property->GetName() == ContextPropertyName)
+	else
 	{
 		OnContextClassChanged.Broadcast();
 	}
@@ -61,6 +59,7 @@ void UChooserTable::PostLoad()
 		Context.Class = ContextObjectType_DEPRECATED;
 		Context.Direction = EContextObjectDirection::ReadWrite;
 		ContextObjectType_DEPRECATED = nullptr;
+		OnContextClassChanged.Broadcast();
 	}
 
 	if (Results_DEPRECATED.Num() > 0 || Columns_DEPRECATED.Num() > 0)
@@ -119,11 +118,11 @@ void UChooserTable::IterateRecentContextObjects(TFunction<void(const UObject*)>
 void UChooserTable::UpdateDebugging(FChooserEvaluationContext& Context) const
 {
 	FScopeLock Lock(&DebugLock);
-	for (const FContextEntry& Entry : Context.ContextData)
+	for (const FInstancedStruct& Param : Context.Params)
 	{
-		if (Entry.Data && Entry.Type && Entry.Type->IsA(UObject::StaticClass()))
+		if (const FChooserEvaluationInputObject* ObjectParam = Param.GetPtr<FChooserEvaluationInputObject>())
 		{
-			UObject* ContextObject = static_cast<UObject*>(Entry.Data);
+			UObject* ContextObject = ObjectParam->Object;
 			RecentContextObjects.Add(MakeWeakObjectPtr(ContextObject));
 			if (ContextObject == DebugTarget)
 			{
@@ -138,7 +137,7 @@ void UChooserTable::UpdateDebugging(FChooserEvaluationContext& Context) const
 
 #endif
 
-static FObjectChooserBase::EIteratorStatus StaticEvaluateChooser(FChooserEvaluationContext& Context, const UChooserTable* Chooser, FObjectChooserBase::FObjectChooserIteratorCallback Callback)
+FObjectChooserBase::EIteratorStatus UChooserTable::EvaluateChooser(FChooserEvaluationContext& Context, const UChooserTable* Chooser, FObjectChooserBase::FObjectChooserIteratorCallback Callback)
 {
 	if (Chooser == nullptr)
 	{
@@ -212,7 +211,7 @@ static FObjectChooserBase::EIteratorStatus StaticEvaluateChooser(FChooserEvaluat
 UObject* FEvaluateChooser::ChooseObject(FChooserEvaluationContext& Context) const
 {
 	UObject* Result = nullptr;
-	StaticEvaluateChooser(Context, Chooser, FObjectChooserIteratorCallback::CreateLambda([&Result](UObject* InResult)
+	UChooserTable::EvaluateChooser(Context, Chooser, FObjectChooserIteratorCallback::CreateLambda([&Result](UObject* InResult)
 	{
 		Result = InResult;
 		return FObjectChooserBase::EIteratorStatus::Stop;
@@ -223,58 +222,5 @@ UObject* FEvaluateChooser::ChooseObject(FChooserEvaluationContext& Context) cons
 
 FObjectChooserBase::EIteratorStatus FEvaluateChooser::ChooseMulti(FChooserEvaluationContext& Context, FObjectChooserIteratorCallback Callback) const
 {
-	return StaticEvaluateChooser(Context, Chooser, Callback);
-}
-
-UChooserFunctionLibrary::UChooserFunctionLibrary(const FObjectInitializer& ObjectInitializer)
-	: Super(ObjectInitializer)
-{
-}
-
-static FObjectChooserBase::EIteratorStatus StaticEvaluateChooser(const UObject* ContextObject, const UChooserTable* Chooser, FObjectChooserBase::FObjectChooserIteratorCallback Callback)
-{
-	// Fallback single context object version
-	FChooserEvaluationContext Context;
-	Context.ContextData.Push(FContextEntry{ContextObject->GetClass(), const_cast<void*>(static_cast<const void*>(ContextObject))});
-	return StaticEvaluateChooser(Context, Chooser, Callback);
-}
-
-UObject* UChooserFunctionLibrary::EvaluateChooser(const UObject* ContextObject, const UChooserTable* Chooser, TSubclassOf<UObject> ObjectClass)
-{
-	UObject* Result = nullptr;
-	StaticEvaluateChooser(ContextObject, Chooser, FObjectChooserBase::FObjectChooserIteratorCallback::CreateLambda([&Result, ObjectClass](UObject* InResult)
-	{
-		if (InResult && (ObjectClass == nullptr || InResult->IsA(ObjectClass)))
-		{
-			Result = InResult;
-			return FObjectChooserBase::EIteratorStatus::Stop;
-		}
-		return FObjectChooserBase::EIteratorStatus::Continue;
-	}));
-
-	return Result;
-}
-
-TArray<UObject*> UChooserFunctionLibrary::EvaluateChooserMulti(const UObject* ContextObject, const UChooserTable* Chooser, TSubclassOf<UObject> ObjectClass)
-{
-	TArray<UObject*> Result;
-	StaticEvaluateChooser(ContextObject, Chooser, FObjectChooserBase::FObjectChooserIteratorCallback::CreateLambda([&Result, ObjectClass](UObject* InResult)
-	{
-		if (InResult && (ObjectClass == nullptr || InResult->IsA(ObjectClass)))
-		{
-			Result.Add(InResult);
-		}
-		// trigger output columns only on the first result
-		return Result.Num() == 1 ? FObjectChooserBase::EIteratorStatus::ContinueWithOutputs : FObjectChooserBase::EIteratorStatus::Continue;
-	}));
-
-	for (int Index = 0; Index < Result.Num(); ++Index)
-	{
-		if (Result[Index] && !Result[Index]->IsA(ObjectClass))
-		{
-			Result[Index] = nullptr;
-		}
-	}
-
-	return Result;
-}
+	return UChooserTable::EvaluateChooser(Context, Chooser, Callback);
+}
\ No newline at end of file
diff --git a/Engine/Plugins/Experimental/Chooser/Source/Chooser/Private/ChooserFunctionLibrary.cpp b/Engine/Plugins/Experimental/Chooser/Source/Chooser/Private/ChooserFunctionLibrary.cpp
new file mode 100644
--- /dev/null
+++ b/Engine/Plugins/Experimental/Chooser/Source/Chooser/Private/ChooserFunctionLibrary.cpp
@@ -0,0 +1,192 @@
+// Copyright Epic Games, Inc. All Rights Reserved.
+#include "ChooserFunctionLibrary.h"
+#include "Chooser.h"
+
+#include UE_INLINE_GENERATED_CPP_BY_NAME(ChooserFunctionLibrary)
+
+UChooserFunctionLibrary::UChooserFunctionLibrary(const FObjectInitializer& ObjectInitializer)
+	: Super(ObjectInitializer)
+{
+}
+
+static FObjectChooserBase::EIteratorStatus StaticEvaluateChooser(const UObject* ContextObject, const UChooserTable* Chooser, FObjectChooserBase::FObjectChooserIteratorCallback Callback)
+{
+	// Fallback single context object version
+	FChooserEvaluationContext Context;
+	Context.Params.AddDefaulted();
+	Context.Params.Last().InitializeAs(FChooserEvaluationInputObject::StaticStruct());
+	Context.Params.Last().GetMutable<FChooserEvaluationInputObject>().Object = const_cast<UObject*>(ContextObject);
+	return UChooserTable::EvaluateChooser(Context, Chooser, Callback);
+}
+
+UObject* UChooserFunctionLibrary::EvaluateChooser(const UObject* ContextObject, const UChooserTable* Chooser, TSubclassOf<UObject> ObjectClass)
+{
+	UObject* Result = nullptr;
+	StaticEvaluateChooser(ContextObject, Chooser, FObjectChooserBase::FObjectChooserIteratorCallback::CreateLambda([&Result, ObjectClass](UObject* InResult)
+	{
+		if (InResult && (ObjectClass == nullptr || InResult->IsA(ObjectClass)))
+		{
+			Result = InResult;
+			return FObjectChooserBase::EIteratorStatus::Stop;
+		}
+		return FObjectChooserBase::EIteratorStatus::Continue;
+	}));
+
+	return Result;
+}
+
+void UChooserFunctionLibrary::AddChooserObjectInput(FChooserEvaluationContext& Context, UObject* Object)
+{
+	Context.Params.AddDefaulted();
+	Context.Params.Last().InitializeAs(FChooserEvaluationInputObject::StaticStruct());
+	Context.Params.Last().GetMutable<FChooserEvaluationInputObject>().Object = Object;
+}
+
+void UChooserFunctionLibrary::AddChooserStructInput(FChooserEvaluationContext& Inputs, int32 Value)
+{
+	checkNoEntry();
+}
+
+DEFINE_FUNCTION(UChooserFunctionLibrary::execAddChooserStructInput)
+{
+	// Read wildcard Value input.
+	Stack.MostRecentPropertyAddress = nullptr;
+	Stack.MostRecentPropertyContainer = nullptr;
+	
+	P_GET_STRUCT_REF(FChooserEvaluationContext, ChooserContext);
+	
+	Stack.StepCompiledIn<FStructProperty>(nullptr);
+	const FStructProperty* ValueProp = CastField<FStructProperty>(Stack.MostRecentProperty);
+	const void* ValuePtr = Stack.MostRecentPropertyAddress;
+	P_FINISH;
+	
+	if (!ValueProp || !ValuePtr)
+	{
+		FBlueprintExceptionInfo ExceptionInfo(
+			EBlueprintExceptionType::AbortExecution,
+			NSLOCTEXT("ChooserTable", "InstancedStruct_MakeInvalidValueWarning", "Failed to resolve the Value for Add Chooser Struct Input")
+		);
+
+		FBlueprintCoreDelegates::ThrowScriptException(P_THIS, Stack, ExceptionInfo);
+	}
+	else
+	{
+		P_NATIVE_BEGIN;
+		ChooserContext.Params.AddDefaulted();
+		ChooserContext.Params.Last().InitializeAs(ValueProp->Struct, static_cast<const uint8*>(ValuePtr));
+		P_NATIVE_END;
+	}
+}
+
+void UChooserFunctionLibrary::GetChooserStructOutput(FChooserEvaluationContext& Context, int index, int32& Value)
+{
+	checkNoEntry();
+}
+
+DEFINE_FUNCTION(UChooserFunctionLibrary::execGetChooserStructOutput)
+{
+	P_GET_STRUCT_REF(FChooserEvaluationContext, Outputs);
+	P_GET_PROPERTY(FIntProperty, Index);
+
+	// Read wildcard Value input.
+	Stack.MostRecentPropertyAddress = nullptr;
+	Stack.MostRecentPropertyContainer = nullptr;
+	Stack.StepCompiledIn<FStructProperty>(nullptr);
+	
+	const FStructProperty* ValueProp = CastField<FStructProperty>(Stack.MostRecentProperty);
+	void* ValuePtr = Stack.MostRecentPropertyAddress;
+
+	P_FINISH;
+
+	if (!ValueProp || !ValuePtr)
+	{
+		FBlueprintExceptionInfo ExceptionInfo(
+			EBlueprintExceptionType::AbortExecution,
+			NSLOCTEXT("ChooserTable","GetChooserStructOutput_InvalidValueWarning", "Failed to resolve the Value for Get Chooser Struct Output")
+		);
+
+		FBlueprintCoreDelegates::ThrowScriptException(P_THIS, Stack, ExceptionInfo);
+	}
+	else
+	{
+		P_NATIVE_BEGIN;
+		if (Outputs.Params.Num() > Index && Outputs.Params[Index].IsValid() && Outputs.Params[Index].GetScriptStruct()->IsChildOf(ValueProp->Struct))
+		{
+			ValueProp->Struct->CopyScriptStruct(ValuePtr, Outputs.Params[Index].GetMemory());
+		}
+		P_NATIVE_END;
+	}
+}
+
+UObject* UChooserFunctionLibrary::EvaluateObjectChooserBase(FChooserEvaluationContext& Context, const FInstancedStruct& ObjectChooser, TSubclassOf<UObject> ObjectClass)
+ {
+ 	UObject* Result = nullptr;
+
+	if (const FObjectChooserBase* ObjectChooserPtr = ObjectChooser.GetPtr<FObjectChooserBase>())
+	{
+		ObjectChooserPtr->ChooseMulti(Context, FObjectChooserBase::FObjectChooserIteratorCallback::CreateLambda([&Result, ObjectClass](UObject* InResult)
+		{
+			if (InResult && (ObjectClass == nullptr || InResult->IsA(ObjectClass)))
+			{
+				Result = InResult;
+				return FObjectChooserBase::EIteratorStatus::Stop;
+			}
+			return FObjectChooserBase::EIteratorStatus::Continue;
+		}));
+	}
+	
+ 	return Result;
+ }
+
+
+TArray<UObject*> UChooserFunctionLibrary::EvaluateObjectChooserBaseMulti(FChooserEvaluationContext& Context, const FInstancedStruct& ObjectChooser, TSubclassOf<UObject> ObjectClass)
+{
+ 	TArray<UObject*> Result;
+
+	if (const FObjectChooserBase* ObjectChooserPtr = ObjectChooser.GetPtr<FObjectChooserBase>())
+	{
+		ObjectChooserPtr->ChooseMulti(Context, FObjectChooserBase::FObjectChooserIteratorCallback::CreateLambda([&Result, ObjectClass](UObject* InResult)
+		{
+			if (InResult && (ObjectClass == nullptr || InResult->IsA(ObjectClass)))
+			{
+				Result.Add(InResult);
+			}
+			// trigger output columns only on the first result
+			return Result.Num() == 1 ? FObjectChooserBase::EIteratorStatus::ContinueWithOutputs : FObjectChooserBase::EIteratorStatus::Continue;
+		}));
+	}
+	
+ 	return Result;
+ }
+
+FInstancedStruct UChooserFunctionLibrary::MakeEvaluateChooser(UChooserTable* Chooser)
+{
+	FInstancedStruct Struct;
+	Struct.InitializeAs(FEvaluateChooser::StaticStruct());
+	Struct.GetMutable<FEvaluateChooser>().Chooser = Chooser;
+	return Struct;
+}
+
+TArray<UObject*> UChooserFunctionLibrary::EvaluateChooserMulti(const UObject* ContextObject, const UChooserTable* Chooser, TSubclassOf<UObject> ObjectClass)
+{
+	TArray<UObject*> Result;
+	StaticEvaluateChooser(ContextObject, Chooser, FObjectChooserBase::FObjectChooserIteratorCallback::CreateLambda([&Result, ObjectClass](UObject* InResult)
+	{
+		if (InResult && (ObjectClass == nullptr || InResult->IsA(ObjectClass)))
+		{
+			Result.Add(InResult);
+		}
+		// trigger output columns only on the first result
+		return Result.Num() == 1 ? FObjectChooserBase::EIteratorStatus::ContinueWithOutputs : FObjectChooserBase::EIteratorStatus::Continue;
+	}));
+
+	for (int Index = 0; Index < Result.Num(); ++Index)
+	{
+		if (Result[Index] && !Result[Index]->IsA(ObjectClass))
+		{
+			Result[Index] = nullptr;
+		}
+	}
+
+	return Result;
+}
diff --git a/Engine/Plugins/Experimental/Chooser/Source/Chooser/Private/ChooserPropertyAccess.cpp b/Engine/Plugins/Experimental/Chooser/Source/Chooser/Private/ChooserPropertyAccess.cpp
--- a/Engine/Plugins/Experimental/Chooser/Source/Chooser/Private/ChooserPropertyAccess.cpp
+++ b/Engine/Plugins/Experimental/Chooser/Source/Chooser/Private/ChooserPropertyAccess.cpp
@@ -13,10 +13,18 @@ namespace UE::Chooser
 {
 	bool ResolvePropertyChain(FChooserEvaluationContext& Context, const FChooserPropertyBinding& PropertyBinding, const void*& OutContainer, const UStruct*& OutStructType)
 	{
-		if (PropertyBinding.ContextIndex < Context.ContextData.Num())
+		if (PropertyBinding.ContextIndex < Context.Params.Num())
 		{
-			OutContainer = Context.ContextData[PropertyBinding.ContextIndex].Data;
-			OutStructType = Context.ContextData[PropertyBinding.ContextIndex].Type;
+			if (FChooserEvaluationInputObject* ObjectParam = Context.Params[PropertyBinding.ContextIndex].GetMutablePtr<FChooserEvaluationInputObject>())
+			{
+				OutContainer = ObjectParam->Object;
+				OutStructType = ObjectParam->Object->GetClass();
+			}
+			else
+			{
+				OutContainer = Context.Params[PropertyBinding.ContextIndex].GetMutableMemory();
+				OutStructType = Context.Params[PropertyBinding.ContextIndex].GetScriptStruct();
+			}
 
 			return ResolvePropertyChain(OutContainer, OutStructType, PropertyBinding.PropertyBindingChain);
 		}
diff --git a/Engine/Plugins/Experimental/Chooser/Source/Chooser/Public/ChooserFunctionLibrary.h b/Engine/Plugins/Experimental/Chooser/Source/Chooser/Public/ChooserFunctionLibrary.h
--- a/Engine/Plugins/Experimental/Chooser/Source/Chooser/Public/ChooserFunctionLibrary.h
+++ b/Engine/Plugins/Experimental/Chooser/Source/Chooser/Public/ChooserFunctionLibrary.h
@@ -9,7 +9,7 @@
 #include "ChooserFunctionLibrary.generated.h"
 
 /**
- * Morpheus Extensions Function Library
+ * Chooser Function Library
  */
 UCLASS()
 class CHOOSER_API UChooserFunctionLibrary : public UBlueprintFunctionLibrary
@@ -22,11 +22,11 @@ public:
 	*
 	* @param ContextObject			(in) An Object from which the parameters to the Chooser Table will be read
 	* @param ChooserTable			(in) The ChooserTable asset
-	* @param ObjectClass			(in) Expected type of result objects
+	* @param ObjectClass			(in) Expected type of result object
 	*/
-	UFUNCTION(BlueprintPure, meta = (BlueprintThreadSafe, DeterminesOutputType = "ObjectClass"), Category = "Animation")
+	UFUNCTION(BlueprintPure, meta = (BlueprintThreadSafe, BlueprintInternalUseOnly = "true", DeterminesOutputType = "ObjectClass"))
 	static UObject* EvaluateChooser(const UObject* ContextObject, const UChooserTable* ChooserTable, TSubclassOf<UObject> ObjectClass);
-
+	
 	/**
     * Evaluate a chooser table and return the list of all selected UObjects
     *
@@ -34,6 +34,66 @@ public:
     * @param ChooserTable			(in) The ChooserTable asset
     * @param ObjectClass			(in) Expected type of result objects
     */
-    UFUNCTION(BlueprintPure, meta = (BlueprintThreadSafe, DeterminesOutputType = "ObjectClass"), Category = "Animation")
+    UFUNCTION(BlueprintPure, meta = (BlueprintThreadSafe, BlueprintInternalUseOnly = "true", DeterminesOutputType = "ObjectClass"))
     static TArray<UObject*> EvaluateChooserMulti(const UObject* ContextObject, const UChooserTable* ChooserTable, TSubclassOf<UObject> ObjectClass);
+	
+	/**
+	* Evaluate an ObjectChooserBase and return the selected UObject, or null
+	*
+	* @param Context			(in) A struct reference to the chooser evaluation context
+	* @param ObjectChooser		(in) An Instanced struct containing an ObjectChooserBase implementation, such as EvaluateChooser, or EvaluateProxyAsset
+	* @param ObjectClass		(in) Expected type of result object
+	*/
+	UFUNCTION(BlueprintCallable, meta = (BlueprintThreadSafe, BlueprintInternalUseOnly = "true", DeterminesOutputType = "ObjectClass"))
+	static UObject* EvaluateObjectChooserBase(UPARAM(Ref) FChooserEvaluationContext& Context,UPARAM(Ref) const FInstancedStruct& ObjectChooser, TSubclassOf<UObject> ObjectClass);
+
+	/**
+	* Evaluate a chooser table and return all selected UObjects
+	*
+	* @param Context			(in) A struct reference to the chooser evaluation context
+	* @param ObjectChooser		(in) An Instanced struct containing an ObjectChooserBase implementation, such as EvaluateChooser, or EvaluateProxyAsset
+	* @param ObjectClass		(in) Expected type of result object
+	*/
+	UFUNCTION(BlueprintCallable, meta = (BlueprintThreadSafe, BlueprintInternalUseOnly = "true", DeterminesOutputType = "ObjectClass"))
+	static TArray<UObject*> EvaluateObjectChooserBaseMulti(UPARAM(Ref) FChooserEvaluationContext& Context,UPARAM(Ref) const FInstancedStruct& ObjectChooser, TSubclassOf<UObject> ObjectClass);
+
+
+	/**
+	* Add an Object to a ChooserEvaluation context
+	*
+	* @param Context			(in) A struct reference to the chooser evaluation context
+	* @param Object				(in) The Object to add
+	*/
+	UFUNCTION(BlueprintCallable, meta = (BlueprintThreadSafe, BlueprintInternalUseOnly = "true"))
+	static void AddChooserObjectInput(UPARAM(Ref) FChooserEvaluationContext& Context, UObject* Object);
+	
+	/**
+	* Add a Struct to a ChooserEvaluation context
+	*
+	* @param Context			(in) A struct reference to the chooser evaluation context
+	* @param Object				(in) The Object to add
+	*/
+	UFUNCTION(BlueprintCallable, CustomThunk, meta = (BlueprintThreadSafe, BlueprintInternalUseOnly = "true", CustomStructureParam = "Value"))
+	static void AddChooserStructInput(UPARAM(Ref) FChooserEvaluationContext& Context, int32 Value);
+	
+	DECLARE_FUNCTION(execAddChooserStructInput);
+	
+	/**
+	* Get a Struct to a ChooserEvaluation context
+	*
+	* @param Context			(in) A struct reference to the chooser evaluation context
+	* @param Object				(in) The Object to add
+	*/
+	UFUNCTION(BlueprintPure, CustomThunk, meta = (BlueprintThreadSafe, BlueprintInternalUseOnly = "true", CustomStructureParam = "Value"))
+	static void GetChooserStructOutput(UPARAM(Ref) FChooserEvaluationContext& Context, int Index, int32& Value);
+    	
+	DECLARE_FUNCTION(execGetChooserStructOutput);
+
+	/**
+	* Create an EvaluateChooser struct
+	*
+	* @param Chooser				(in) the ChooserTable asset to evaluate
+	*/
+	UFUNCTION(BlueprintPure, Category = "Animation", meta = (BlueprintThreadSafe, BlueprintInternalUseOnly="true", NativeMakeFunc))
+	static FInstancedStruct MakeEvaluateChooser(UChooserTable* Chooser);
 };
diff --git a/Engine/Plugins/Experimental/Chooser/Source/Chooser/Public/IObjectChooser.h b/Engine/Plugins/Experimental/Chooser/Source/Chooser/Public/IObjectChooser.h
--- a/Engine/Plugins/Experimental/Chooser/Source/Chooser/Public/IObjectChooser.h
+++ b/Engine/Plugins/Experimental/Chooser/Source/Chooser/Public/IObjectChooser.h
@@ -21,12 +21,6 @@ public:
 	virtual void ConvertToInstancedStruct(FInstancedStruct& OutInstancedStruct) const { }
 };
 
-struct FContextEntry
-{
-	const UStruct* Type = nullptr;
-	void* Data = nullptr;
-};
-
 #if WITH_EDITOR
 struct FChooserDebuggingInfo
 {
@@ -34,13 +28,23 @@ struct FChooserDebuggingInfo
 };
 #endif
 
-struct FChooserEvaluationContext
+
+USTRUCT()
+struct FChooserEvaluationInputObject
 {
-#if WITH_EDITOR
-	FChooserDebuggingInfo DebuggingInfo;
-#endif
+	GENERATED_BODY()
+	TObjectPtr<UObject> Object;
+};
 
-	TArray<FContextEntry, TFixedAllocator<4>> ContextData;
+USTRUCT(BlueprintType)
+struct FChooserEvaluationContext
+{
+	#if WITH_EDITOR
+    	FChooserDebuggingInfo DebuggingInfo;
+    #endif
+	
+	GENERATED_BODY()
+	TArray<FInstancedStruct, TFixedAllocator<4>> Params;
 };
 
 USTRUCT()
diff --git a/Engine/Plugins/Experimental/Chooser/Source/ChooserUncooked/Private/EvaluateChooserNode.cpp b/Engine/Plugins/Experimental/Chooser/Source/ChooserUncooked/Private/EvaluateChooserNode.cpp
--- a/Engine/Plugins/Experimental/Chooser/Source/ChooserUncooked/Private/EvaluateChooserNode.cpp
+++ b/Engine/Plugins/Experimental/Chooser/Source/ChooserUncooked/Private/EvaluateChooserNode.cpp
@@ -15,6 +15,7 @@
 #include "Internationalization/Internationalization.h"
 #include "K2Node_CallFunction.h"
 #include "K2Node_MakeArray.h"
+#include "K2Node_MakeStruct.h"
 #include "Kismet2/BlueprintEditorUtils.h"
 #include "Kismet2/CompilerResultsLog.h"
 #include "KismetCompiler.h"
@@ -28,11 +29,11 @@
 #include "UObject/WeakObjectPtrTemplates.h"
 #include "K2Node_Self.h"
 
-
 #define LOCTEXT_NAMESPACE "EvaluateChooserNode"
 
 /////////////////////////////////////////////////////
 // UK2Node_EvaluateChooser
+// old implementation of this node for backwards compatibility - not currently accessible to create new instances in content
 
 UK2Node_EvaluateChooser::UK2Node_EvaluateChooser(const FObjectInitializer& ObjectInitializer)
 	: Super(ObjectInitializer)
@@ -132,14 +133,7 @@ void UK2Node_EvaluateChooser::AllocateDefaultPins()
 
 FText UK2Node_EvaluateChooser::GetNodeTitle(ENodeTitleType::Type TitleType) const
 {
-	if (Chooser)
-	{
-		return FText::FromString(Chooser->GetName());
-	}
-	else
-	{
-		return LOCTEXT("EvaluateChooser_Title", "Evaluate Chooser");
-	}
+	return LOCTEXT("EvaluateChooser_Title", "Evaluate Chooser (Deprecated)");
 }
 
 FText UK2Node_EvaluateChooser::GetPinDisplayName(const UEdGraphPin* Pin) const
@@ -313,4 +307,532 @@ FText UK2Node_EvaluateChooser::GetMenuCategory() const
 }
 
 
+//----------------------------------------------------------------------------------------------
+// UK2Node_EvaluateChooser2
+// New Implementation of EvaluateChooser with support for passing in/out multiple objects and structs
+
+UK2Node_EvaluateChooser2::UK2Node_EvaluateChooser2(const FObjectInitializer& ObjectInitializer)
+	: Super(ObjectInitializer)
+{
+	NodeTooltip = LOCTEXT("NodeTooltip", "Evaluates an Chooser Table, and returns the resulting Object or Objects.");
+}
+
+void UK2Node_EvaluateChooser2::UnregisterChooserCallback()
+{
+	if(CurrentCallbackChooser)
+	{
+		CurrentCallbackChooser->OnOutputObjectTypeChanged.RemoveAll(this);
+		CurrentCallbackChooser->OnContextClassChanged.RemoveAll(this);
+		CurrentCallbackChooser = nullptr;
+	}
+}
+
+void UK2Node_EvaluateChooser2::BeginDestroy()
+{
+	UnregisterChooserCallback();
+	
+	Super::BeginDestroy();
+}
+
+void UK2Node_EvaluateChooser2::PostEditUndo()
+{
+	Super::PostEditUndo();
+	ChooserChanged();
+}
+
+void UK2Node_EvaluateChooser2::DestroyNode()
+{
+	UnregisterChooserCallback();
+	Super::DestroyNode();
+}
+
+
+void UK2Node_EvaluateChooser2::ChooserChanged()
+{
+	if (Chooser != CurrentCallbackChooser)
+	{
+		UnregisterChooserCallback();
+	
+		if (Chooser)
+		{
+			Chooser->OnOutputObjectTypeChanged.AddUObject(this, &UK2Node_EvaluateChooser2::ResultTypeChanged);
+			Chooser->OnContextClassChanged.AddUObject(this, &UK2Node::ReconstructNode);
+		}
+	
+		CurrentCallbackChooser = Chooser;
+
+		ReconstructNode();
+	}
+}
+
+void UK2Node_EvaluateChooser2::ResultTypeChanged(const UClass*)
+{
+	AllocateDefaultPins();
+}
+
+void UK2Node_EvaluateChooser2::PreloadRequiredAssets()
+{
+	if (Chooser)
+	{
+		if (FLinkerLoad* ObjLinker = GetLinker())
+		{
+			ObjLinker->Preload(Chooser);
+		}
+	}
+    		
+	Super::PreloadRequiredAssets();
+}
+
+void UK2Node_EvaluateChooser2::AllocateDefaultPins()
+{
+	Super::AllocateDefaultPins();
+	
+	UClass* ChooserResultType = UObject::StaticClass();
+
+	if (!FindPin(UEdGraphSchema_K2::PN_Execute, EGPD_Input))
+	{
+		// Input - Execution Pin
+		CreatePin(EGPD_Input, UEdGraphSchema_K2::PC_Exec, UEdGraphSchema_K2::PN_Execute);
+	}
+
+	if (!FindPin(UEdGraphSchema_K2::PN_Execute, EGPD_Output))
+	{
+		// Output - Execution Pin
+		CreatePin(EGPD_Output, UEdGraphSchema_K2::PC_Exec, UEdGraphSchema_K2::PN_Execute);
+	}
+
+	if (Chooser)
+	{
+		if (Chooser->ContextObjectType_DEPRECATED)
+		{
+			// post load fix up code doesn't run until it's too late when loading a blueprint that references an chooser
+			// need to manually trigger it here so we generate the right pins
+			Chooser->PostLoad();
+		}
+		
+		for(FInstancedStruct& ContextDataEntry : Chooser->ContextData)
+		{
+			if (ContextDataEntry.IsValid())
+			{
+				const UScriptStruct* EntryType = ContextDataEntry.GetScriptStruct();
+				if (EntryType == FContextObjectTypeClass::StaticStruct())
+				{
+					const FContextObjectTypeClass& ClassContext = ContextDataEntry.Get<FContextObjectTypeClass>();
+					if (UEdGraphPin* Pin = FindPin(ClassContext.Class.GetFName(), EGPD_Input))
+					{
+						Pin->PinType.PinSubCategoryObject = ClassContext.Class;
+					}
+					else
+					{
+						CreatePin(EGPD_Input, UEdGraphSchema_K2::PC_Object, ClassContext.Class, ClassContext.Class.GetFName());
+					}
+				}
+				else if (EntryType == FContextObjectTypeStruct::StaticStruct())
+				{
+					const FContextObjectTypeStruct& StructContext = ContextDataEntry.Get<FContextObjectTypeStruct>();
+
+					if (StructContext.Direction == EContextObjectDirection::Read || StructContext.Direction == EContextObjectDirection::ReadWrite)
+					{
+						if (UEdGraphPin* Pin = FindPin(StructContext.Struct.GetFName(), EGPD_Input))
+						{
+							Pin->PinType.PinSubCategoryObject = StructContext.Struct;
+						}
+						else
+						{
+							CreatePin(EGPD_Input, UEdGraphSchema_K2::PC_Struct, StructContext.Struct, StructContext.Struct.GetFName());
+						}
+					}
+
+					if (StructContext.Direction == EContextObjectDirection::Write || StructContext.Direction == EContextObjectDirection::ReadWrite)
+					{
+						if (UEdGraphPin* Pin = FindPin(StructContext.Struct.GetFName(), EGPD_Output))
+						{
+							Pin->PinType.PinSubCategoryObject = StructContext.Struct;
+						}
+						else
+						{
+							CreatePin(EGPD_Output, UEdGraphSchema_K2::PC_Struct, StructContext.Struct, StructContext.Struct.GetFName());
+						}
+					}
+
+					if (StructContext.Direction == EContextObjectDirection::Read)
+					{
+					    // for read only structs, remove the output pin if found
+						if (UEdGraphPin* Pin = FindPin(StructContext.Struct.GetFName(), EGPD_Output))
+						{
+							RemovePin(Pin);
+						}
+					}
+					
+					if (StructContext.Direction == EContextObjectDirection::Write)
+					{
+					    // for write only structs, remove the input pin if found
+						if (UEdGraphPin* Pin = FindPin(StructContext.Struct.GetFName(), EGPD_Input))
+						{
+							RemovePin(Pin);
+						}
+					}
+				}
+			}
+		}
+	}
+
+   	if (Chooser && Chooser->OutputObjectType)
+   	{
+   		ChooserResultType = Chooser->OutputObjectType;
+   	}
+
+	static const FName ResultPinName = "Result";
+	if (UEdGraphPin* ResultPin = FindPin(ResultPinName, EGPD_Output))
+	{
+		ResultPin->PinType.PinSubCategoryObject = ChooserResultType;
+		ResultPin->PinType.ContainerType = (Mode == EEvaluateChooserMode::AllResults) ? EPinContainerType::Array : EPinContainerType::None;
+	}
+	else
+	{
+		UEdGraphNode::FCreatePinParams PinParams;
+		PinParams.ContainerType = (Mode == EEvaluateChooserMode::AllResults) ? EPinContainerType::Array : EPinContainerType::None;
+		PinParams.ValueTerminalType.TerminalCategory = UEdGraphSchema_K2::PC_Object;
+				
+		CreatePin(EGPD_Output, UEdGraphSchema_K2::PC_Object, ChooserResultType, TEXT("Result"), PinParams);
+	}
+}
+
+FText UK2Node_EvaluateChooser2::GetNodeTitle(ENodeTitleType::Type TitleType) const
+{
+	if (Chooser)
+	{
+		return FText::Format(LOCTEXT("EvaluateChooser_TitleWithChooser", "Evaluate Chooser: {0}"), {FText::FromString(Chooser->GetName())});
+	}
+	else
+	{
+		return LOCTEXT("EvaluateChooser_Title", "Evaluate Chooser");
+	}
+}
+
+FText UK2Node_EvaluateChooser2::GetPinDisplayName(const UEdGraphPin* Pin) const
+{
+	if (Pin->PinName == UEdGraphSchema_K2::PN_Execute)
+	{
+		return FText();
+	}
+	return FText::FromName(Pin->PinName);
+}
+
+void UK2Node_EvaluateChooser2::PostEditChangeProperty(struct FPropertyChangedEvent& PropertyChangedEvent)
+{
+	if (PropertyChangedEvent.Property->GetName() == "Chooser")
+	{
+		ChooserChanged();
+	}
+	if (PropertyChangedEvent.Property->GetName() == "Mode")
+	{
+		ReconstructNode();
+	}
+	Super::PostEditChangeProperty(PropertyChangedEvent);
+}
+
+void UK2Node_EvaluateChooser2::PostLoad()
+{
+	Super::PostLoad();
+	ChooserChanged();
+}
+
+void UK2Node_EvaluateChooser2::PinConnectionListChanged(UEdGraphPin* Pin)
+{
+
+	Modify();
+
+	FBlueprintEditorUtils::MarkBlueprintAsStructurallyModified(GetBlueprint());
+}
+
+void UK2Node_EvaluateChooser2::PinDefaultValueChanged(UEdGraphPin* Pin)
+{
+	Super::PinDefaultValueChanged(Pin);
+}
+
+void UK2Node_EvaluateChooser2::PinTypeChanged(UEdGraphPin* Pin)
+{
+	Super::PinTypeChanged(Pin);
+}
+
+FText UK2Node_EvaluateChooser2::GetTooltipText() const
+{
+	return NodeTooltip;
+}
+
+void UK2Node_EvaluateChooser2::PostReconstructNode()
+{
+	Super::PostReconstructNode();
+}
+
+void UK2Node_EvaluateChooser2::ExpandNode(class FKismetCompilerContext& CompilerContext, UEdGraph* SourceGraph)
+{
+	Super::ExpandNode(CompilerContext, SourceGraph);
+
+	UEdGraphPin* ExecInput = GetExecPin();
+	UEdGraphPin* ExecOutput = GetPassThroughPin(ExecInput);
+	UEdGraphPin* ResultPin = FindPinChecked(TEXT("Result"));
+
+	UEdGraphPin* PreviousNodeExecOutput = nullptr;
+	
+	if (ExecInput->HasAnyConnections())
+	{
+		UK2Node_CallFunction* CallFunction = CompilerContext.SpawnIntermediateNode<UK2Node_CallFunction>(this, SourceGraph);
+		CompilerContext.MessageLog.NotifyIntermediateObjectCreation(CallFunction, this);
+
+		if (Mode == EEvaluateChooserMode::AllResults)
+		{
+			CallFunction->SetFromFunction(UChooserFunctionLibrary::StaticClass()->FindFunctionByName(GET_MEMBER_NAME_CHECKED(UChooserFunctionLibrary, EvaluateObjectChooserBaseMulti)));
+		}
+		else
+		{
+			CallFunction->SetFromFunction(UChooserFunctionLibrary::StaticClass()->FindFunctionByName(GET_MEMBER_NAME_CHECKED(UChooserFunctionLibrary, EvaluateObjectChooserBase)));
+		}
+		CallFunction->AllocateDefaultPins();
+		
+		UK2Node_MakeStruct* ContextStructNode = CompilerContext.SpawnIntermediateNode<UK2Node_MakeStruct>(this, SourceGraph);
+        CompilerContext.MessageLog.NotifyIntermediateObjectCreation(ContextStructNode, this);
+		ContextStructNode->StructType = FChooserEvaluationContext::StaticStruct();
+        ContextStructNode->AllocateDefaultPins();
+		UEdGraphPin* ContextStructPin = ContextStructNode->FindPin(FChooserEvaluationContext::StaticStruct()->GetFName(), EGPD_Output);
+
+		ContextStructPin->MakeLinkTo(CallFunction->FindPin(FName("Context")));
+
+		UK2Node_Self* SelfNode = CompilerContext.SpawnIntermediateNode<UK2Node_Self>(this, SourceGraph);
+		CompilerContext.MessageLog.NotifyIntermediateObjectCreation(SelfNode, this);
+		SelfNode->AllocateDefaultPins();
+		UEdGraphPin* SelfPin = SelfNode->FindPin(TEXT("self"));
+		
+		if (Chooser)
+		{
+			int ContextDataCount = Chooser->ContextData.Num();
+			for(int ContextDataIndex = 0; ContextDataIndex < ContextDataCount; ContextDataIndex++)
+			{
+				FInstancedStruct& ContextDataEntry = Chooser->ContextData[ContextDataIndex];
+				if (ContextDataEntry.IsValid())
+				{
+					const UScriptStruct* EntryType = ContextDataEntry.GetScriptStruct();
+					if (EntryType == FContextObjectTypeClass::StaticStruct())
+					{
+						const FContextObjectTypeClass& ClassContext = ContextDataEntry.Get<FContextObjectTypeClass>();
+						if (UEdGraphPin* Pin = FindPin(ClassContext.Class.GetFName(), EGPD_Input))
+						{
+							UK2Node_CallFunction* AddObjectFunction = CompilerContext.SpawnIntermediateNode<UK2Node_CallFunction>(this, SourceGraph);
+                       		CompilerContext.MessageLog.NotifyIntermediateObjectCreation(AddObjectFunction, this);
+							AddObjectFunction->SetFromFunction(UChooserFunctionLibrary::StaticClass()->FindFunctionByName(GET_MEMBER_NAME_CHECKED(UChooserFunctionLibrary, AddChooserObjectInput)));
+							AddObjectFunction->AllocateDefaultPins();
+							ContextStructPin->MakeLinkTo(AddObjectFunction->FindPin(FName("Context")));
+
+							if (PreviousNodeExecOutput == nullptr)
+							{
+								CompilerContext.MovePinLinksToIntermediate(*ExecInput, *AddObjectFunction->GetExecPin());
+							}
+							else
+							{
+								PreviousNodeExecOutput->MakeLinkTo(AddObjectFunction->GetExecPin());
+							}
+							PreviousNodeExecOutput = AddObjectFunction->GetThenPin();
+							
+							UEdGraphPin* AddObjectPin = AddObjectFunction->FindPin(FName("Object"));
+							
+							if (Pin->HasAnyConnections())
+							{
+								CompilerContext.MovePinLinksToIntermediate(*Pin, *AddObjectPin);
+							}
+							else // would be nice to check that self is the same type as this Object parameter
+							{
+								// auto connect self node to any disconnected object pin
+								SelfPin->MakeLinkTo(AddObjectPin);
+							}
+						}
+					}
+					else if (EntryType == FContextObjectTypeStruct::StaticStruct())
+					{
+						const FContextObjectTypeStruct& StructContext = ContextDataEntry.Get<FContextObjectTypeStruct>();
+						
+						UK2Node_CallFunction* AddStructFunction = CompilerContext.SpawnIntermediateNode<UK2Node_CallFunction>(this, SourceGraph);
+						CompilerContext.MessageLog.NotifyIntermediateObjectCreation(AddStructFunction, this);
+						AddStructFunction->SetFromFunction(UChooserFunctionLibrary::StaticClass()->FindFunctionByName(GET_MEMBER_NAME_CHECKED(UChooserFunctionLibrary, AddChooserStructInput)));
+						AddStructFunction->AllocateDefaultPins();
+						ContextStructPin->MakeLinkTo(AddStructFunction->FindPin(FName("Context")));
+
+						if (PreviousNodeExecOutput == nullptr)
+						{
+							CompilerContext.MovePinLinksToIntermediate(*ExecInput, *AddStructFunction->GetExecPin());
+						}
+						else
+						{
+							PreviousNodeExecOutput->MakeLinkTo(AddStructFunction->GetExecPin());
+						}
+						PreviousNodeExecOutput = AddStructFunction->GetThenPin();
+							
+						UEdGraphPin* AddStructPin = AddStructFunction->FindPin(FName("Value"));
+						// not sure why this isn't automatically happening:
+						AddStructPin->PinType.PinCategory = UEdGraphSchema_K2::PC_Struct;
+						AddStructPin->PinType.PinSubCategoryObject = StructContext.Struct;
+
+						bool bLinked = false;
+						if (StructContext.Direction != EContextObjectDirection::Write)
+						{
+							if (UEdGraphPin* Pin = FindPin(StructContext.Struct.GetFName(), EGPD_Input))
+							{
+								if (Pin->HasAnyConnections())
+								{
+									CompilerContext.MovePinLinksToIntermediate(*Pin, *AddStructPin);
+									bLinked = true;
+								}
+							}
+						}
+						
+						if (!bLinked)
+						{
+							// create a struct to hold the output (or for input structs that were not connected to anything)
+							UK2Node_MakeStruct* OutputStructNode = CompilerContext.SpawnIntermediateNode<UK2Node_MakeStruct>(this, SourceGraph);
+							CompilerContext.MessageLog.NotifyIntermediateObjectCreation(OutputStructNode, this);
+							OutputStructNode->StructType = static_cast<UScriptStruct*>(StructContext.Struct);
+							OutputStructNode->AllocateDefaultPins();
+							OutputStructNode->FindPin(StructContext.Struct->GetFName(), EGPD_Output)->MakeLinkTo(AddStructPin);
+						}
+
+						if (StructContext.Direction != EContextObjectDirection::Read)
+						{
+							
+							if (UEdGraphPin* Pin = FindPin(StructContext.Struct.GetFName(), EGPD_Output))
+							{
+								if (Pin->HasAnyConnections())
+								{
+									// set up struct output pin
+									UK2Node_CallFunction* GetStructFunction = CompilerContext.SpawnIntermediateNode<UK2Node_CallFunction>(this, SourceGraph);
+									CompilerContext.MessageLog.NotifyIntermediateObjectCreation(GetStructFunction, this);
+									GetStructFunction->SetFromFunction(UChooserFunctionLibrary::StaticClass()->FindFunctionByName(GET_MEMBER_NAME_CHECKED(UChooserFunctionLibrary, GetChooserStructOutput)));
+									GetStructFunction->AllocateDefaultPins();
+									GetStructFunction->FindPin(FName("Index"))->DefaultValue = FString::FromInt(ContextDataIndex);
+									
+									ContextStructPin->MakeLinkTo(GetStructFunction->FindPin(FName("Context")));
+									UEdGraphPin* ValuePin = GetStructFunction->FindPin(FName("Value"));
+									// not sure why this isn't automatically happening:
+									ValuePin->PinType.PinCategory = UEdGraphSchema_K2::PC_Struct;
+									ValuePin->PinType.PinSubCategoryObject = StructContext.Struct;
+					
+									CompilerContext.MovePinLinksToIntermediate(*Pin, *ValuePin);
+								}
+							}
+						}
+					}
+				}
+			}
+
+			if (PreviousNodeExecOutput)
+			{
+				PreviousNodeExecOutput->MakeLinkTo(CallFunction->GetExecPin());
+			}
+			else
+			{
+				CompilerContext.MovePinLinksToIntermediate(*ExecInput, *CallFunction->GetExecPin());
+			}
+			
+			CompilerContext.MovePinLinksToIntermediate(*ExecOutput, *CallFunction->GetThenPin());
+		}
+
+		UK2Node_CallFunction* MakeEvaluateChooserFunction = CompilerContext.SpawnIntermediateNode<UK2Node_CallFunction>(this, SourceGraph);
+		CompilerContext.MessageLog.NotifyIntermediateObjectCreation(MakeEvaluateChooserFunction, this);
+		MakeEvaluateChooserFunction->SetFromFunction(UChooserFunctionLibrary::StaticClass()->FindFunctionByName(GET_MEMBER_NAME_CHECKED(UChooserFunctionLibrary, MakeEvaluateChooser)));
+		MakeEvaluateChooserFunction->AllocateDefaultPins();
+		
+		UEdGraphPin* ChooserTablePin = MakeEvaluateChooserFunction->FindPin(FName("Chooser"));
+		CallFunction->GetSchema()->TrySetDefaultObject(*ChooserTablePin, Chooser);
+
+		MakeEvaluateChooserFunction->GetReturnValuePin()->MakeLinkTo(CallFunction->FindPin(FName("ObjectChooser")));
+
+		UEdGraphPin* OutputPin = CallFunction->GetReturnValuePin();
+
+		if (Chooser && Chooser->OutputObjectType)
+		{
+			UEdGraphPin* OutputClassPin = CallFunction->FindPin(TEXT("ObjectClass"));
+			CallFunction->GetSchema()->TrySetDefaultObject(*OutputClassPin, Chooser->OutputObjectType);
+		}
+
+		CompilerContext.MovePinLinksToIntermediate(*ResultPin, *OutputPin);
+	}
+	
+	BreakAllNodeLinks();
+}
+
+UK2Node::ERedirectType UK2Node_EvaluateChooser2::DoPinsMatchForReconstruction(const UEdGraphPin* NewPin, int32 NewPinIndex, const UEdGraphPin* OldPin, int32 OldPinIndex) const
+{
+	ERedirectType RedirectType = ERedirectType_None;
+
+	// if the pin names do match
+	if (NewPin->PinName.ToString().Equals(OldPin->PinName.ToString(), ESearchCase::CaseSensitive))
+	{
+		// Make sure we're not dealing with a menu node
+		UEdGraph* OuterGraph = GetGraph();
+		if( OuterGraph && OuterGraph->Schema )
+		{
+			const UEdGraphSchema_K2* K2Schema = Cast<const UEdGraphSchema_K2>(GetSchema());
+			if( !K2Schema || K2Schema->IsSelfPin(*NewPin) || K2Schema->ArePinTypesCompatible(OldPin->PinType, NewPin->PinType) )
+			{
+				RedirectType = ERedirectType_Name;
+			}
+			else
+			{
+				RedirectType = ERedirectType_None;
+			}
+		}
+	}
+	else
+	{
+		// try looking for a redirect if it's a K2 node
+		if (UK2Node* Node = Cast<UK2Node>(NewPin->GetOwningNode()))
+		{	
+			// if you don't have matching pin, now check if there is any redirect param set
+			TArray<FString> OldPinNames;
+			GetRedirectPinNames(*OldPin, OldPinNames);
+
+			FName NewPinName;
+			RedirectType = ShouldRedirectParam(OldPinNames, /*out*/ NewPinName, Node);
+
+			// make sure they match
+			if ((RedirectType != ERedirectType_None) && (!NewPin->PinName.ToString().Equals(NewPinName.ToString(), ESearchCase::CaseSensitive)))
+			{
+				RedirectType = ERedirectType_None;
+			}
+		}
+	}
+
+	return RedirectType;
+}
+
+bool UK2Node_EvaluateChooser2::IsConnectionDisallowed(const UEdGraphPin* MyPin, const UEdGraphPin* OtherPin, FString& OutReason) const
+{
+	return Super::IsConnectionDisallowed(MyPin, OtherPin, OutReason);
+}
+
+void UK2Node_EvaluateChooser2::GetMenuActions(FBlueprintActionDatabaseRegistrar& ActionRegistrar) const
+{
+	// actions get registered under specific object-keys; the idea is that 
+	// actions might have to be updated (or deleted) if their object-key is  
+	// mutated (or removed)... here we use the node's class (so if the node 
+	// type disappears, then the action should go with it)
+	UClass* ActionKey = GetClass();
+	// to keep from needlessly instantiating a UBlueprintNodeSpawner, first   
+	// check to make sure that the registrar is looking for actions of this type
+	// (could be regenerating actions for a specific asset, and therefore the 
+	// registrar would only accept actions corresponding to that asset)
+	if (ActionRegistrar.IsOpenForRegistration(ActionKey))
+	{
+		UBlueprintNodeSpawner* NodeSpawner = UBlueprintNodeSpawner::Create(GetClass());
+		check(NodeSpawner != nullptr);
+
+		ActionRegistrar.AddBlueprintAction(ActionKey, NodeSpawner);
+	}
+}
+
+FText UK2Node_EvaluateChooser2::GetMenuCategory() const
+{
+	return FEditorCategoryUtils::GetCommonCategory(FCommonEditorCategory::Animation);
+}
+
+
 #undef LOCTEXT_NAMESPACE
diff --git a/Engine/Plugins/Experimental/Chooser/Source/ChooserUncooked/Private/EvaluateChooserNode.h b/Engine/Plugins/Experimental/Chooser/Source/ChooserUncooked/Private/EvaluateChooserNode.h
--- a/Engine/Plugins/Experimental/Chooser/Source/ChooserUncooked/Private/EvaluateChooserNode.h
+++ b/Engine/Plugins/Experimental/Chooser/Source/ChooserUncooked/Private/EvaluateChooserNode.h
@@ -21,7 +21,10 @@ enum class EEvaluateChooserMode
 	AllResults
 };
 
-UCLASS(MinimalAPI)
+/////////////////////////////////////////////////////
+// UK2Node_EvaluateChooser
+// old implementation of this node for backwards compatibility - not currently accessible to create new instances in content
+UCLASS(MinimalAPI, Hidden)
 class UK2Node_EvaluateChooser : public UK2Node
 {
 	GENERATED_UCLASS_BODY()
@@ -65,10 +68,10 @@ private:
 	void ResultTypeChanged(const UClass*);
 	void UnregisterChooserCallback();
 	UChooserTable* CurrentCallbackChooser = nullptr;
-	
+
 	UPROPERTY(EditAnywhere, Category = "Chooser")
 	TObjectPtr<UChooserTable> Chooser;
-	
+
 	UPROPERTY(EditAnywhere, Category = "Chooser")
 	EEvaluateChooserMode Mode;
 
@@ -76,3 +79,60 @@ private:
 	FText NodeTooltip;
 };
 
+//----------------------------------------------------------------------------------------------
+// UK2Node_EvaluateChooser2
+// New Implementation of EvaluateChooser with support for passing in/out multiple objects and structs
+UCLASS(MinimalAPI)
+class UK2Node_EvaluateChooser2 : public UK2Node
+{
+	GENERATED_UCLASS_BODY()
+	
+	//~ Begin UObject Interface
+	virtual void BeginDestroy() override;
+	virtual void PostEditUndo() override;
+	virtual void PostEditChangeProperty(struct FPropertyChangedEvent& PropertyChangedEvent) override;
+	virtual void PostLoad() override;
+	//~ End UObject Interface
+
+	//~ Begin UEdGraphNode Interface.
+	virtual void AllocateDefaultPins() override;
+	virtual FText GetNodeTitle(ENodeTitleType::Type TitleType) const override;
+	virtual bool ShouldShowNodeProperties() const override { return true; }
+	virtual void PinConnectionListChanged(UEdGraphPin* Pin) override;
+	virtual void PinDefaultValueChanged(UEdGraphPin* Pin) override;
+	virtual void PinTypeChanged(UEdGraphPin* Pin) override;
+	virtual FText GetTooltipText() const override;
+	virtual FText GetPinDisplayName(const UEdGraphPin* Pin) const override;
+	//~ End UEdGraphNode Interface.
+
+	//~ Begin UK2Node Interface.
+	virtual bool IsNodePure() const override { return true; }
+	virtual bool NodeCausesStructuralBlueprintChange() const override { return true; }
+	virtual void PostReconstructNode() override;
+	virtual void ExpandNode(class FKismetCompilerContext& CompilerContext, UEdGraph* SourceGraph) override;
+	virtual ERedirectType DoPinsMatchForReconstruction(const UEdGraphPin* NewPin, int32 NewPinIndex, const UEdGraphPin* OldPin, int32 OldPinIndex) const override;
+	virtual bool IsConnectionDisallowed(const UEdGraphPin* MyPin, const UEdGraphPin* OtherPin, FString& OutReason) const override;
+	virtual void GetMenuActions(FBlueprintActionDatabaseRegistrar& ActionRegistrar) const override;
+	virtual FText GetMenuCategory() const override;
+	virtual int32 GetNodeRefreshPriority() const override { return EBaseNodeRefreshPriority::Low_UsesDependentWildcard; }
+	virtual void PreloadRequiredAssets() override;
+
+	virtual void DestroyNode() override;
+    
+	//~ End UK2Node Interface.
+
+private:
+	void ChooserChanged();
+	void ResultTypeChanged(const UClass*);
+	void UnregisterChooserCallback();
+	UChooserTable* CurrentCallbackChooser = nullptr;
+
+	UPROPERTY(EditAnywhere, Category = "Chooser")
+	TObjectPtr<UChooserTable> Chooser;
+
+	UPROPERTY(EditAnywhere, Category = "Chooser")
+	EEvaluateChooserMode Mode;
+
+	/** Tooltip text for this node. */
+	FText NodeTooltip;
+};
\ No newline at end of file
diff --git a/Engine/Plugins/Experimental/Chooser/Source/ProxyTable/Private/ProxyTableFunctionLibrary.cpp b/Engine/Plugins/Experimental/Chooser/Source/ProxyTable/Private/ProxyTableFunctionLibrary.cpp
--- a/Engine/Plugins/Experimental/Chooser/Source/ProxyTable/Private/ProxyTableFunctionLibrary.cpp
+++ b/Engine/Plugins/Experimental/Chooser/Source/ProxyTable/Private/ProxyTableFunctionLibrary.cpp
@@ -11,7 +11,9 @@ UObject* UProxyTableFunctionLibrary::EvaluateProxyAsset(const UObject* ContextOb
 	if (Proxy)
 	{
 		FChooserEvaluationContext Context;
-		Context.ContextData.Add({ContextObject->GetClass(), const_cast<UObject*>(ContextObject)});
+		Context.Params.AddDefaulted();
+		Context.Params.Last().InitializeAs(FChooserEvaluationInputObject::StaticStruct());
+		Context.Params.Last().GetMutable<FChooserEvaluationInputObject>().Object = const_cast<UObject*>(ContextObject);
 		
 		Result = Proxy->FindProxyObject(Context);
 		if (ObjectClass && Result && !Result->IsA(ObjectClass))
@@ -31,7 +33,9 @@ UObject* UProxyTableFunctionLibrary::EvaluateProxyTable(const UObject* ContextOb
 		FGuid Guid;
 		Guid.A = GetTypeHash(Key);
 		FChooserEvaluationContext Context;
-		Context.ContextData.Add({ContextObject->GetClass(), const_cast<UObject*>(ContextObject)});
+		Context.Params.AddDefaulted();
+		Context.Params.Last().InitializeAs(FChooserEvaluationInputObject::StaticStruct());
+		Context.Params.Last().GetMutable<FChooserEvaluationInputObject>().Object = const_cast<UObject*>(ContextObject);
 		if (UObject* Value = ProxyTable->FindProxyObject(Guid, Context))
 		{
 			return Value;
