fix localization errors

diff --git a/Engine/Plugins/Experimental/Chooser/Source/ChooserUncooked/Private/EvaluateChooserNode.cpp b/Engine/Plugins/Experimental/Chooser/Source/ChooserUncooked/Private/EvaluateChooserNode.cpp
--- a/Engine/Plugins/Experimental/Chooser/Source/ChooserUncooked/Private/EvaluateChooserNode.cpp
+++ b/Engine/Plugins/Experimental/Chooser/Source/ChooserUncooked/Private/EvaluateChooserNode.cpp
@@ -314,7 +314,7 @@ FText UK2Node_EvaluateChooser::GetMenuCategory() const
 UK2Node_EvaluateChooser2::UK2Node_EvaluateChooser2(const FObjectInitializer& ObjectInitializer)
 	: Super(ObjectInitializer)
 {
-	NodeTooltip = LOCTEXT("NodeTooltip", "Evaluates an Chooser Table, and returns the resulting Object or Objects.");
+	NodeTooltip = LOCTEXT("EvaluateChooser2Tooltip", "Evaluates an Chooser Table, and returns the resulting Object or Objects.");
 }
 
 void UK2Node_EvaluateChooser2::UnregisterChooserCallback()
@@ -502,11 +502,11 @@ FText UK2Node_EvaluateChooser2::GetNodeTitle(ENodeTitleType::Type TitleType) con
 {
 	if (Chooser)
 	{
-		return FText::Format(LOCTEXT("EvaluateChooser_TitleWithChooser", "Evaluate Chooser: {0}"), {FText::FromString(Chooser->GetName())});
+		return FText::Format(LOCTEXT("EvaluateChooser2_TitleWithChooser", "Evaluate Chooser: {0}"), {FText::FromString(Chooser->GetName())});
 	}
 	else
 	{
-		return LOCTEXT("EvaluateChooser_Title", "Evaluate Chooser");
+		return LOCTEXT("EvaluateChooser2_Title", "Evaluate Chooser");
 	}
 }
 
