fix crash deleting selected column from a chooser

diff --git a/Engine/Plugins/Experimental/Chooser/Source/ChooserEditor/Private/ChooserTableEditor.cpp b/Engine/Plugins/Experimental/Chooser/Source/ChooserEditor/Private/ChooserTableEditor.cpp
--- a/Engine/Plugins/Experimental/Chooser/Source/ChooserEditor/Private/ChooserTableEditor.cpp
+++ b/Engine/Plugins/Experimental/Chooser/Source/ChooserEditor/Private/ChooserTableEditor.cpp
@@ -796,7 +796,7 @@ void FChooserTableEditor::UpdateTableColumns()
 
 				FMenuBuilder MenuBuilder(true,nullptr);
 
-				MenuBuilder.AddMenuEntry(LOCTEXT("Column Properties","Properties"),LOCTEXT("Delete Column ToolTip", "Remove this column and all its data from the table"),FSlateIcon(),
+				MenuBuilder.AddMenuEntry(LOCTEXT("Column Properties","Properties"),LOCTEXT("Select Column ToolTip", "Select this Column, and show its properties in the Details panel"),FSlateIcon(),
 					FUIAction(
 						FExecuteAction::CreateLambda([this,Chooser,ColumnIndex, ColumnId, &Column]()
 						{
@@ -840,6 +840,7 @@ void FChooserTableEditor::UpdateTableColumns()
 							const FScopedTransaction Transaction(LOCTEXT("Delete Column Transaction", "Delete Column"));
 							Chooser->Modify(true);
 							Chooser->ColumnsStructs.RemoveAt(ColumnIndex);
+							ClearSelectedColumn();
 							UpdateTableColumns();
 						})
 						));
@@ -1057,11 +1058,8 @@ TSharedRef<SDockTab> FChooserTableEditor::SpawnTableTab( const FSpawnTabArgs& Ar
 				)
 				.OnSelectionChanged_Lambda([this](TSharedPtr<FChooserTableRow> SelectedItem,  ESelectInfo::Type SelectInfo)
 				{
-					if (SelectedColumn)
-					{
-						// deselect any selected column
-						SelectedColumn->Column = -1;
-					}
+					// deselect any selected column
+					ClearSelectedColumn();
 					
 					if (SelectedItem)
 					{
@@ -1199,6 +1197,19 @@ void FChooserTableEditor::SelectColumn(int Index)
 	}
 }
 	
+void FChooserTableEditor::ClearSelectedColumn()
+{
+	UChooserTable* Chooser = GetChooser();
+	if (SelectedColumn != nullptr)
+	{
+		SelectedColumn->Column = -1;
+		if (DetailsView->GetSelectedObjects().Contains(SelectedColumn))
+		{
+			SelectRootProperties();
+		}
+	}
+}
+	
 void FChooserTableEditor::DeleteColumn(int Index)
 {
 	UChooserTable* Chooser = GetChooser();
@@ -1400,7 +1411,7 @@ void FChooserColumnDetails::CustomizeDetails(IDetailLayoutBuilder& DetailBuilder
 	UChooserColumnDetails* Column = Cast<UChooserColumnDetails>(Objects[0]);
 	UChooserTable* Chooser = Column->Chooser;
 	
-	if (Chooser->ColumnsStructs.Num() > Column->Column)
+	if (Chooser->ColumnsStructs.IsValidIndex(Column->Column))
 	{
 		IDetailCategoryBuilder& PropertiesCategory = DetailBuilder.EditCategory("Column Properties");
 
diff --git a/Engine/Plugins/Experimental/Chooser/Source/ChooserEditor/Private/ChooserTableEditor.h b/Engine/Plugins/Experimental/Chooser/Source/ChooserEditor/Private/ChooserTableEditor.h
--- a/Engine/Plugins/Experimental/Chooser/Source/ChooserEditor/Private/ChooserTableEditor.h
+++ b/Engine/Plugins/Experimental/Chooser/Source/ChooserEditor/Private/ChooserTableEditor.h
@@ -36,7 +36,7 @@ class UChooserColumnDetails : public UObject
 public:
 	UPROPERTY(EditAnywhere, Instanced, Category="Hidden")
 	TObjectPtr<UChooserTable> Chooser;
-	int Column;
+	int Column = -1;
 };
 
 
@@ -100,6 +100,7 @@ namespace UE::ChooserEditor
 
 		void UpdateTableRows();
 		void SelectColumn(int Index);
+		void ClearSelectedColumn();
 		void DeleteColumn(int Index);
 		void AddColumn(const UScriptStruct* ColumnType);
 		void MoveRow(int SourceRowIndex, int TargetIndex);
