float output column for choosers

diff --git a/Engine/Plugins/Experimental/Chooser/Source/Chooser/Internal/FloatRangeColumn.h b/Engine/Plugins/Experimental/Chooser/Source/Chooser/Internal/FloatRangeColumn.h
--- a/Engine/Plugins/Experimental/Chooser/Source/Chooser/Internal/FloatRangeColumn.h
+++ b/Engine/Plugins/Experimental/Chooser/Source/Chooser/Internal/FloatRangeColumn.h
@@ -13,7 +13,8 @@ struct CHOOSER_API FFloatContextProperty :  public FChooserParameterFloatBase
 {
 	GENERATED_BODY()
 	
-	virtual bool GetValue(FChooserEvaluationContext& Context, float& OutResult) const override;
+	virtual bool GetValue(FChooserEvaluationContext& Context, double& OutResult) const override;
+	virtual bool SetValue(FChooserEvaluationContext& Context, double Value) const override;
 
 	UPROPERTY()
 	TArray<FName> PropertyBindingChain_DEPRECATED;
@@ -58,7 +59,7 @@ USTRUCT()
 struct FChooserFloatRangeRowData
 {
 	GENERATED_BODY()
-	
+
 	UPROPERTY(EditAnywhere, Category=Runtime)
 	float Min=0;
 	
diff --git a/Engine/Plugins/Experimental/Chooser/Source/Chooser/Internal/OutputFloatColumn.h b/Engine/Plugins/Experimental/Chooser/Source/Chooser/Internal/OutputFloatColumn.h
new file mode 100644
--- /dev/null
+++ b/Engine/Plugins/Experimental/Chooser/Source/Chooser/Internal/OutputFloatColumn.h
@@ -0,0 +1,37 @@
+// Copyright Epic Games, Inc. All Rights Reserved.
+#pragma once
+
+#include "CoreMinimal.h"
+#include "IChooserColumn.h"
+#include "IChooserParameterFloat.h"
+#include "ChooserPropertyAccess.h"
+#include "InstancedStruct.h"
+#include "OutputFloatColumn.generated.h"
+
+USTRUCT()
+struct CHOOSER_API FOutputFloatColumn : public FChooserColumnBase
+{
+	GENERATED_BODY()
+	public:
+	FOutputFloatColumn();
+	virtual bool HasFilters() const override { return false; }
+	virtual bool HasOutputs() const override { return true; }
+	virtual void SetOutputs(FChooserEvaluationContext& Context, int RowIndex) const override;
+	
+	UPROPERTY(EditAnywhere, Meta = (ExcludeBaseStruct, BaseStruct = "/Script/Chooser.ChooserParameterFloatBase"), Category = "Hidden")
+	FInstancedStruct InputValue;
+
+#if WITH_EDITOR
+	mutable double TestValue=false;
+#endif
+	
+#if WITH_EDITORONLY_DATA
+	UPROPERTY(EditAnywhere, Category=Runtime);
+	double DefaultRowValue = 0.0;
+#endif
+	
+	UPROPERTY(EditAnywhere, Category=Runtime);
+	TArray<double> RowValues; 
+
+	CHOOSER_COLUMN_BOILERPLATE(FChooserParameterFloatBase);
+};
\ No newline at end of file
diff --git a/Engine/Plugins/Experimental/Chooser/Source/Chooser/Private/FloatRangeColumn.cpp b/Engine/Plugins/Experimental/Chooser/Source/Chooser/Private/FloatRangeColumn.cpp
--- a/Engine/Plugins/Experimental/Chooser/Source/Chooser/Private/FloatRangeColumn.cpp
+++ b/Engine/Plugins/Experimental/Chooser/Source/Chooser/Private/FloatRangeColumn.cpp
@@ -3,7 +3,31 @@
 
 #include "ChooserPropertyAccess.h"
 
-bool FFloatContextProperty::GetValue(FChooserEvaluationContext& Context, float& OutResult) const
+bool FFloatContextProperty::SetValue(FChooserEvaluationContext& Context, double InValue) const
+{
+	const UStruct* StructType = nullptr;
+	const void* Container = nullptr;
+	
+	if (UE::Chooser::ResolvePropertyChain(Context, Binding, Container, StructType))
+	{
+		if (FDoubleProperty* DoubleProperty = FindFProperty<FDoubleProperty>(StructType, Binding.PropertyBindingChain.Last()))
+		{
+			// const cast is here just because ResolvePropertyChain expects a const void*&
+			*DoubleProperty->ContainerPtrToValuePtr<double>(const_cast<void*>(Container)) = InValue;
+			return true;
+		}
+		if (FFloatProperty* FloatProperty = FindFProperty<FFloatProperty>(StructType, Binding.PropertyBindingChain.Last()))
+        {
+			// const cast is here just because ResolvePropertyChain expects a const void*&
+        	*FloatProperty->ContainerPtrToValuePtr<float>(const_cast<void*>(Container)) = InValue;
+        	return true;
+        }
+	}
+
+	return false;
+}
+
+bool FFloatContextProperty::GetValue(FChooserEvaluationContext& Context, double& OutResult) const
 {
 	const UStruct* StructType = nullptr;
 	const void* Container = nullptr;
@@ -26,7 +50,7 @@ bool FFloatContextProperty::GetValue(FChooserEvaluationContext& Context, float&
 	    {
 			if (UFunction* Function = ClassType->FindFunctionByName(Binding.PropertyBindingChain.Last()))
 			{
-				bool bReturnsDouble = CastField<FDoubleProperty>(Function->GetReturnProperty()) != nullptr;
+				const bool bReturnsDouble = CastField<FDoubleProperty>(Function->GetReturnProperty()) != nullptr;
 					
 				UObject* Object = reinterpret_cast<UObject*>(const_cast<void*>(Container));
 				if (Function->IsNative())
@@ -34,26 +58,26 @@ bool FFloatContextProperty::GetValue(FChooserEvaluationContext& Context, float&
 					FFrame Stack(Object, Function, nullptr, nullptr, Function->ChildProperties);
 					if (bReturnsDouble)
 					{
-						double result;
-						Function->Invoke(Object, Stack, &result);
-						OutResult = result;
+						Function->Invoke(Object, Stack, &OutResult);
 					}
 					else
 					{
-						Function->Invoke(Object, Stack, &OutResult);
+						float result = 0;
+						Function->Invoke(Object, Stack, &result);
+						OutResult = result;
 					}
 				}
 				else
 				{
 					if (bReturnsDouble)
 					{
-						double result = 0;
-						Object->ProcessEvent(Function, &result);
-						OutResult = result;
+						Object->ProcessEvent(Function, &OutResult);
 					}
 					else
 					{
-						Object->ProcessEvent(Function, &OutResult);
+						float result = 0;
+						Object->ProcessEvent(Function, &result);
+						OutResult = result;
 					}
 				}
 				return true;
@@ -73,7 +97,7 @@ void FFloatRangeColumn::Filter(FChooserEvaluationContext& Context, const TArray<
 {
 	if (InputValue.IsValid())
 	{
-		float Result = 0.0f;
+		double Result = 0.0f;
 		InputValue.Get<FChooserParameterFloatBase>().GetValue(Context, Result);
 
 #if WITH_EDITOR
diff --git a/Engine/Plugins/Experimental/Chooser/Source/Chooser/Private/OutputFloatColumn .cpp b/Engine/Plugins/Experimental/Chooser/Source/Chooser/Private/OutputFloatColumn .cpp
new file mode 100644
--- /dev/null
+++ b/Engine/Plugins/Experimental/Chooser/Source/Chooser/Private/OutputFloatColumn .cpp	
@@ -0,0 +1,24 @@
+// Copyright Epic Games, Inc. All Rights Reserved.
+#include "OutputFloatColumn.h"
+#include "ChooserPropertyAccess.h"
+#include "FloatRangeColumn.h"
+
+FOutputFloatColumn::FOutputFloatColumn()
+{
+	InputValue.InitializeAs(FFloatContextProperty::StaticStruct());
+}
+
+void FOutputFloatColumn::SetOutputs(FChooserEvaluationContext& Context, int RowIndex) const
+{
+	if (InputValue.IsValid())
+	{
+		InputValue.Get<FChooserParameterFloatBase>().SetValue(Context, RowValues[RowIndex]);
+	}
+	
+#if WITH_EDITOR
+	if (Context.DebuggingInfo.bCurrentDebugTarget)
+	{
+		TestValue = RowValues[RowIndex];
+	}
+#endif
+}
diff --git a/Engine/Plugins/Experimental/Chooser/Source/Chooser/Public/IChooserParameterFloat.h b/Engine/Plugins/Experimental/Chooser/Source/Chooser/Public/IChooserParameterFloat.h
--- a/Engine/Plugins/Experimental/Chooser/Source/Chooser/Public/IChooserParameterFloat.h
+++ b/Engine/Plugins/Experimental/Chooser/Source/Chooser/Public/IChooserParameterFloat.h
@@ -28,5 +28,6 @@ struct FChooserParameterFloatBase : public FChooserParameterBase
 	GENERATED_BODY()
     
 public:
-	virtual bool GetValue(FChooserEvaluationContext& Context, float& OutResult) const { return false; }
+	virtual bool GetValue(FChooserEvaluationContext& Context, double& OutResult) const { return false; }
+    virtual bool SetValue(FChooserEvaluationContext& Context, double InValue) const { return false; }
 };
\ No newline at end of file
diff --git a/Engine/Plugins/Experimental/Chooser/Source/ChooserEditor/Private/OutputFloatColumnEditor.cpp b/Engine/Plugins/Experimental/Chooser/Source/ChooserEditor/Private/OutputFloatColumnEditor.cpp
new file mode 100644
--- /dev/null
+++ b/Engine/Plugins/Experimental/Chooser/Source/ChooserEditor/Private/OutputFloatColumnEditor.cpp
@@ -0,0 +1,89 @@
+// Copyright Epic Games, Inc. All Rights Reserved.
+
+#include "OutputFloatColumnEditor.h"
+#include "OutputFloatColumn.h"
+#include "ObjectChooserWidgetFactories.h"
+#include "ChooserTableEditor.h"
+#include "Widgets/Input/SNumericEntryBox.h"
+#include "Widgets/Images/SImage.h"
+#include "GraphEditorSettings.h"
+#include "SPropertyAccessChainWidget.h"
+
+#define LOCTEXT_NAMESPACE "OutputBoolColumnEditor"
+
+namespace UE::ChooserEditor
+{
+	
+TSharedRef<SWidget> CreateOutputFloatColumnWidget(UChooserTable* Chooser, FChooserColumnBase* Column, int Row)
+{
+	FOutputFloatColumn* OutputFloatColumn = static_cast<FOutputFloatColumn*>(Column);
+
+	if (Row < 0)
+	{
+		// create column header widget
+		TSharedPtr<SWidget> InputValueWidget = nullptr;
+		if (FChooserParameterBase* InputValue = Column->GetInputValue())
+		{
+			InputValueWidget = FObjectChooserWidgetFactories::CreateWidget(false, Chooser, InputValue, Column->GetInputType(), Chooser->OutputObjectType);
+		}
+		
+		const FSlateBrush* ColumnIcon = FCoreStyle::Get().GetBrush("Icons.ArrowRight");
+		
+		TSharedRef<SWidget> ColumnHeaderWidget = SNew(SHorizontalBox)
+			+ SHorizontalBox::Slot().AutoWidth()
+			[
+				SNew(SBorder)
+				.BorderBackgroundColor(FLinearColor(0,0,0,0))
+				.Content()
+				[
+					SNew(SImage).Image(ColumnIcon)
+				]
+			]
+			+ SHorizontalBox::Slot()
+			[
+				InputValueWidget ? InputValueWidget.ToSharedRef() : SNullWidget::NullWidget
+			];
+	
+		if (Chooser->bEnableDebugTesting)
+		{
+			ColumnHeaderWidget = SNew(SVerticalBox)
+			+ SVerticalBox::Slot()
+			[
+				ColumnHeaderWidget
+			]
+			+ SVerticalBox::Slot()
+			[
+				SNew(SNumericEntryBox<float>).IsEnabled(false)
+				.Value_Lambda([OutputFloatColumn]() { return OutputFloatColumn->TestValue; })
+			];
+		}
+
+		return ColumnHeaderWidget;
+	}
+
+	// create cell widget
+	return SNew(SNumericEntryBox<double>)
+    		.Value_Lambda([OutputFloatColumn, Row]()
+    		{
+    			return (Row < OutputFloatColumn->RowValues.Num()) ? OutputFloatColumn->RowValues[Row] : 0;
+    		})
+    		.OnValueCommitted_Lambda([Chooser, OutputFloatColumn, Row](double NewValue, ETextCommit::Type CommitType)
+    		{
+    			if (Row < OutputFloatColumn->RowValues.Num())
+    			{
+    				const FScopedTransaction Transaction(LOCTEXT("Edit Min Value", "Edit Min Value"));
+    				Chooser->Modify(true);
+    				OutputFloatColumn->RowValues[Row] = NewValue;
+    			}
+    		});
+}
+
+	
+void RegisterOutputFloatWidgets()
+{
+	FObjectChooserWidgetFactories::RegisterColumnWidgetCreator(FOutputFloatColumn::StaticStruct(), CreateOutputFloatColumnWidget);
+}
+	
+}
+
+#undef LOCTEXT_NAMESPACE
diff --git a/Engine/Plugins/Experimental/Chooser/Source/ChooserEditor/Private/OutputFloatColumnEditor.h b/Engine/Plugins/Experimental/Chooser/Source/ChooserEditor/Private/OutputFloatColumnEditor.h
new file mode 100644
--- /dev/null
+++ b/Engine/Plugins/Experimental/Chooser/Source/ChooserEditor/Private/OutputFloatColumnEditor.h
@@ -0,0 +1,10 @@
+// Copyright Epic Games, Inc. All Rights Reserved.
+
+#pragma once
+
+#include "CoreMinimal.h"
+
+namespace UE::ChooserEditor
+{
+	void RegisterOutputFloatWidgets();
+}
\ No newline at end of file
