Missing change for chooser float output column

diff --git a/Engine/Plugins/Experimental/Chooser/Source/ChooserEditor/Private/ChooserEditorModule.cpp b/Engine/Plugins/Experimental/Chooser/Source/ChooserEditor/Private/ChooserEditorModule.cpp
--- a/Engine/Plugins/Experimental/Chooser/Source/ChooserEditor/Private/ChooserEditorModule.cpp
+++ b/Engine/Plugins/Experimental/Chooser/Source/ChooserEditor/Private/ChooserEditorModule.cpp
@@ -6,6 +6,7 @@
 #include "BoolColumnEditor.h"
 #include "EnumColumnEditor.h"
 #include "FloatRangeColumnEditor.h"
+#include "OutputFloatColumnEditor.h"
 #include "GameplayTagColumnEditor.h"
 #include "ObjectColumnEditor.h"
 #include "ChooserTableEditorCommands.h"
@@ -23,6 +24,7 @@ void FModule::StartupModule()
 	FChooserTableEditor::RegisterWidgets();
 	RegisterGameplayTagWidgets();
 	RegisterFloatRangeWidgets();
+	RegisterOutputFloatWidgets();
 	RegisterBoolWidgets();
 	RegisterEnumWidgets();
 	RegisterObjectWidgets();
