Enum output column for Chooser Tables

diff --git a/Engine/Plugins/Experimental/Chooser/Source/Chooser/Internal/EnumColumn.h b/Engine/Plugins/Experimental/Chooser/Source/Chooser/Internal/EnumColumn.h
--- a/Engine/Plugins/Experimental/Chooser/Source/Chooser/Internal/EnumColumn.h
+++ b/Engine/Plugins/Experimental/Chooser/Source/Chooser/Internal/EnumColumn.h
@@ -22,6 +22,7 @@ struct CHOOSER_API FEnumContextProperty : public FChooserParameterEnumBase
 	FChooserEnumPropertyBinding Binding;
 
 	virtual bool GetValue(FChooserEvaluationContext& Context, uint8& OutResult) const override;
+	virtual bool SetValue(FChooserEvaluationContext& Context, uint8 InValue) const override;
 
 	virtual void PostLoad() override
 	{
diff --git a/Engine/Plugins/Experimental/Chooser/Source/Chooser/Internal/OutputEnumColumn.h b/Engine/Plugins/Experimental/Chooser/Source/Chooser/Internal/OutputEnumColumn.h
new file mode 100644
--- /dev/null
+++ b/Engine/Plugins/Experimental/Chooser/Source/Chooser/Internal/OutputEnumColumn.h
@@ -0,0 +1,49 @@
+// Copyright Epic Games, Inc. All Rights Reserved.
+#pragma once
+
+#include "CoreMinimal.h"
+#include "IChooserColumn.h"
+#include "IChooserParameterEnum.h"
+#include "InstancedStruct.h"
+#include "ChooserPropertyAccess.h"
+#include "OutputEnumColumn.generated.h"
+
+USTRUCT()
+struct FChooserOutputEnumRowData
+{
+	GENERATED_BODY()
+
+	UPROPERTY(EditAnywhere, Category = Runtime)
+	uint8 Value = 0;
+};
+
+
+USTRUCT()
+struct CHOOSER_API FOutputEnumColumn : public FChooserColumnBase
+{
+	GENERATED_BODY()
+public:
+	FOutputEnumColumn();
+	virtual bool HasFilters() const override { return false; }
+	virtual bool HasOutputs() const override { return true; }
+	virtual void SetOutputs(FChooserEvaluationContext& Context, int RowIndex) const override;
+
+	UPROPERTY(EditAnywhere, Meta = (ExcludeBaseStruct, BaseStruct = "/Script/Chooser.ChooserParameterEnumBase"), Category = "Data")
+	FInstancedStruct InputValue;
+
+#if WITH_EDITORONLY_DATA
+	UPROPERTY(EditAnywhere, Category = "Data")
+	FChooserOutputEnumRowData DefaultRowValue;
+#endif
+	
+	UPROPERTY(EditAnywhere, Category = "Data")
+	// array of results (cells for this column for each row in the table)
+	// should match the length of the Results array
+	TArray<FChooserOutputEnumRowData> RowValues;
+	
+#if WITH_EDITOR
+	mutable uint8 TestValue;
+#endif
+	
+	CHOOSER_COLUMN_BOILERPLATE(FChooserParameterEnumBase);
+};
\ No newline at end of file
diff --git a/Engine/Plugins/Experimental/Chooser/Source/Chooser/Private/EnumColumn.cpp b/Engine/Plugins/Experimental/Chooser/Source/Chooser/Private/EnumColumn.cpp
--- a/Engine/Plugins/Experimental/Chooser/Source/Chooser/Private/EnumColumn.cpp
+++ b/Engine/Plugins/Experimental/Chooser/Source/Chooser/Private/EnumColumn.cpp
@@ -32,6 +32,32 @@ bool FEnumContextProperty::GetValue(FChooserEvaluationContext& Context, uint8& O
 	return false;
 }
 
+bool FEnumContextProperty::SetValue(FChooserEvaluationContext& Context, uint8 InValue) const
+{
+	const UStruct* StructType = nullptr;
+	const void* Container = nullptr;
+
+	if (UE::Chooser::ResolvePropertyChain(Context, Binding, Container, StructType))
+	{
+		if (const FEnumProperty* EnumProperty = FindFProperty<FEnumProperty>(StructType, Binding.PropertyBindingChain.Last()))
+		{
+			*EnumProperty->ContainerPtrToValuePtr<uint8>(const_cast<void*>(Container)) = InValue;
+			return true;
+		}
+
+		if (const FByteProperty* ByteProperty = FindFProperty<FByteProperty>(StructType, Binding.PropertyBindingChain.Last()))
+		{
+			if (ByteProperty->IsEnum())
+			{
+				*ByteProperty->ContainerPtrToValuePtr<uint8>(const_cast<void*>(Container)) = InValue;
+				return true;
+			}
+		}
+	}
+
+	return false;
+}
+
 #if WITH_EDITOR
 
 void FEnumContextProperty::SetBinding(const TArray<FBindingChainElement>& InBindingChain)
diff --git a/Engine/Plugins/Experimental/Chooser/Source/Chooser/Private/OutputEnumColumn.cpp b/Engine/Plugins/Experimental/Chooser/Source/Chooser/Private/OutputEnumColumn.cpp
new file mode 100644
--- /dev/null
+++ b/Engine/Plugins/Experimental/Chooser/Source/Chooser/Private/OutputEnumColumn.cpp
@@ -0,0 +1,26 @@
+// Copyright Epic Games, Inc. All Rights Reserved.
+#include "OutputEnumColumn.h"
+#include "EnumColumn.h"
+#include "ChooserPropertyAccess.h"
+
+FOutputEnumColumn::FOutputEnumColumn()
+{
+#if WITH_EDITOR
+	InputValue.InitializeAs(FEnumContextProperty::StaticStruct());
+#endif
+}
+
+void FOutputEnumColumn::SetOutputs(FChooserEvaluationContext& Context, int RowIndex) const
+{
+	if (InputValue.IsValid())
+	{
+		InputValue.Get<FChooserParameterEnumBase>().SetValue(Context, RowValues[RowIndex].Value);
+	}
+	
+#if WITH_EDITOR
+	if (Context.DebuggingInfo.bCurrentDebugTarget)
+	{
+		TestValue = RowValues[RowIndex].Value;
+	}
+#endif
+}
\ No newline at end of file
diff --git a/Engine/Plugins/Experimental/Chooser/Source/Chooser/Public/IChooserParameterEnum.h b/Engine/Plugins/Experimental/Chooser/Source/Chooser/Public/IChooserParameterEnum.h
--- a/Engine/Plugins/Experimental/Chooser/Source/Chooser/Public/IChooserParameterEnum.h
+++ b/Engine/Plugins/Experimental/Chooser/Source/Chooser/Public/IChooserParameterEnum.h
@@ -27,6 +27,7 @@ struct FChooserParameterEnumBase : public FChooserParameterBase
 	GENERATED_BODY()
 	
 		virtual bool GetValue(FChooserEvaluationContext& Context, uint8& OutResult) const { return false; }
+    	virtual bool SetValue(FChooserEvaluationContext& Context, uint8 InValue) const { return false; }
 
 	#if WITH_EDITOR
     	virtual const UEnum* GetEnum() const { return nullptr; }
diff --git a/Engine/Plugins/Experimental/Chooser/Source/ChooserEditor/Private/BoolColumnEditor.cpp b/Engine/Plugins/Experimental/Chooser/Source/ChooserEditor/Private/BoolColumnEditor.cpp
--- a/Engine/Plugins/Experimental/Chooser/Source/ChooserEditor/Private/BoolColumnEditor.cpp
+++ b/Engine/Plugins/Experimental/Chooser/Source/ChooserEditor/Private/BoolColumnEditor.cpp
@@ -59,7 +59,7 @@ TSharedRef<SWidget> CreateBoolColumnWidget(UChooserTable* Chooser, FChooserColum
 				 + SHorizontalBox::Slot().FillWidth(1)
 				 + SHorizontalBox::Slot().AutoWidth().VAlign(VAlign_Center)
 				[
-					SNew(SCheckBox).IsEnabled_Lambda([Chooser](){ return Chooser->HasDebugTarget(); })
+					SNew(SCheckBox).IsEnabled_Lambda([Chooser](){ return !Chooser->HasDebugTarget(); })
 					.IsChecked_Lambda([BoolColumn]() { return BoolColumn->TestValue ? ECheckBoxState::Checked : ECheckBoxState::Unchecked; })
 					.OnCheckStateChanged_Lambda([Chooser, BoolColumn](ECheckBoxState NewValue) { BoolColumn->TestValue = NewValue == ECheckBoxState::Checked; })
 				]
diff --git a/Engine/Plugins/Experimental/Chooser/Source/ChooserEditor/Private/EnumColumnEditor.cpp b/Engine/Plugins/Experimental/Chooser/Source/ChooserEditor/Private/EnumColumnEditor.cpp
--- a/Engine/Plugins/Experimental/Chooser/Source/ChooserEditor/Private/EnumColumnEditor.cpp
+++ b/Engine/Plugins/Experimental/Chooser/Source/ChooserEditor/Private/EnumColumnEditor.cpp
@@ -2,6 +2,7 @@
 
 #include "EnumColumnEditor.h"
 #include "EnumColumn.h"
+#include "OutputEnumColumn.h"
 #include "SPropertyAccessChainWidget.h"
 #include "ObjectChooserWidgetFactories.h"
 #include "ChooserTableEditor.h"
@@ -16,6 +17,7 @@ namespace UE::ChooserEditor
 {
 	
 // Wrapper widget for EnumComboBox which will reconstruct the combo box when the Enum has changed
+template <typename ColumnType>
 class SEnumCell : public SCompoundWidget
 {
 public:
@@ -24,29 +26,21 @@ public:
 	{}
 
 	SLATE_ARGUMENT(UObject*, TransactionObject)
-	SLATE_ARGUMENT(FEnumColumn*, EnumColumn)
+	SLATE_ARGUMENT(ColumnType*, EnumColumn)
 	SLATE_ATTRIBUTE(int, RowIndex);
             
 	SLATE_END_ARGS()
 
 	TSharedRef<SWidget> CreateEnumComboBox()
 	{
-		if (const FEnumColumn* EnumColumnPointer = EnumColumn)
+		if (const ColumnType* EnumColumnPointer = EnumColumn)
 		{
 			if (EnumColumnPointer->InputValue.IsValid())
 			{
-				if (const UEnum* Enum = EnumColumnPointer->InputValue.Get<FChooserParameterEnumBase>().GetEnum())
+				if (const UEnum* Enum = EnumColumnPointer->InputValue.template Get<FChooserParameterEnumBase>().GetEnum())
 				{
 					return SNew(SEnumComboBox, Enum)
-						.IsEnabled_Lambda([this]
-						{
-							if (UChooserTable* Chooser = Cast<UChooserTable>(TransactionObject))
-							{
-								// return false only for the column header (RowIndex = -1) and when there is a debug target object bound
-								return RowIndex.Get()>=0 || !Chooser->HasDebugTarget();
-							}
-							return true;
-						})
+						.IsEnabled_Lambda([this](){ return IsEnabled(); } )
 						.CurrentValue_Lambda([this]()
 						{
 							int Row = RowIndex.Get();
@@ -56,7 +50,7 @@ public:
 							}
 							else
 							{
-								return EnumColumn->TestValue;
+								return static_cast<int32>(EnumColumn->TestValue);
 							}
 						})
 						.OnEnumSelectionChanged_Lambda([this](int32 EnumValue, ESelectInfo::Type)
@@ -90,7 +84,7 @@ public:
 		const UEnum* CurrentEnumSource = nullptr;
 		if (EnumColumn->InputValue.IsValid())
 		{
-			CurrentEnumSource = EnumColumn->InputValue.Get<FChooserParameterEnumBase>().GetEnum(); 
+			CurrentEnumSource = EnumColumn->InputValue.template Get<FChooserParameterEnumBase>().GetEnum(); 
 		}
 		if (EnumSource != CurrentEnumSource)
 		{
@@ -102,16 +96,18 @@ public:
 
 	void Construct( const FArguments& InArgs)
 	{
+		SetEnabled(InArgs._IsEnabled);
+		
 		SetCanTick(true);
 		RowIndex = InArgs._RowIndex;
 		EnumColumn = InArgs._EnumColumn;
 		TransactionObject = InArgs._TransactionObject;
 
-		if (const FEnumColumn* EnumColumnPointer = EnumColumn)
+		if (EnumColumn)
 		{
-			if (EnumColumnPointer->InputValue.IsValid())
+			if (EnumColumn->InputValue.IsValid())
 			{
-				EnumSource = EnumColumnPointer->InputValue.Get<FChooserParameterEnumBase>().GetEnum();
+				EnumSource = EnumColumn->InputValue.template Get<FChooserParameterEnumBase>().GetEnum();
 			}
 		}
 
@@ -121,36 +117,9 @@ public:
 
 		ChildSlot
 		[
-			SNew(SHorizontalBox)
-			+ SHorizontalBox::Slot().AutoWidth()
+			SAssignNew(EnumComboBorder, SBorder).Padding(0).BorderBackgroundColor(FLinearColor(0,0,0,0))
 			[
-				SNew(SBox).WidthOverride(Row < 0 ? 0 : 45)
-				[
-					SNew(SButton).ButtonStyle(FAppStyle::Get(),"FlatButton").TextStyle(FAppStyle::Get(),"RichTextBlock.Bold").HAlign(HAlign_Center)
-					.Visibility(Row < 0 ? EVisibility::Hidden : EVisibility::Visible)
-					.Text_Lambda([this, Row]()
-					{
-
-						return (EnumColumn->RowValues.IsValidIndex(Row) && EnumColumn->RowValues[Row].CompareNotEqual ? LOCTEXT("Not Equal", "!=") : LOCTEXT("Equal", "="));
-					})
-					.OnClicked_Lambda([this, Row]()
-					{
-						if (EnumColumn->RowValues.IsValidIndex(Row))
-						{
-							const FScopedTransaction Transaction(LOCTEXT("Edit Comparison", "Edit Comparison Operation"));
-							TransactionObject->Modify(true);
-							EnumColumn->RowValues[Row].CompareNotEqual = !EnumColumn->RowValues[Row].CompareNotEqual;
-						}
-						return FReply::Handled();
-					})
-				]
-			]
-			+ SHorizontalBox::Slot().FillWidth(1)
-			[
-				SAssignNew(EnumComboBorder, SBorder).Padding(0).BorderBackgroundColor(FLinearColor(0,0,0,0))
-				[
-					CreateEnumComboBox()
-				]
+				CreateEnumComboBox()
 			]
 		];
 		
@@ -162,7 +131,7 @@ public:
 
 private:
 	UObject* TransactionObject = nullptr;
-	FEnumColumn* EnumColumn = nullptr;
+	ColumnType* EnumColumn = nullptr;
 	const UEnum* EnumSource = nullptr;
 	TSharedPtr<SBorder> EnumComboBorder;
 	TAttribute<int> RowIndex;
@@ -208,7 +177,85 @@ TSharedRef<SWidget> CreateEnumColumnWidget(UChooserTable* Chooser, FChooserColum
 			]
 			+ SVerticalBox::Slot()
 			[
-				SNew(SEnumCell).TransactionObject(Chooser).EnumColumn(EnumColumn).RowIndex(Row)
+				SNew(SEnumCell<FEnumColumn>).TransactionObject(Chooser).EnumColumn(EnumColumn).RowIndex(Row)
+					.IsEnabled_Lambda([Chooser] { return !Chooser->HasDebugTarget(); })
+			];
+		}
+
+		return ColumnHeaderWidget;
+	}
+
+	// create cell widget
+	
+	return SNew(SHorizontalBox)
+    		+ SHorizontalBox::Slot().AutoWidth()
+    		[
+    			SNew(SBox).WidthOverride(Row < 0 ? 0 : 45)
+    			[
+    				SNew(SButton).ButtonStyle(FAppStyle::Get(),"FlatButton").TextStyle(FAppStyle::Get(),"RichTextBlock.Bold").HAlign(HAlign_Center)
+    				.Visibility(Row < 0 ? EVisibility::Hidden : EVisibility::Visible)
+					.Text_Lambda([EnumColumn, Row]()
+					{
+						return (EnumColumn->RowValues.IsValidIndex(Row) && EnumColumn->RowValues[Row].CompareNotEqual ? LOCTEXT("Not Equal", "!=") : LOCTEXT("Equal", "="));
+					})
+					.OnClicked_Lambda([EnumColumn, Chooser, Row]()
+					{
+						if (EnumColumn->RowValues.IsValidIndex(Row))
+						{
+							const FScopedTransaction Transaction(LOCTEXT("Edit Comparison", "Edit Comparison Operation"));
+							Chooser->Modify(true);
+							EnumColumn->RowValues[Row].CompareNotEqual = !EnumColumn->RowValues[Row].CompareNotEqual;
+						}
+						return FReply::Handled();
+					})
+				]
+			]
+			+ SHorizontalBox::Slot().FillWidth(1)
+			[
+				SNew(SEnumCell<FEnumColumn>).TransactionObject(Chooser).EnumColumn(EnumColumn).RowIndex(Row)
+			];
+}
+
+TSharedRef<SWidget> CreateOutputEnumColumnWidget(UChooserTable* Chooser, FChooserColumnBase* Column, int Row)
+{
+	FOutputEnumColumn* EnumColumn = static_cast<FOutputEnumColumn*>(Column);
+	
+	if (Row < 0)
+	{
+		// create column header widget
+		TSharedPtr<SWidget> InputValueWidget = nullptr;
+		if (FChooserParameterBase* InputValue = Column->GetInputValue())
+		{
+			InputValueWidget = FObjectChooserWidgetFactories::CreateWidget(false, Chooser, InputValue, Column->GetInputType(), Chooser->OutputObjectType);
+		}
+		
+		const FSlateBrush* ColumnIcon = FCoreStyle::Get().GetBrush("Icons.ArrowRight");
+		
+		TSharedRef<SWidget> ColumnHeaderWidget = SNew(SHorizontalBox)
+			+ SHorizontalBox::Slot().AutoWidth()
+			[
+				SNew(SBorder)
+				.BorderBackgroundColor(FLinearColor(0,0,0,0))
+				.Content()
+				[
+					SNew(SImage).Image(ColumnIcon)
+				]
+			]
+			+ SHorizontalBox::Slot()
+			[
+				InputValueWidget ? InputValueWidget.ToSharedRef() : SNullWidget::NullWidget
+			];
+	
+		if (Chooser->bEnableDebugTesting)
+		{
+			ColumnHeaderWidget = SNew(SVerticalBox)
+			+ SVerticalBox::Slot()
+			[
+				ColumnHeaderWidget
+			]
+			+ SVerticalBox::Slot()
+			[
+				SNew(SEnumCell<FOutputEnumColumn>).TransactionObject(Chooser).EnumColumn(EnumColumn).RowIndex(Row).IsEnabled(false)
 			];
 		}
 
@@ -217,7 +264,10 @@ TSharedRef<SWidget> CreateEnumColumnWidget(UChooserTable* Chooser, FChooserColum
 
 	// create cell widget
 	
-	return SNew(SEnumCell).TransactionObject(Chooser).EnumColumn(EnumColumn).RowIndex(Row);
+	return
+
+
+	SNew(SEnumCell<FOutputEnumColumn>).TransactionObject(Chooser).EnumColumn(EnumColumn).RowIndex(Row);
 }
 
 TSharedRef<SWidget> CreateEnumPropertyWidget(bool bReadOnly, UObject* TransactionObject, void* Value, UClass* ResultBaseClass)
@@ -241,6 +291,7 @@ void RegisterEnumWidgets()
 {
 	FObjectChooserWidgetFactories::RegisterWidgetCreator(FEnumContextProperty::StaticStruct(), CreateEnumPropertyWidget);
 	FObjectChooserWidgetFactories::RegisterColumnWidgetCreator(FEnumColumn::StaticStruct(), CreateEnumColumnWidget);
+	FObjectChooserWidgetFactories::RegisterColumnWidgetCreator(FOutputEnumColumn::StaticStruct(), CreateOutputEnumColumnWidget);
 }
 	
 }
