Fix note on EvaluateChooser node compilation

diff --git a/Engine/Plugins/Experimental/Chooser/Source/ChooserUncooked/Private/EvaluateChooserNode.cpp b/Engine/Plugins/Experimental/Chooser/Source/ChooserUncooked/Private/EvaluateChooserNode.cpp
--- a/Engine/Plugins/Experimental/Chooser/Source/ChooserUncooked/Private/EvaluateChooserNode.cpp
+++ b/Engine/Plugins/Experimental/Chooser/Source/ChooserUncooked/Private/EvaluateChooserNode.cpp
@@ -592,6 +592,7 @@ void UK2Node_EvaluateChooser2::ExpandNode(class FKismetCompilerContext& Compiler
 		CallFunction->AllocateDefaultPins();
 		
 		UK2Node_MakeStruct* ContextStructNode = CompilerContext.SpawnIntermediateNode<UK2Node_MakeStruct>(this, SourceGraph);
+		ContextStructNode->PostPlacedNewNode();
         CompilerContext.MessageLog.NotifyIntermediateObjectCreation(ContextStructNode, this);
 		ContextStructNode->StructType = FChooserEvaluationContext::StaticStruct();
         ContextStructNode->AllocateDefaultPins();
