Fix crash when using Evaluate Chooser in an Anim BP Event Graph

diff --git a/Engine/Plugins/Experimental/Chooser/Source/ChooserUncooked/Private/EvaluateChooserNode.cpp b/Engine/Plugins/Experimental/Chooser/Source/ChooserUncooked/Private/EvaluateChooserNode.cpp
--- a/Engine/Plugins/Experimental/Chooser/Source/ChooserUncooked/Private/EvaluateChooserNode.cpp
+++ b/Engine/Plugins/Experimental/Chooser/Source/ChooserUncooked/Private/EvaluateChooserNode.cpp
@@ -535,7 +535,14 @@ void UK2Node_EvaluateChooser2::PostEditChangeProperty(struct FPropertyChangedEve
 void UK2Node_EvaluateChooser2::PostLoad()
 {
 	Super::PostLoad();
-	ChooserChanged();
+	
+	if (Chooser)
+	{
+		Chooser->OnOutputObjectTypeChanged.AddUObject(this, &UK2Node_EvaluateChooser2::ResultTypeChanged);
+		Chooser->OnContextClassChanged.AddUObject(this, &UK2Node::ReconstructNode);
+	}
+
+	CurrentCallbackChooser = Chooser;
 }
 
 void UK2Node_EvaluateChooser2::PinConnectionListChanged(UEdGraphPin* Pin)
