Output Struct Column for Choosers

diff --git a/Engine/Plugins/Experimental/Chooser/Source/Chooser/Internal/OutputStructColumn.h b/Engine/Plugins/Experimental/Chooser/Source/Chooser/Internal/OutputStructColumn.h
new file mode 100644
--- /dev/null
+++ b/Engine/Plugins/Experimental/Chooser/Source/Chooser/Internal/OutputStructColumn.h
@@ -0,0 +1,69 @@
+// Copyright Epic Games, Inc. All Rights Reserved.
+#pragma once
+
+#include "CoreMinimal.h"
+#include "IChooserColumn.h"
+#include "IChooserParameterStruct.h"
+#include "ChooserPropertyAccess.h"
+#include "InstancedStruct.h"
+#include "OutputStructColumn.generated.h"
+
+struct FBindingChainElement;
+
+USTRUCT(DisplayName = "Struct Property Binding")
+struct CHOOSER_API FStructContextProperty : public FChooserParameterStructBase
+{
+	GENERATED_BODY()
+
+	UPROPERTY(EditAnywhere, Meta = (BindingType = "struct", BindingColor = "StructPinTypeColor"), Category = "Binding")
+	FChooserStructPropertyBinding Binding;
+
+	virtual bool SetValue(FChooserEvaluationContext& Context, const FInstancedStruct &Value) const override;
+
+#if WITH_EDITOR
+	void SetBinding(const UObject* OuterObject, const TArray<FBindingChainElement>& InBindingChain);
+
+	virtual void GetDisplayName(FText& OutName) const override
+	{
+		if (!Binding.DisplayName.IsEmpty())
+		{
+			OutName = FText::FromString(Binding.DisplayName);
+		} 
+		else if (!Binding.PropertyBindingChain.IsEmpty())
+		{
+			OutName = FText::FromName(Binding.PropertyBindingChain.Last());
+		}
+	}
+
+	virtual UScriptStruct* GetStructType() const override { return Binding.StructType; }
+#endif
+};
+
+USTRUCT()
+struct CHOOSER_API FOutputStructColumn : public FChooserColumnBase
+{
+	GENERATED_BODY()
+	public:
+	FOutputStructColumn();
+	virtual bool HasFilters() const override { return false; }
+	virtual bool HasOutputs() const override { return true; }
+	virtual void SetOutputs(FChooserEvaluationContext& Context, int RowIndex) const override;
+	
+	UPROPERTY(EditAnywhere, Meta = (ExcludeBaseStruct, BaseStruct = "/Script/Chooser.ChooserParameterStructBase"), Category = "Hidden")
+	FInstancedStruct InputValue;
+
+#if WITH_EDITOR
+	void StructTypeChanged();
+	mutable FInstancedStruct TestValue;
+#endif
+	
+#if WITH_EDITORONLY_DATA
+	UPROPERTY(EditAnywhere, Meta = (StructTypeConst), Category=Runtime);
+	FInstancedStruct DefaultRowValue;
+#endif
+	
+	UPROPERTY(EditAnywhere, Meta = (StructTypeConst), Category=Runtime);
+	TArray<FInstancedStruct> RowValues; 
+
+	CHOOSER_COLUMN_BOILERPLATE(FChooserParameterStructBase);
+};
\ No newline at end of file
diff --git a/Engine/Plugins/Experimental/Chooser/Source/Chooser/Private/ChooserPropertyAccess.cpp b/Engine/Plugins/Experimental/Chooser/Source/Chooser/Private/ChooserPropertyAccess.cpp
--- a/Engine/Plugins/Experimental/Chooser/Source/Chooser/Private/ChooserPropertyAccess.cpp
+++ b/Engine/Plugins/Experimental/Chooser/Source/Chooser/Private/ChooserPropertyAccess.cpp
@@ -13,7 +13,7 @@ namespace UE::Chooser
 {
 	bool ResolvePropertyChain(FChooserEvaluationContext& Context, const FChooserPropertyBinding& PropertyBinding, const void*& OutContainer, const UStruct*& OutStructType)
 	{
-		if (PropertyBinding.ContextIndex < Context.Params.Num())
+		if (Context.Params.IsValidIndex(PropertyBinding.ContextIndex))
 		{
 			if (FChooserEvaluationInputObject* ObjectParam = Context.Params[PropertyBinding.ContextIndex].GetMutablePtr<FChooserEvaluationInputObject>())
 			{
@@ -103,7 +103,11 @@ namespace UE::Chooser
 	{
 		OutPropertyBinding.PropertyBindingChain.Empty();
 
-		if (InBindingChain.Num() > 0)
+		if (InBindingChain.Num() == 0)
+		{
+			OutPropertyBinding.ContextIndex = -1;
+		}
+		else
 		{
 			OutPropertyBinding.ContextIndex = InBindingChain[0].ArrayIndex;
 		}
diff --git a/Engine/Plugins/Experimental/Chooser/Source/Chooser/Private/OutputStructColumn.cpp b/Engine/Plugins/Experimental/Chooser/Source/Chooser/Private/OutputStructColumn.cpp
new file mode 100644
--- /dev/null
+++ b/Engine/Plugins/Experimental/Chooser/Source/Chooser/Private/OutputStructColumn.cpp
@@ -0,0 +1,128 @@
+// Copyright Epic Games, Inc. All Rights Reserved.
+#include "OutputStructColumn.h"
+#include "ChooserPropertyAccess.h"
+#include "OutputStructColumn.h"
+
+#if WITH_EDITOR
+	#include "IPropertyAccessEditor.h"
+#endif
+
+bool FStructContextProperty::SetValue(FChooserEvaluationContext& Context, const FInstancedStruct& InValue) const
+{
+	const UStruct* StructType = nullptr;
+	const void* Container = nullptr;
+
+	if (Binding.PropertyBindingChain.IsEmpty())
+	{
+		if(Context.Params.IsValidIndex((Binding.ContextIndex)))
+		{
+			// directly bound to context struct
+			if (Context.Params[Binding.ContextIndex].GetScriptStruct() == InValue.GetScriptStruct())
+			{
+				void* TargetData = Context.Params[Binding.ContextIndex].GetMutableMemory();
+				InValue.GetScriptStruct()->CopyScriptStruct(TargetData, InValue.GetMemory());
+			}
+		}
+	}
+	else if (UE::Chooser::ResolvePropertyChain(Context, Binding, Container, StructType))
+	{
+		if (FStructProperty* Property = FindFProperty<FStructProperty>(StructType, Binding.PropertyBindingChain.Last()))
+		{
+			// const cast is here just because ResolvePropertyChain expects a const void*&
+			void* TargetData = Property->ContainerPtrToValuePtr<void>(const_cast<void*>(Container));
+			
+			if (Property->Struct == InValue.GetScriptStruct())
+			{
+				Property->Struct->CopyScriptStruct(TargetData, InValue.GetMemory());
+			}
+
+			return true;
+		}
+	}
+
+	return false;
+}
+
+#if WITH_EDITOR
+
+void FStructContextProperty::SetBinding(const UObject* OuterObject, const TArray<FBindingChainElement>& InBindingChain)
+{
+	Binding.StructType = nullptr;
+
+	UE::Chooser::CopyPropertyChain(InBindingChain, Binding);
+
+	if (Binding.PropertyBindingChain.Num() == 0)
+	{
+		// binding directly to context struct, get struct type from there
+	    if (const IHasContextClass* HasContextClass = Cast<IHasContextClass>(OuterObject))
+	    {
+	    	TConstArrayView<FInstancedStruct> ContextData = HasContextClass->GetContextData();
+	    	if (ContextData.IsValidIndex(Binding.ContextIndex))
+		    {
+	    		if (const FContextObjectTypeStruct* StructContextData = ContextData[Binding.ContextIndex].GetPtr<FContextObjectTypeStruct>())
+	    		{
+					Binding.StructType = StructContextData->Struct;
+	    			Binding.DisplayName = StructContextData->Struct->GetAuthoredName();
+	    		}
+		    }
+	    }
+	}
+	else
+	{
+		const FField* Field = InBindingChain.Last().Field.ToField();
+		if (const FStructProperty* StructProperty = CastField<FStructProperty>(Field))
+		{
+			Binding.StructType = StructProperty->Struct;
+			Binding.DisplayName = StructProperty->GetAuthoredName();
+			if (InBindingChain.Num() > 2)
+			{
+				// add the parent struct name to the display name if there is one
+				const FField* NextToLastField = InBindingChain[InBindingChain.Num()-2].Field.ToField();
+				Binding.DisplayName = NextToLastField->GetAuthoredName() + "." + Binding.DisplayName;
+			}
+		}
+	}
+}
+
+void FOutputStructColumn::StructTypeChanged()
+{
+	if (InputValue.IsValid())
+	{
+		const UScriptStruct* Struct = InputValue.Get<FChooserParameterStructBase>().GetStructType();
+
+		if (DefaultRowValue.GetScriptStruct() != Struct)
+		{
+			DefaultRowValue.InitializeAs(Struct);
+		}
+
+		for (FInstancedStruct& RowValue : RowValues)
+		{
+			if (RowValue.GetScriptStruct() != Struct)
+			{
+				RowValue.InitializeAs(Struct);
+			}
+		}
+	}
+}
+
+#endif // WITH_EDITOR
+
+FOutputStructColumn::FOutputStructColumn()
+{
+	InputValue.InitializeAs(FStructContextProperty::StaticStruct());
+}
+
+void FOutputStructColumn::SetOutputs(FChooserEvaluationContext& Context, int RowIndex) const
+{
+	if (InputValue.IsValid())
+	{
+		InputValue.Get<FChooserParameterStructBase>().SetValue(Context, RowValues[RowIndex]);
+	}
+	
+#if WITH_EDITOR
+	if (Context.DebuggingInfo.bCurrentDebugTarget)
+	{
+		TestValue = RowValues[RowIndex];
+	}
+#endif
+}
diff --git a/Engine/Plugins/Experimental/Chooser/Source/Chooser/Public/ChooserPropertyAccess.h b/Engine/Plugins/Experimental/Chooser/Source/Chooser/Public/ChooserPropertyAccess.h
--- a/Engine/Plugins/Experimental/Chooser/Source/Chooser/Public/ChooserPropertyAccess.h
+++ b/Engine/Plugins/Experimental/Chooser/Source/Chooser/Public/ChooserPropertyAccess.h
@@ -38,7 +38,12 @@ struct FChooserPropertyBinding
 	TArray<FName> PropertyBindingChain;
 	
 	UPROPERTY()
-	int ContextIndex = 0;
+	int ContextIndex;
+
+#if WITH_EDITORONLY_DATA
+	UPROPERTY()
+	FString DisplayName;
+#endif
 };
 
 USTRUCT()
@@ -63,6 +68,16 @@ struct FChooserObjectPropertyBinding : public FChooserPropertyBinding
 #endif
 };
 
+USTRUCT()
+struct FChooserStructPropertyBinding : public FChooserPropertyBinding
+{
+	GENERATED_BODY()
+
+#if WITH_EDITORONLY_DATA
+	UPROPERTY()
+	TObjectPtr<UScriptStruct> StructType = nullptr;
+#endif
+};
 
 DECLARE_MULTICAST_DELEGATE_OneParam(FChooserOutputObjectTypeChanged, const UClass* OutputObjectType);
 UENUM()
@@ -97,7 +112,7 @@ struct FContextObjectTypeStruct : public FContextObjectTypeBase
 	GENERATED_BODY()
 	
 	UPROPERTY(EditAnywhere, Category="Type")
-	TObjectPtr<UStruct> Struct;
+	TObjectPtr<UScriptStruct> Struct;
 };
 
 
diff --git a/Engine/Plugins/Experimental/Chooser/Source/Chooser/Public/IChooserParameterStruct.h b/Engine/Plugins/Experimental/Chooser/Source/Chooser/Public/IChooserParameterStruct.h
new file mode 100644
--- /dev/null
+++ b/Engine/Plugins/Experimental/Chooser/Source/Chooser/Public/IChooserParameterStruct.h
@@ -0,0 +1,22 @@
+// Copyright Epic Games, Inc. All Rights Reserved.
+
+#pragma once
+
+#include "CoreMinimal.h"
+#include "IChooserParameterBase.h"
+#include "IChooserParameterStruct.generated.h"
+
+struct FInstancedStruct;
+class UStruct;
+
+USTRUCT()
+struct FChooserParameterStructBase : public FChooserParameterBase
+{
+	GENERATED_BODY()
+
+	virtual bool SetValue(FChooserEvaluationContext& Context, const FInstancedStruct &Value)  const { return false; }
+
+#if WITH_EDITOR
+	virtual UScriptStruct* GetStructType() const { return nullptr; }
+#endif
+};
diff --git a/Engine/Plugins/Experimental/Chooser/Source/ChooserEditor/Private/BoolColumnEditor.cpp b/Engine/Plugins/Experimental/Chooser/Source/ChooserEditor/Private/BoolColumnEditor.cpp
--- a/Engine/Plugins/Experimental/Chooser/Source/ChooserEditor/Private/BoolColumnEditor.cpp
+++ b/Engine/Plugins/Experimental/Chooser/Source/ChooserEditor/Private/BoolColumnEditor.cpp
@@ -183,7 +183,7 @@ TSharedRef<SWidget> CreateOutputBoolColumnWidget(UChooserTable* Chooser, FChoose
 	+ SHorizontalBox::Slot().FillWidth(1);
 }
 	
-TSharedRef<SWidget> CreateBoolPropertyWidget(bool bReadOnly, UObject* TransactionObject, void* Value, UClass* ResultBaseClass)
+TSharedRef<SWidget> CreateBoolPropertyWidget(bool bReadOnly, UObject* TransactionObject, void* Value, UClass* ResultBaseClass, FChooserWidgetValueChanged ValueChanged)
 {
 	IHasContextClass* HasContextClass = Cast<IHasContextClass>(TransactionObject);
 
@@ -192,11 +192,12 @@ TSharedRef<SWidget> CreateBoolPropertyWidget(bool bReadOnly, UObject* Transactio
 	return SNew(SPropertyAccessChainWidget).ContextClassOwner(HasContextClass).AllowFunctions(true).BindingColor("BooleanPinTypeColor").TypeFilter("bool")
 	.PropertyBindingValue(&ContextProperty->Binding)
 	.OnAddBinding_Lambda(
-		[ContextProperty, TransactionObject](FName InPropertyName, const TArray<FBindingChainElement>& InBindingChain)
+		[ContextProperty, TransactionObject, ValueChanged](FName InPropertyName, const TArray<FBindingChainElement>& InBindingChain)
 		{
 			const FScopedTransaction Transaction(NSLOCTEXT("ContextPropertyWidget", "Change Property Binding", "Change Property Binding"));
 			TransactionObject->Modify(true);
-			ContextProperty->SetBinding(InBindingChain);	
+			ContextProperty->SetBinding(InBindingChain);
+			ValueChanged.ExecuteIfBound();
 		});
 }
 	
diff --git a/Engine/Plugins/Experimental/Chooser/Source/ChooserEditor/Private/ChooserEditorModule.cpp b/Engine/Plugins/Experimental/Chooser/Source/ChooserEditor/Private/ChooserEditorModule.cpp
--- a/Engine/Plugins/Experimental/Chooser/Source/ChooserEditor/Private/ChooserEditorModule.cpp
+++ b/Engine/Plugins/Experimental/Chooser/Source/ChooserEditor/Private/ChooserEditorModule.cpp
@@ -13,6 +13,7 @@
 #include "ChooserPropertyAccess.h"
 #include "PropertyEditorModule.h"
 #include "RandomizeColumnEditor.h"
+#include "OutputStructColumnEditor.h"
 
 #define LOCTEXT_NAMESPACE "ChooserEditorModule"
 
@@ -28,6 +29,7 @@ void FModule::StartupModule()
 	RegisterBoolWidgets();
 	RegisterEnumWidgets();
 	RegisterObjectWidgets();
+	RegisterStructWidgets();
 	RegisterRandomizeWidgets();
 	
 	FChooserTableEditorCommands::Register();
@@ -36,6 +38,8 @@ void FModule::StartupModule()
 	
 	PropertyModule.RegisterCustomPropertyTypeLayout(FChooserPropertyBinding::StaticStruct()->GetFName(), FOnGetPropertyTypeCustomizationInstance::CreateLambda([] { return MakeShared<FPropertyAccessChainCustomization>(); }));
 	PropertyModule.RegisterCustomPropertyTypeLayout(FChooserEnumPropertyBinding::StaticStruct()->GetFName(), FOnGetPropertyTypeCustomizationInstance::CreateLambda([] { return MakeShared<FPropertyAccessChainCustomization>(); }));
+	PropertyModule.RegisterCustomPropertyTypeLayout(FChooserObjectPropertyBinding::StaticStruct()->GetFName(), FOnGetPropertyTypeCustomizationInstance::CreateLambda([] { return MakeShared<FPropertyAccessChainCustomization>(); }));
+	PropertyModule.RegisterCustomPropertyTypeLayout(FChooserStructPropertyBinding::StaticStruct()->GetFName(), FOnGetPropertyTypeCustomizationInstance::CreateLambda([] { return MakeShared<FPropertyAccessChainCustomization>(); }));
 
 }
 
diff --git a/Engine/Plugins/Experimental/Chooser/Source/ChooserEditor/Private/ChooserTableEditor.cpp b/Engine/Plugins/Experimental/Chooser/Source/ChooserEditor/Private/ChooserTableEditor.cpp
--- a/Engine/Plugins/Experimental/Chooser/Source/ChooserEditor/Private/ChooserTableEditor.cpp
+++ b/Engine/Plugins/Experimental/Chooser/Source/ChooserEditor/Private/ChooserTableEditor.cpp
@@ -1220,7 +1220,7 @@ void FChooserTableEditor::DeleteColumn(int Index)
 	}
 }
 
-TSharedRef<SWidget> CreateAssetWidget(bool bReadOnly, UObject* TransactionObject, void* Value, UClass* ResultBaseClass)
+TSharedRef<SWidget> CreateAssetWidget(bool bReadOnly, UObject* TransactionObject, void* Value, UClass* ResultBaseClass, FChooserWidgetValueChanged ValueChanged)
 {
 	FAssetChooser* DIAsset = static_cast<FAssetChooser*>(Value);
 
@@ -1238,7 +1238,7 @@ TSharedRef<SWidget> CreateAssetWidget(bool bReadOnly, UObject* TransactionObject
 		});
 }
 	
-TSharedRef<SWidget> CreateClassWidget(bool bReadOnly, UObject* TransactionObject, void* Value, UClass* ResultBaseClass)
+TSharedRef<SWidget> CreateClassWidget(bool bReadOnly, UObject* TransactionObject, void* Value, UClass* ResultBaseClass, FChooserWidgetValueChanged ValueChanged)
 {
 	FClassChooser* ClassChooser = static_cast<FClassChooser*>(Value);
 
@@ -1259,7 +1259,7 @@ TSharedRef<SWidget> CreateClassWidget(bool bReadOnly, UObject* TransactionObject
 			});
 }
 
-TSharedRef<SWidget> CreateEvaluateChooserWidget(bool bReadOnly, UObject* TransactionObject, void* Value, UClass* ResultBaseClass)
+TSharedRef<SWidget> CreateEvaluateChooserWidget(bool bReadOnly, UObject* TransactionObject, void* Value, UClass* ResultBaseClass, FChooserWidgetValueChanged ValueChanged)
 {
 	FEvaluateChooser* EvaluateChooser = static_cast<FEvaluateChooser*>(Value);
 	
diff --git a/Engine/Plugins/Experimental/Chooser/Source/ChooserEditor/Private/EnumColumnEditor.cpp b/Engine/Plugins/Experimental/Chooser/Source/ChooserEditor/Private/EnumColumnEditor.cpp
--- a/Engine/Plugins/Experimental/Chooser/Source/ChooserEditor/Private/EnumColumnEditor.cpp
+++ b/Engine/Plugins/Experimental/Chooser/Source/ChooserEditor/Private/EnumColumnEditor.cpp
@@ -270,7 +270,7 @@ TSharedRef<SWidget> CreateOutputEnumColumnWidget(UChooserTable* Chooser, FChoose
 	SNew(SEnumCell<FOutputEnumColumn>).TransactionObject(Chooser).EnumColumn(EnumColumn).RowIndex(Row);
 }
 
-TSharedRef<SWidget> CreateEnumPropertyWidget(bool bReadOnly, UObject* TransactionObject, void* Value, UClass* ResultBaseClass)
+TSharedRef<SWidget> CreateEnumPropertyWidget(bool bReadOnly, UObject* TransactionObject, void* Value, UClass* ResultBaseClass, FChooserWidgetValueChanged ValueChanged)
 {
 	IHasContextClass* HasContextClass = Cast<IHasContextClass>(TransactionObject);
 
@@ -279,11 +279,12 @@ TSharedRef<SWidget> CreateEnumPropertyWidget(bool bReadOnly, UObject* Transactio
 	return SNew(SPropertyAccessChainWidget).ContextClassOwner(HasContextClass).AllowFunctions(false).BindingColor("BytePinTypeColor").TypeFilter("enum")
 	.PropertyBindingValue(&ContextProperty->Binding)
 	.OnAddBinding_Lambda(
-		[ContextProperty, TransactionObject](FName InPropertyName, const TArray<FBindingChainElement>& InBindingChain)
+		[ContextProperty, TransactionObject, ValueChanged](FName InPropertyName, const TArray<FBindingChainElement>& InBindingChain)
 		{
 			const FScopedTransaction Transaction(NSLOCTEXT("ContextPropertyWidget", "Change Property Binding", "Change Property Binding"));
 			TransactionObject->Modify(true);
-			ContextProperty->SetBinding(InBindingChain);	
+			ContextProperty->SetBinding(InBindingChain);
+			ValueChanged.ExecuteIfBound();
 		});
 }
 	
diff --git a/Engine/Plugins/Experimental/Chooser/Source/ChooserEditor/Private/FloatRangeColumnEditor.cpp b/Engine/Plugins/Experimental/Chooser/Source/ChooserEditor/Private/FloatRangeColumnEditor.cpp
--- a/Engine/Plugins/Experimental/Chooser/Source/ChooserEditor/Private/FloatRangeColumnEditor.cpp
+++ b/Engine/Plugins/Experimental/Chooser/Source/ChooserEditor/Private/FloatRangeColumnEditor.cpp
@@ -13,7 +13,7 @@
 
 namespace UE::ChooserEditor
 {
-TSharedRef<SWidget> CreateFloatPropertyWidget(bool bReadOnly, UObject* TransactionObject, void* Value, UClass* ResultBaseClass)
+TSharedRef<SWidget> CreateFloatPropertyWidget(bool bReadOnly, UObject* TransactionObject, void* Value, UClass* ResultBaseClass, FChooserWidgetValueChanged ValueChanged)
 {
 	IHasContextClass* HasContextClass = Cast<IHasContextClass>(TransactionObject);
 
@@ -22,11 +22,12 @@ TSharedRef<SWidget> CreateFloatPropertyWidget(bool bReadOnly, UObject* Transacti
 	return SNew(SPropertyAccessChainWidget).ContextClassOwner(HasContextClass).AllowFunctions(true).BindingColor("FloatPinTypeColor").TypeFilter("double")
 	.PropertyBindingValue(&ContextProperty->Binding)
 	.OnAddBinding_Lambda(
-		[ContextProperty, TransactionObject](FName InPropertyName, const TArray<FBindingChainElement>& InBindingChain)
+		[ContextProperty, TransactionObject, ValueChanged](FName InPropertyName, const TArray<FBindingChainElement>& InBindingChain)
 		{
 			const FScopedTransaction Transaction(NSLOCTEXT("ContextPropertyWidget", "Change Property Binding", "Change Property Binding"));
 			TransactionObject->Modify(true);
-			ContextProperty->SetBinding(InBindingChain);	
+			ContextProperty->SetBinding(InBindingChain);
+			ValueChanged.ExecuteIfBound();
 		});
 }
 	
diff --git a/Engine/Plugins/Experimental/Chooser/Source/ChooserEditor/Private/GameplayTagColumnEditor.cpp b/Engine/Plugins/Experimental/Chooser/Source/ChooserEditor/Private/GameplayTagColumnEditor.cpp
--- a/Engine/Plugins/Experimental/Chooser/Source/ChooserEditor/Private/GameplayTagColumnEditor.cpp
+++ b/Engine/Plugins/Experimental/Chooser/Source/ChooserEditor/Private/GameplayTagColumnEditor.cpp
@@ -111,7 +111,7 @@ TSharedRef<SWidget> CreateGameplayTagColumnWidget(UChooserTable* Chooser, FChoos
 	);
 }
 
-TSharedRef<SWidget> CreateGameplayTagPropertyWidget(bool bReadOnly, UObject* TransactionObject, void* Value, UClass* ResultBaseClass)
+TSharedRef<SWidget> CreateGameplayTagPropertyWidget(bool bReadOnly, UObject* TransactionObject, void* Value, UClass* ResultBaseClass, FChooserWidgetValueChanged ValueChanged)
 {
 	IHasContextClass* HasContextClass = Cast<IHasContextClass>(TransactionObject);
 
@@ -120,11 +120,12 @@ TSharedRef<SWidget> CreateGameplayTagPropertyWidget(bool bReadOnly, UObject* Tra
 	return SNew(SPropertyAccessChainWidget).ContextClassOwner(HasContextClass).AllowFunctions(false).BindingColor("StructPinTypeColor").TypeFilter("FGameplayTagContainer")
 	.PropertyBindingValue(&ContextProperty->Binding)
 	.OnAddBinding_Lambda(
-		[ContextProperty, TransactionObject](FName InPropertyName, const TArray<FBindingChainElement>& InBindingChain)
+		[ContextProperty, TransactionObject, ValueChanged](FName InPropertyName, const TArray<FBindingChainElement>& InBindingChain)
 		{
 			const FScopedTransaction Transaction(NSLOCTEXT("ContextPropertyWidget", "Change Property Binding", "Change Property Binding"));
 			TransactionObject->Modify(true);
-			ContextProperty->SetBinding(InBindingChain);	
+			ContextProperty->SetBinding(InBindingChain);
+			ValueChanged.ExecuteIfBound();
 		});
 }
 	
diff --git a/Engine/Plugins/Experimental/Chooser/Source/ChooserEditor/Private/ObjectChooserWidgetFactories.cpp b/Engine/Plugins/Experimental/Chooser/Source/ChooserEditor/Private/ObjectChooserWidgetFactories.cpp
--- a/Engine/Plugins/Experimental/Chooser/Source/ChooserEditor/Private/ObjectChooserWidgetFactories.cpp
+++ b/Engine/Plugins/Experimental/Chooser/Source/ChooserEditor/Private/ObjectChooserWidgetFactories.cpp
@@ -46,7 +46,7 @@ TSharedPtr<SWidget> FObjectChooserWidgetFactories::CreateColumnWidget(FChooserCo
 }
 	
 	
-TSharedPtr<SWidget> FObjectChooserWidgetFactories::CreateWidget(bool bReadOnly, UObject* TransactionObject, void* Value, const UStruct* ValueType, UClass* ResultBaseClass)
+TSharedPtr<SWidget> FObjectChooserWidgetFactories::CreateWidget(bool bReadOnly, UObject* TransactionObject, void* Value, const UStruct* ValueType, UClass* ResultBaseClass, FChooserWidgetValueChanged ValueChanged)
 {
 	if (Value)
 	{
@@ -54,7 +54,7 @@ TSharedPtr<SWidget> FObjectChooserWidgetFactories::CreateWidget(bool bReadOnly,
 		{
 			if (FChooserWidgetCreator* Creator = ChooserWidgetCreators.Find(ValueType))
 			{
-				return (*Creator)(bReadOnly, TransactionObject, Value, ResultBaseClass);
+				return (*Creator)(bReadOnly, TransactionObject, Value, ResultBaseClass, ValueChanged);
 			}
 			ValueType = ValueType->GetSuperStruct();
 		}
@@ -63,9 +63,9 @@ TSharedPtr<SWidget> FObjectChooserWidgetFactories::CreateWidget(bool bReadOnly,
 	return nullptr;
 }
 
-TSharedPtr<SWidget> FObjectChooserWidgetFactories::CreateWidget(bool bReadOnly, UObject* TransactionObject, const UScriptStruct* BaseType, void* Value, const UStruct* ValueType, UClass* ResultBaseClass, const FOnStructPicked& CreateClassCallback, TSharedPtr<SBorder>* InnerWidget)
+TSharedPtr<SWidget> FObjectChooserWidgetFactories::CreateWidget(bool bReadOnly, UObject* TransactionObject, const UScriptStruct* BaseType, void* Value, const UStruct* ValueType, UClass* ResultBaseClass, const FOnStructPicked& CreateClassCallback, TSharedPtr<SBorder>* InnerWidget, FChooserWidgetValueChanged ValueChanged)
 {
-	TSharedPtr<SWidget> LeftWidget = CreateWidget(bReadOnly, TransactionObject, Value, ValueType, ResultBaseClass);
+	TSharedPtr<SWidget> LeftWidget = CreateWidget(bReadOnly, TransactionObject, Value, ValueType, ResultBaseClass, ValueChanged);
 	if (bReadOnly)
 	{
 		// don't need the type selector dropdown when read only
diff --git a/Engine/Plugins/Experimental/Chooser/Source/ChooserEditor/Private/ObjectColumnEditor.cpp b/Engine/Plugins/Experimental/Chooser/Source/ChooserEditor/Private/ObjectColumnEditor.cpp
--- a/Engine/Plugins/Experimental/Chooser/Source/ChooserEditor/Private/ObjectColumnEditor.cpp
+++ b/Engine/Plugins/Experimental/Chooser/Source/ChooserEditor/Private/ObjectColumnEditor.cpp
@@ -215,7 +215,7 @@ namespace UE::ChooserEditor
 		return CellWidget;
 	}
 
-	static TSharedRef<SWidget> CreateObjectPropertyWidget(bool bReadOnly, UObject* TransactionObject, void* Value, UClass* ResultBaseClass)
+	static TSharedRef<SWidget> CreateObjectPropertyWidget(bool bReadOnly, UObject* TransactionObject, void* Value, UClass* ResultBaseClass, FChooserWidgetValueChanged ValueChanged)
 	{
 		IHasContextClass* HasContextClass = Cast<IHasContextClass>(TransactionObject);
 
@@ -224,11 +224,12 @@ namespace UE::ChooserEditor
 		return SNew(SPropertyAccessChainWidget).ContextClassOwner(HasContextClass).AllowFunctions(false).BindingColor("ObjectPinTypeColor").TypeFilter("object")
 		.PropertyBindingValue(&ContextProperty->Binding)
 		.OnAddBinding_Lambda(
-			[ContextProperty, TransactionObject](FName InPropertyName, const TArray<FBindingChainElement>& InBindingChain)
+			[ContextProperty, TransactionObject, ValueChanged](FName InPropertyName, const TArray<FBindingChainElement>& InBindingChain)
 			{
 				const FScopedTransaction Transaction(NSLOCTEXT("ContextPropertyWidget", "Change Property Binding", "Change Property Binding"));
 				TransactionObject->Modify(true);
-				ContextProperty->SetBinding(InBindingChain);	
+				ContextProperty->SetBinding(InBindingChain);
+				ValueChanged.ExecuteIfBound();
 			});
 	}
 
diff --git a/Engine/Plugins/Experimental/Chooser/Source/ChooserEditor/Private/OutputStructColumnEditor.cpp b/Engine/Plugins/Experimental/Chooser/Source/ChooserEditor/Private/OutputStructColumnEditor.cpp
new file mode 100644
--- /dev/null
+++ b/Engine/Plugins/Experimental/Chooser/Source/ChooserEditor/Private/OutputStructColumnEditor.cpp
@@ -0,0 +1,83 @@
+// Copyright Epic Games, Inc. All Rights Reserved.
+
+#include "OutputStructColumnEditor.h"
+#include "OutputStructColumn.h"
+#include "SPropertyAccessChainWidget.h"
+#include "ObjectChooserWidgetFactories.h"
+#include "Widgets/SBoxPanel.h"
+#include "Widgets/Images/SImage.h"
+#include "Widgets/SNullWidget.h"
+#include "GraphEditorSettings.h"
+
+#define LOCTEXT_NAMESPACE "StructOutputColumnEditor"
+
+namespace UE::ChooserEditor
+{
+	
+TSharedRef<SWidget> CreateOutputStructColumnWidget(UChooserTable* Chooser, FChooserColumnBase* Column, int Row)
+{
+	if (Row < 0)
+	{
+		// create column header widget
+		TSharedPtr<SWidget> InputValueWidget = nullptr;
+		if (FChooserParameterBase* InputValue = Column->GetInputValue())
+		{
+			InputValueWidget = FObjectChooserWidgetFactories::CreateWidget(false, Chooser, InputValue, Column->GetInputType(), Chooser->OutputObjectType,
+				FChooserWidgetValueChanged::CreateLambda([Column]()
+				{
+					FOutputStructColumn* StructColumn = static_cast<FOutputStructColumn*>(Column);
+					StructColumn->StructTypeChanged();
+				})
+				);
+		}
+		
+		const FSlateBrush* ColumnIcon = FCoreStyle::Get().GetBrush("Icons.ArrowRight");
+		
+		TSharedRef<SWidget> ColumnHeaderWidget = SNew(SHorizontalBox)
+			+ SHorizontalBox::Slot().AutoWidth()
+			[
+				SNew(SBorder)
+				.BorderBackgroundColor(FLinearColor(0,0,0,0))
+				.Content()
+				[
+					SNew(SImage).Image(ColumnIcon)
+				]
+			]
+			+ SHorizontalBox::Slot()
+			[
+				InputValueWidget ? InputValueWidget.ToSharedRef() : SNullWidget::NullWidget
+			];
+	
+		return ColumnHeaderWidget;
+	}
+
+	return SNullWidget::NullWidget;
+}
+	
+TSharedRef<SWidget> CreateStructPropertyWidget(bool bReadOnly, UObject* TransactionObject, void* Value, UClass* ResultBaseClass, FChooserWidgetValueChanged ValueChanged)
+{
+	IHasContextClass* HasContextClass = Cast<IHasContextClass>(TransactionObject);
+
+	FStructContextProperty* ContextProperty = reinterpret_cast<FStructContextProperty*>(Value);
+
+	return SNew(SPropertyAccessChainWidget).ContextClassOwner(HasContextClass).BindingColor("StructPinTypeColor").TypeFilter("struct")
+		.PropertyBindingValue(&ContextProperty->Binding)
+		.OnAddBinding_Lambda(
+    		[ContextProperty, TransactionObject, ValueChanged](FName InPropertyName, const TArray<FBindingChainElement>& InBindingChain)
+    		{
+    			const FScopedTransaction Transaction(NSLOCTEXT("ContextPropertyWidget", "Change Property Binding", "Change Property Binding"));
+    			TransactionObject->Modify(true);
+    			ContextProperty->SetBinding(TransactionObject, InBindingChain);
+    			ValueChanged.ExecuteIfBound();
+    		});
+}
+	
+void RegisterStructWidgets()
+{
+	FObjectChooserWidgetFactories::RegisterWidgetCreator(FStructContextProperty::StaticStruct(), CreateStructPropertyWidget);
+	FObjectChooserWidgetFactories::RegisterColumnWidgetCreator(FOutputStructColumn::StaticStruct(), CreateOutputStructColumnWidget);
+}
+	
+}
+
+#undef LOCTEXT_NAMESPACE
diff --git a/Engine/Plugins/Experimental/Chooser/Source/ChooserEditor/Private/OutputStructColumnEditor.h b/Engine/Plugins/Experimental/Chooser/Source/ChooserEditor/Private/OutputStructColumnEditor.h
new file mode 100644
--- /dev/null
+++ b/Engine/Plugins/Experimental/Chooser/Source/ChooserEditor/Private/OutputStructColumnEditor.h
@@ -0,0 +1,10 @@
+// Copyright Epic Games, Inc. All Rights Reserved.
+
+#pragma once
+
+#include "CoreMinimal.h"
+
+namespace UE::ChooserEditor
+{
+	void RegisterStructWidgets();
+}
\ No newline at end of file
diff --git a/Engine/Plugins/Experimental/Chooser/Source/ChooserEditor/Private/PropertyAccessChainCustomization.cpp b/Engine/Plugins/Experimental/Chooser/Source/ChooserEditor/Private/PropertyAccessChainCustomization.cpp
--- a/Engine/Plugins/Experimental/Chooser/Source/ChooserEditor/Private/PropertyAccessChainCustomization.cpp
+++ b/Engine/Plugins/Experimental/Chooser/Source/ChooserEditor/Private/PropertyAccessChainCustomization.cpp
@@ -95,6 +95,26 @@ void FPropertyAccessChainCustomization::CustomizeHeader(TSharedRef<IPropertyHand
          	 						EnumPropertyValue->Enum = ByteProperty->Enum;
          	 					}
          	 				}
+							if (TypeFilter == "object")
+							{
+								FField* Property = InBindingChain.Last().Field.ToField();
+								FChooserObjectPropertyBinding* ObjectPropertyBinding = static_cast<FChooserObjectPropertyBinding*>(PropertyValue);
+								
+								if (const FObjectPropertyBase* ObjectProperty = CastField<const FObjectPropertyBase>(Property))
+								{
+									ObjectPropertyBinding->AllowedClass = ObjectProperty->PropertyClass;
+								}
+							}
+         	 				if (TypeFilter == "struct")
+							{
+								FField* Property = InBindingChain.Last().Field.ToField();
+								FChooserStructPropertyBinding* StructPropertyBinding = static_cast<FChooserStructPropertyBinding*>(PropertyValue);
+								
+								if (const FStructProperty* StructProperty = CastField<const FStructProperty>(Property))
+								{
+									StructPropertyBinding->StructType = StructProperty->Struct;
+								}
+							}
          	
          	 				PropertyHandle->NotifyPostChange(EPropertyChangeType::ValueSet);
          	 			}
diff --git a/Engine/Plugins/Experimental/Chooser/Source/ChooserEditor/Private/RandomizeColumnEditor.cpp b/Engine/Plugins/Experimental/Chooser/Source/ChooserEditor/Private/RandomizeColumnEditor.cpp
--- a/Engine/Plugins/Experimental/Chooser/Source/ChooserEditor/Private/RandomizeColumnEditor.cpp
+++ b/Engine/Plugins/Experimental/Chooser/Source/ChooserEditor/Private/RandomizeColumnEditor.cpp
@@ -75,7 +75,7 @@ TSharedRef<SWidget> CreateRandomizeColumnWidget(UChooserTable* Chooser, FChooser
 		+ SHorizontalBox::Slot().FillWidth(1);
 }
 	
-TSharedRef<SWidget> CreateRandomizePropertyWidget(bool bReadOnly, UObject* TransactionObject, void* Value, UClass* ResultBaseClass)
+TSharedRef<SWidget> CreateRandomizePropertyWidget(bool bReadOnly, UObject* TransactionObject, void* Value, UClass* ResultBaseClass, FChooserWidgetValueChanged ValueChanged)
 {
 	IHasContextClass* HasContextClass = Cast<IHasContextClass>(TransactionObject);
 
@@ -84,11 +84,12 @@ TSharedRef<SWidget> CreateRandomizePropertyWidget(bool bReadOnly, UObject* Trans
 	return SNew(SPropertyAccessChainWidget).ContextClassOwner(HasContextClass).AllowFunctions(true).BindingColor("StructPinTypeColor").TypeFilter("FChooserRandomizationContext")
 	.PropertyBindingValue(&ContextProperty->Binding)
 	.OnAddBinding_Lambda(
-		[ContextProperty, TransactionObject](FName InPropertyName, const TArray<FBindingChainElement>& InBindingChain)
+		[ContextProperty, TransactionObject, ValueChanged](FName InPropertyName, const TArray<FBindingChainElement>& InBindingChain)
 		{
 			const FScopedTransaction Transaction(NSLOCTEXT("ContextPropertyWidget", "Change Property Binding", "Change Property Binding"));
 			TransactionObject->Modify(true);
-			ContextProperty->SetBinding(InBindingChain);	
+			ContextProperty->SetBinding(InBindingChain);
+			ValueChanged.ExecuteIfBound();
 		});
 }
 	
diff --git a/Engine/Plugins/Experimental/Chooser/Source/ChooserEditor/Private/SPropertyAccessChainWidget.cpp b/Engine/Plugins/Experimental/Chooser/Source/ChooserEditor/Private/SPropertyAccessChainWidget.cpp
--- a/Engine/Plugins/Experimental/Chooser/Source/ChooserEditor/Private/SPropertyAccessChainWidget.cpp
+++ b/Engine/Plugins/Experimental/Chooser/Source/ChooserEditor/Private/SPropertyAccessChainWidget.cpp
@@ -38,6 +38,11 @@ TSharedRef<SWidget> SPropertyAccessChainWidget::CreatePropertyAccessWidget()
 		{
 			return true;
 		}
+		if (TypeFilter == "struct")
+		{
+			// special case for struct of any type
+			return CastField<FStructProperty>(Property) != nullptr;
+		}
 		if (TypeFilter == "object")
 		{
 			// special case for objects references of any type
@@ -66,6 +71,20 @@ TSharedRef<SWidget> SPropertyAccessChainWidget::CreatePropertyAccessWidget()
 		return Property->GetCPPType() == TypeFilter;
 	};
 
+	// allow struct bindings to bind context structs directly
+	Args.OnCanBindToContextStruct = FOnCanBindToContextStruct::CreateLambda([this](UStruct* StructType)
+	{
+		if (TypeFilter == "struct" && !StructType->IsChildOf(UObject::StaticClass()))
+		{
+			// struct bindings can bind any type of struct
+			return true;
+		}
+		else
+		{
+			return StructType->GetName()  == TypeFilter;
+		}
+	});
+
 	Args.OnCanBindProperty = FOnCanBindProperty::CreateLambda(CanBindProperty);
 
 	Args.OnCanBindFunction = FOnCanBindFunction::CreateLambda([CanBindProperty](UFunction* Function)
@@ -150,9 +169,31 @@ TSharedRef<SWidget> SPropertyAccessChainWidget::CreatePropertyAccessWidget()
 				const FChooserPropertyBinding* PropertyValue = PropertyBindingValue.Get();
 	
 				int BindingChainLength = PropertyValue->PropertyBindingChain.Num();
-				if (BindingChainLength > 0)
+				if (BindingChainLength == 0)
+				{
+					if (PropertyValue->ContextIndex >= 0)
+					{
+						// direct binding to a context struct
+						if (ContextClassOwner)
+						{
+							TConstArrayView<FInstancedStruct> ContextData = ContextClassOwner->GetContextData();
+							if (ContextData.IsValidIndex(PropertyValue->ContextIndex))
+							{
+								if (const FContextObjectTypeStruct* StructType = ContextData[PropertyValue->ContextIndex].GetPtr<FContextObjectTypeStruct>())
+								{
+									CurrentValue = FText::FromString(StructType->Struct->GetAuthoredName());
+								}
+							}
+						}
+					}
+				}
+				else
 				{
-					if (BindingChainLength == 1)
+					if (!PropertyValue->DisplayName.IsEmpty())
+					{
+						CurrentValue = FText::FromString(PropertyValue->DisplayName);
+					}
+					else if (BindingChainLength == 1)
 					{
 						// single property, just use the property name
 						CurrentValue = FText::FromName(PropertyValue->PropertyBindingChain.Last());
diff --git a/Engine/Plugins/Experimental/Chooser/Source/ChooserEditor/Public/ObjectChooserWidgetFactories.h b/Engine/Plugins/Experimental/Chooser/Source/ChooserEditor/Public/ObjectChooserWidgetFactories.h
--- a/Engine/Plugins/Experimental/Chooser/Source/ChooserEditor/Public/ObjectChooserWidgetFactories.h
+++ b/Engine/Plugins/Experimental/Chooser/Source/ChooserEditor/Public/ObjectChooserWidgetFactories.h
@@ -13,15 +13,17 @@ struct FChooserColumnBase;
 
 namespace UE::ChooserEditor
 {
-	typedef TFunction<TSharedRef<SWidget>(bool bReadOnly, UObject* TransactionObject, void* Value, UClass* ResultBaseClass)> FChooserWidgetCreator;
+	DECLARE_DELEGATE(FChooserWidgetValueChanged)
+	typedef TFunction<TSharedRef<SWidget>(bool bReadOnly, UObject* TransactionObject, void* Value, UClass* ResultBaseClass, FChooserWidgetValueChanged)> FChooserWidgetCreator;
 	typedef TFunction<TSharedRef<SWidget>(UChooserTable* Chooser, FChooserColumnBase* Column, int Row)> FColumnWidgetCreator;
 
 	class CHOOSEREDITOR_API FObjectChooserWidgetFactories
 	{
 	public:
 
-		static TSharedPtr<SWidget> CreateWidget(bool ReadOnly, UObject* TransactionObject, void* Value, const UStruct* ValueType, UClass* ResultBaseClass);
-		static TSharedPtr<SWidget> CreateWidget(bool ReadOnly, UObject* TransactionObject, const UScriptStruct* BaseType, void* Value, const UStruct* ValueType, UClass* ResultBaseClass, const FOnStructPicked& CreateClassCallback, TSharedPtr<SBorder>* InnerWidget = nullptr);
+		static TSharedPtr<SWidget> CreateWidget(bool ReadOnly, UObject* TransactionObject, void* Value, const UStruct* ValueType, UClass* ResultBaseClass, FChooserWidgetValueChanged ValueChanged = FChooserWidgetValueChanged());
+		static TSharedPtr<SWidget> CreateWidget(bool ReadOnly, UObject* TransactionObject, const UScriptStruct* BaseType, void* Value, const UStruct* ValueType, UClass* ResultBaseClass,
+										const FOnStructPicked& CreateClassCallback, TSharedPtr<SBorder>* InnerWidget = nullptr, FChooserWidgetValueChanged ValueChanged = FChooserWidgetValueChanged());
 		
 		static TSharedPtr<SWidget> CreateColumnWidget(FChooserColumnBase* Column, const UStruct* ColumnType, UChooserTable* Chooser, int Row);
 		
diff --git a/Engine/Plugins/Experimental/Chooser/Source/ProxyTableEditor/Private/ProxyTableEditor.cpp b/Engine/Plugins/Experimental/Chooser/Source/ProxyTableEditor/Private/ProxyTableEditor.cpp
--- a/Engine/Plugins/Experimental/Chooser/Source/ProxyTableEditor/Private/ProxyTableEditor.cpp
+++ b/Engine/Plugins/Experimental/Chooser/Source/ProxyTableEditor/Private/ProxyTableEditor.cpp
@@ -974,7 +974,7 @@ TSharedRef<FProxyTableEditor> FProxyTableEditor::CreateEditor( const EToolkitMod
 
 /// Result widgets
 
-TSharedRef<SWidget> CreateLookupProxyWidget(bool bReadOnly, UObject* TransactionObject, void* Value, UClass* ResultBaseClass)
+TSharedRef<SWidget> CreateLookupProxyWidget(bool bReadOnly, UObject* TransactionObject, void* Value, UClass* ResultBaseClass, ChooserEditor::FChooserWidgetValueChanged ValueChanged)
 {
 	FLookupProxy* LookupProxy = static_cast<FLookupProxy*>(Value);
 	
@@ -997,11 +997,12 @@ TSharedRef<SWidget> CreateLookupProxyWidget(bool bReadOnly, UObject* Transaction
 			}
 			return true;
 		})
-		.OnObjectChanged_Lambda([TransactionObject, LookupProxy](const FAssetData& AssetData)
+		.OnObjectChanged_Lambda([TransactionObject, LookupProxy, ValueChanged](const FAssetData& AssetData)
 		{
 			const FScopedTransaction Transaction(LOCTEXT("Edit Chooser", "Edit Chooser"));
 			TransactionObject->Modify(true);
 			LookupProxy->Proxy = Cast<UProxyAsset>(AssetData.GetAsset());
+			ValueChanged.ExecuteIfBound();
 		});
 }
 
diff --git a/Engine/Plugins/Runtime/PropertyAccess/Source/PropertyAccessEditor/Private/SPropertyBinding.cpp b/Engine/Plugins/Runtime/PropertyAccess/Source/PropertyAccessEditor/Private/SPropertyBinding.cpp
--- a/Engine/Plugins/Runtime/PropertyAccess/Source/PropertyAccessEditor/Private/SPropertyBinding.cpp
+++ b/Engine/Plugins/Runtime/PropertyAccess/Source/PropertyAccessEditor/Private/SPropertyBinding.cpp
@@ -496,7 +496,7 @@ TSharedRef<SWidget> SPropertyBinding::OnGenerateDelegateMenu()
 					FUIAction(FExecuteAction::CreateSP(this, &SPropertyBinding::HandleAddBinding, BindingChain)),
 					MakeContextStructWidget(ContextStruct));
 			}
-			else if (HasBindableProperties(ContextStruct.Struct))
+			if (HasBindableProperties(ContextStruct.Struct))
 			{
 				// Show struct properties.
 				MenuBuilder.AddSubMenu(
