Evaluate Proxy node updates to match Evaluate Chooser

diff --git a/Engine/Plugins/Experimental/Chooser/Source/Chooser/Public/IObjectChooser.h b/Engine/Plugins/Experimental/Chooser/Source/Chooser/Public/IObjectChooser.h
--- a/Engine/Plugins/Experimental/Chooser/Source/Chooser/Public/IObjectChooser.h
+++ b/Engine/Plugins/Experimental/Chooser/Source/Chooser/Public/IObjectChooser.h
@@ -22,7 +22,7 @@ public:
 };
 
 #if WITH_EDITOR
-struct FChooserDebuggingInfo
+struct CHOOSER_API FChooserDebuggingInfo
 {
 	bool bCurrentDebugTarget = false;
 };
@@ -30,14 +30,14 @@ struct FChooserDebuggingInfo
 
 
 USTRUCT()
-struct FChooserEvaluationInputObject
+struct CHOOSER_API FChooserEvaluationInputObject
 {
 	GENERATED_BODY()
 	TObjectPtr<UObject> Object;
 };
 
 USTRUCT(BlueprintType)
-struct FChooserEvaluationContext
+struct CHOOSER_API FChooserEvaluationContext
 {
 	#if WITH_EDITOR
     	FChooserDebuggingInfo DebuggingInfo;
diff --git a/Engine/Plugins/Experimental/Chooser/Source/ProxyTable/Private/ProxyAsset.cpp b/Engine/Plugins/Experimental/Chooser/Source/ProxyTable/Private/ProxyAsset.cpp
--- a/Engine/Plugins/Experimental/Chooser/Source/ProxyTable/Private/ProxyAsset.cpp
+++ b/Engine/Plugins/Experimental/Chooser/Source/ProxyTable/Private/ProxyAsset.cpp
@@ -34,7 +34,7 @@ void UProxyAsset::PostEditChangeProperty(FPropertyChangedEvent& PropertyChangedE
 		}
 		CachedPreviousType = Type;
 	}
-	else if (PropertyChangedEvent.Property->GetName() == ContextClassName)
+	else
 	{
 		OnContextClassChanged.Broadcast();
 	}
diff --git a/Engine/Plugins/Experimental/Chooser/Source/ProxyTable/Private/ProxyTableFunctionLibrary.cpp b/Engine/Plugins/Experimental/Chooser/Source/ProxyTable/Private/ProxyTableFunctionLibrary.cpp
--- a/Engine/Plugins/Experimental/Chooser/Source/ProxyTable/Private/ProxyTableFunctionLibrary.cpp
+++ b/Engine/Plugins/Experimental/Chooser/Source/ProxyTable/Private/ProxyTableFunctionLibrary.cpp
@@ -1,5 +1,7 @@
 // Copyright Epic Games, Inc. All Rights Reserved.
 #include "ProxyTableFunctionLibrary.h"
+
+#include "LookupProxy.h"
 #include "ProxyTable.h"
 
 /////////////////////////////////////////////////////////////////////////////////////
@@ -44,3 +46,12 @@ UObject* UProxyTableFunctionLibrary::EvaluateProxyTable(const UObject* ContextOb
 	
 	return nullptr;
 }
+
+
+FInstancedStruct UProxyTableFunctionLibrary::MakeLookupProxy(UProxyAsset* Proxy)
+ {
+ 	FInstancedStruct Struct;
+ 	Struct.InitializeAs(FLookupProxy::StaticStruct());
+ 	Struct.GetMutable<FLookupProxy>().Proxy = Proxy;
+ 	return Struct;
+ }
diff --git a/Engine/Plugins/Experimental/Chooser/Source/ProxyTable/Public/ProxyTableFunctionLibrary.h b/Engine/Plugins/Experimental/Chooser/Source/ProxyTable/Public/ProxyTableFunctionLibrary.h
--- a/Engine/Plugins/Experimental/Chooser/Source/ProxyTable/Public/ProxyTableFunctionLibrary.h
+++ b/Engine/Plugins/Experimental/Chooser/Source/ProxyTable/Public/ProxyTableFunctionLibrary.h
@@ -23,7 +23,7 @@ public:
 		* @param ProxyTable				(in) The ProxyTable asset
 		* @param Key					(in) The Key from the ProxyTable asset
 	*/
-	UFUNCTION(BlueprintPure, meta = (BlueprintThreadSafe), Category = "Animation")
+	UFUNCTION(BlueprintPure, meta = (BlueprintThreadSafe, BlueprintInternalUseOnly="true"), Category = "Animation")
 	static UObject* EvaluateProxyTable(const UObject* ContextObject, const UProxyTable* ProxyTable, FName Key);
 	
 	/**
@@ -33,7 +33,14 @@ public:
 	* @param ProxyAsset				(in) The ProxyAsset asset
 	* @param ObjectClass			(in) Expected type of result objects
 	*/
-	UFUNCTION(BlueprintPure, meta = (BlueprintThreadSafe, DeterminesOutputType = "ObjectClass"), Category = "Animation")
+	UFUNCTION(BlueprintPure, meta = (BlueprintThreadSafe, DeterminesOutputType = "ObjectClass", BlueprintInternalUseOnly="true"), Category = "Animation")
 	static UObject* EvaluateProxyAsset(const UObject* ContextObject, const UProxyAsset* Proxy, TSubclassOf<UObject> ObjectClass);
-
+	
+	/**
+	* Create a LookupProxy struct
+	*
+	* @param Chooser				(in) the ChooserTable asset to evaluate
+	*/
+	UFUNCTION(BlueprintPure, Category = "Animation", meta = (BlueprintThreadSafe, BlueprintInternalUseOnly="true", NativeMakeFunc))
+	static FInstancedStruct MakeLookupProxy(UProxyAsset* Proxy);
 };
\ No newline at end of file
diff --git a/Engine/Plugins/Experimental/Chooser/Source/ProxyTableUncooked/Private/EvaluateProxyNode.cpp b/Engine/Plugins/Experimental/Chooser/Source/ProxyTableUncooked/Private/EvaluateProxyNode.cpp
--- a/Engine/Plugins/Experimental/Chooser/Source/ProxyTableUncooked/Private/EvaluateProxyNode.cpp
+++ b/Engine/Plugins/Experimental/Chooser/Source/ProxyTableUncooked/Private/EvaluateProxyNode.cpp
@@ -5,6 +5,7 @@
 
 #include "BlueprintActionDatabaseRegistrar.h"
 #include "BlueprintNodeSpawner.h"
+#include "ChooserFunctionLibrary.h"
 #include "ProxyTableFunctionLibrary.h"
 #include "Containers/UnrealString.h"
 #include "EdGraph/EdGraph.h"
@@ -15,6 +16,7 @@
 #include "Internationalization/Internationalization.h"
 #include "K2Node_CallFunction.h"
 #include "K2Node_MakeArray.h"
+#include "K2Node_MakeStruct.h"
 #include "Kismet2/BlueprintEditorUtils.h"
 #include "Kismet2/CompilerResultsLog.h"
 #include "KismetCompiler.h"
@@ -127,14 +129,7 @@ void UK2Node_EvaluateProxy::AllocateDefaultPins()
 
 FText UK2Node_EvaluateProxy::GetNodeTitle(ENodeTitleType::Type TitleType) const
 {
-	if (Proxy)
-	{
-		return FText::FromString(Proxy->GetName());
-	}
-	else
-	{
-		return LOCTEXT("EvaluateProxy_Title", "Evaluate Proxy");
-	}
+	return LOCTEXT("EvaluateProxy_Title", "Evaluate Proxy (Deprecated)");
 }
 
 FText UK2Node_EvaluateProxy::GetPinDisplayName(const UEdGraphPin* Pin) const
@@ -300,4 +295,536 @@ FText UK2Node_EvaluateProxy::GetMenuCategory() const
 }
 
 
+
+
+//----------------------------------------------------------------------------------------------
+// UK2Node_EvaluateProxy2
+// New Implementation of EvaluateChooser with support for passing in/out multiple objects and structs
+
+UK2Node_EvaluateProxy2::UK2Node_EvaluateProxy2(const FObjectInitializer& ObjectInitializer)
+	: Super(ObjectInitializer)
+{
+	NodeTooltip = LOCTEXT("EvaluateChooser2Tooltip", "Evaluates an Chooser Table, and returns the resulting Object or Objects.");
+}
+
+void UK2Node_EvaluateProxy2::UnregisterProxyCallback()
+{
+	if(CurrentCallbackProxy)
+	{
+		CurrentCallbackProxy->OnTypeChanged.RemoveAll(this);
+		CurrentCallbackProxy->OnContextClassChanged.RemoveAll(this);
+		CurrentCallbackProxy = nullptr;
+	}
+}
+
+void UK2Node_EvaluateProxy2::BeginDestroy()
+{
+	UnregisterProxyCallback();
+	
+	Super::BeginDestroy();
+}
+
+void UK2Node_EvaluateProxy2::PostEditUndo()
+{
+	Super::PostEditUndo();
+	ProxyChanged();
+}
+
+void UK2Node_EvaluateProxy2::DestroyNode()
+{
+	UnregisterProxyCallback();
+	Super::DestroyNode();
+}
+
+
+void UK2Node_EvaluateProxy2::ProxyChanged()
+{
+	if (Proxy != CurrentCallbackProxy)
+	{
+		UnregisterProxyCallback();
+	
+		if (Proxy)
+		{
+			Proxy->OnTypeChanged.AddUObject(this, &UK2Node_EvaluateProxy2::ResultTypeChanged);
+			Proxy->OnContextClassChanged.AddUObject(this, &UK2Node::ReconstructNode);
+		}
+	
+		CurrentCallbackProxy = Proxy;
+
+		ReconstructNode();
+	}
+}
+
+void UK2Node_EvaluateProxy2::ResultTypeChanged(const UClass*)
+{
+	AllocateDefaultPins();
+}
+
+void UK2Node_EvaluateProxy2::PreloadRequiredAssets()
+{
+	if (Proxy)
+	{
+		if (FLinkerLoad* ObjLinker = GetLinker())
+		{
+			ObjLinker->Preload(Proxy);
+		}
+	}
+    		
+	Super::PreloadRequiredAssets();
+}
+
+void UK2Node_EvaluateProxy2::AllocateDefaultPins()
+{
+	Super::AllocateDefaultPins();
+	
+	UClass* ResultType = UObject::StaticClass();
+
+	if (!FindPin(UEdGraphSchema_K2::PN_Execute, EGPD_Input))
+	{
+		// Input - Execution Pin
+		CreatePin(EGPD_Input, UEdGraphSchema_K2::PC_Exec, UEdGraphSchema_K2::PN_Execute);
+	}
+
+	if (!FindPin(UEdGraphSchema_K2::PN_Execute, EGPD_Output))
+	{
+		// Output - Execution Pin
+		CreatePin(EGPD_Output, UEdGraphSchema_K2::PC_Exec, UEdGraphSchema_K2::PN_Execute);
+	}
+
+	if (Proxy)
+	{
+		for(FInstancedStruct& ContextDataEntry : Proxy->ContextData)
+		{
+			if (ContextDataEntry.IsValid())
+			{
+				const UScriptStruct* EntryType = ContextDataEntry.GetScriptStruct();
+				if (EntryType == FContextObjectTypeClass::StaticStruct())
+				{
+					const FContextObjectTypeClass& ClassContext = ContextDataEntry.Get<FContextObjectTypeClass>();
+					if (UEdGraphPin* Pin = FindPin(ClassContext.Class.GetFName(), EGPD_Input))
+					{
+						Pin->PinType.PinSubCategoryObject = ClassContext.Class;
+					}
+					else
+					{
+						CreatePin(EGPD_Input, UEdGraphSchema_K2::PC_Object, ClassContext.Class, ClassContext.Class.GetFName());
+					}
+				}
+				else if (EntryType == FContextObjectTypeStruct::StaticStruct())
+				{
+					const FContextObjectTypeStruct& StructContext = ContextDataEntry.Get<FContextObjectTypeStruct>();
+
+					if (StructContext.Direction == EContextObjectDirection::Read || StructContext.Direction == EContextObjectDirection::ReadWrite)
+					{
+						if (UEdGraphPin* Pin = FindPin(StructContext.Struct.GetFName(), EGPD_Input))
+						{
+							Pin->PinType.PinSubCategoryObject = StructContext.Struct;
+						}
+						else
+						{
+							CreatePin(EGPD_Input, UEdGraphSchema_K2::PC_Struct, StructContext.Struct, StructContext.Struct.GetFName());
+						}
+					}
+
+					if (StructContext.Direction == EContextObjectDirection::Write || StructContext.Direction == EContextObjectDirection::ReadWrite)
+					{
+						if (UEdGraphPin* Pin = FindPin(StructContext.Struct.GetFName(), EGPD_Output))
+						{
+							Pin->PinType.PinSubCategoryObject = StructContext.Struct;
+						}
+						else
+						{
+							CreatePin(EGPD_Output, UEdGraphSchema_K2::PC_Struct, StructContext.Struct, StructContext.Struct.GetFName());
+						}
+					}
+
+					if (StructContext.Direction == EContextObjectDirection::Read)
+					{
+					    // for read only structs, remove the output pin if found
+						if (UEdGraphPin* Pin = FindPin(StructContext.Struct.GetFName(), EGPD_Output))
+						{
+							RemovePin(Pin);
+						}
+					}
+					
+					if (StructContext.Direction == EContextObjectDirection::Write)
+					{
+					    // for write only structs, remove the input pin if found
+						if (UEdGraphPin* Pin = FindPin(StructContext.Struct.GetFName(), EGPD_Input))
+						{
+							RemovePin(Pin);
+						}
+					}
+				}
+			}
+		}
+	}
+
+   	if (Proxy && Proxy->Type)
+   	{
+   		ResultType = Proxy->Type;
+   	}
+
+	static const FName ResultPinName = "Result";
+	if (UEdGraphPin* ResultPin = FindPin(ResultPinName, EGPD_Output))
+	{
+		ResultPin->PinType.PinSubCategoryObject = ResultType;
+		ResultPin->PinType.ContainerType = (Mode == EEvaluateProxyMode::AllResults) ? EPinContainerType::Array : EPinContainerType::None;
+	}
+	else
+	{
+		UEdGraphNode::FCreatePinParams PinParams;
+		PinParams.ContainerType = (Mode == EEvaluateProxyMode::AllResults) ? EPinContainerType::Array : EPinContainerType::None;
+		PinParams.ValueTerminalType.TerminalCategory = UEdGraphSchema_K2::PC_Object;
+				
+		CreatePin(EGPD_Output, UEdGraphSchema_K2::PC_Object, ResultType, TEXT("Result"), PinParams);
+	}
+}
+
+FText UK2Node_EvaluateProxy2::GetNodeTitle(ENodeTitleType::Type TitleType) const
+{
+	if (Proxy)
+	{
+		return FText::Format(LOCTEXT("EvaluateProxy2_TitleWithProxy", "Evaluate Proxy: {0}"), {FText::FromString(Proxy->GetName())});
+	}
+	else
+	{
+		return LOCTEXT("EvaluateProxy2_Title", "Evaluate Proxy");
+	}
+}
+
+FText UK2Node_EvaluateProxy2::GetPinDisplayName(const UEdGraphPin* Pin) const
+{
+	if (Pin->PinName == UEdGraphSchema_K2::PN_Execute)
+	{
+		return FText();
+	}
+	return FText::FromName(Pin->PinName);
+}
+
+void UK2Node_EvaluateProxy2::PostEditChangeProperty(struct FPropertyChangedEvent& PropertyChangedEvent)
+{
+	if (PropertyChangedEvent.Property->GetName() == "Proxy")
+	{
+		ProxyChanged();
+	}
+	if (PropertyChangedEvent.Property->GetName() == "Mode")
+	{
+		ReconstructNode();
+	}
+	Super::PostEditChangeProperty(PropertyChangedEvent);
+}
+
+void UK2Node_EvaluateProxy2::PostLoad()
+{
+	Super::PostLoad();
+	
+	if (Proxy)
+	{
+		Proxy->OnTypeChanged.AddUObject(this, &UK2Node_EvaluateProxy2::ResultTypeChanged);
+		Proxy->OnContextClassChanged.AddUObject(this, &UK2Node::ReconstructNode);
+	}
+
+	CurrentCallbackProxy = Proxy;
+}
+
+void UK2Node_EvaluateProxy2::PinConnectionListChanged(UEdGraphPin* Pin)
+{
+
+	Modify();
+
+	FBlueprintEditorUtils::MarkBlueprintAsStructurallyModified(GetBlueprint());
+}
+
+void UK2Node_EvaluateProxy2::PinDefaultValueChanged(UEdGraphPin* Pin)
+{
+	Super::PinDefaultValueChanged(Pin);
+}
+
+void UK2Node_EvaluateProxy2::PinTypeChanged(UEdGraphPin* Pin)
+{
+	Super::PinTypeChanged(Pin);
+}
+
+FText UK2Node_EvaluateProxy2::GetTooltipText() const
+{
+	return NodeTooltip;
+}
+
+void UK2Node_EvaluateProxy2::PostReconstructNode()
+{
+	Super::PostReconstructNode();
+}
+
+void UK2Node_EvaluateProxy2::ExpandNode(class FKismetCompilerContext& CompilerContext, UEdGraph* SourceGraph)
+{
+	Super::ExpandNode(CompilerContext, SourceGraph);
+
+	UEdGraphPin* ExecInput = GetExecPin();
+	UEdGraphPin* ExecOutput = GetPassThroughPin(ExecInput);
+	UEdGraphPin* ResultPin = FindPinChecked(TEXT("Result"));
+
+	UEdGraphPin* PreviousNodeExecOutput = nullptr;
+	
+	if (ExecInput->HasAnyConnections())
+	{
+		UK2Node_CallFunction* CallFunction = CompilerContext.SpawnIntermediateNode<UK2Node_CallFunction>(this, SourceGraph);
+		CompilerContext.MessageLog.NotifyIntermediateObjectCreation(CallFunction, this);
+
+		if (Mode == EEvaluateProxyMode::AllResults)
+		{
+			CallFunction->SetFromFunction(UChooserFunctionLibrary::StaticClass()->FindFunctionByName(GET_MEMBER_NAME_CHECKED(UChooserFunctionLibrary, EvaluateObjectChooserBaseMulti)));
+		}
+		else
+		{
+			CallFunction->SetFromFunction(UChooserFunctionLibrary::StaticClass()->FindFunctionByName(GET_MEMBER_NAME_CHECKED(UChooserFunctionLibrary, EvaluateObjectChooserBase)));
+		}
+		CallFunction->AllocateDefaultPins();
+		
+		UK2Node_MakeStruct* ContextStructNode = CompilerContext.SpawnIntermediateNode<UK2Node_MakeStruct>(this, SourceGraph);
+		ContextStructNode->PostPlacedNewNode();
+        CompilerContext.MessageLog.NotifyIntermediateObjectCreation(ContextStructNode, this);
+		ContextStructNode->StructType = FChooserEvaluationContext::StaticStruct();
+        ContextStructNode->AllocateDefaultPins();
+		UEdGraphPin* ContextStructPin = ContextStructNode->FindPin(FChooserEvaluationContext::StaticStruct()->GetFName(), EGPD_Output);
+
+		ContextStructPin->MakeLinkTo(CallFunction->FindPin(FName("Context")));
+
+		UK2Node_Self* SelfNode = CompilerContext.SpawnIntermediateNode<UK2Node_Self>(this, SourceGraph);
+		CompilerContext.MessageLog.NotifyIntermediateObjectCreation(SelfNode, this);
+		SelfNode->AllocateDefaultPins();
+		UEdGraphPin* SelfPin = SelfNode->FindPin(TEXT("self"));
+		
+		if (Proxy)
+		{
+			int ContextDataCount = Proxy->ContextData.Num();
+			for(int ContextDataIndex = 0; ContextDataIndex < ContextDataCount; ContextDataIndex++)
+			{
+				FInstancedStruct& ContextDataEntry = Proxy->ContextData[ContextDataIndex];
+				if (ContextDataEntry.IsValid())
+				{
+					const UScriptStruct* EntryType = ContextDataEntry.GetScriptStruct();
+					if (EntryType == FContextObjectTypeClass::StaticStruct())
+					{
+						const FContextObjectTypeClass& ClassContext = ContextDataEntry.Get<FContextObjectTypeClass>();
+						if (UEdGraphPin* Pin = FindPin(ClassContext.Class.GetFName(), EGPD_Input))
+						{
+							UK2Node_CallFunction* AddObjectFunction = CompilerContext.SpawnIntermediateNode<UK2Node_CallFunction>(this, SourceGraph);
+                       		CompilerContext.MessageLog.NotifyIntermediateObjectCreation(AddObjectFunction, this);
+							AddObjectFunction->SetFromFunction(UChooserFunctionLibrary::StaticClass()->FindFunctionByName(GET_MEMBER_NAME_CHECKED(UChooserFunctionLibrary, AddChooserObjectInput)));
+							AddObjectFunction->AllocateDefaultPins();
+							ContextStructPin->MakeLinkTo(AddObjectFunction->FindPin(FName("Context")));
+
+							if (PreviousNodeExecOutput == nullptr)
+							{
+								CompilerContext.MovePinLinksToIntermediate(*ExecInput, *AddObjectFunction->GetExecPin());
+							}
+							else
+							{
+								PreviousNodeExecOutput->MakeLinkTo(AddObjectFunction->GetExecPin());
+							}
+							PreviousNodeExecOutput = AddObjectFunction->GetThenPin();
+							
+							UEdGraphPin* AddObjectPin = AddObjectFunction->FindPin(FName("Object"));
+							
+							if (Pin->HasAnyConnections())
+							{
+								CompilerContext.MovePinLinksToIntermediate(*Pin, *AddObjectPin);
+							}
+							else // would be nice to check that self is the same type as this Object parameter
+							{
+								// auto connect self node to any disconnected object pin
+								SelfPin->MakeLinkTo(AddObjectPin);
+							}
+						}
+					}
+					else if (EntryType == FContextObjectTypeStruct::StaticStruct())
+					{
+						const FContextObjectTypeStruct& StructContext = ContextDataEntry.Get<FContextObjectTypeStruct>();
+						
+						UK2Node_CallFunction* AddStructFunction = CompilerContext.SpawnIntermediateNode<UK2Node_CallFunction>(this, SourceGraph);
+						CompilerContext.MessageLog.NotifyIntermediateObjectCreation(AddStructFunction, this);
+						AddStructFunction->SetFromFunction(UChooserFunctionLibrary::StaticClass()->FindFunctionByName(GET_MEMBER_NAME_CHECKED(UChooserFunctionLibrary, AddChooserStructInput)));
+						AddStructFunction->AllocateDefaultPins();
+						ContextStructPin->MakeLinkTo(AddStructFunction->FindPin(FName("Context")));
+
+						if (PreviousNodeExecOutput == nullptr)
+						{
+							CompilerContext.MovePinLinksToIntermediate(*ExecInput, *AddStructFunction->GetExecPin());
+						}
+						else
+						{
+							PreviousNodeExecOutput->MakeLinkTo(AddStructFunction->GetExecPin());
+						}
+						PreviousNodeExecOutput = AddStructFunction->GetThenPin();
+							
+						UEdGraphPin* AddStructPin = AddStructFunction->FindPin(FName("Value"));
+						// not sure why this isn't automatically happening:
+						AddStructPin->PinType.PinCategory = UEdGraphSchema_K2::PC_Struct;
+						AddStructPin->PinType.PinSubCategoryObject = StructContext.Struct;
+
+						bool bLinked = false;
+						if (StructContext.Direction != EContextObjectDirection::Write)
+						{
+							if (UEdGraphPin* Pin = FindPin(StructContext.Struct.GetFName(), EGPD_Input))
+							{
+								if (Pin->HasAnyConnections())
+								{
+									CompilerContext.MovePinLinksToIntermediate(*Pin, *AddStructPin);
+									bLinked = true;
+								}
+							}
+						}
+						
+						if (!bLinked)
+						{
+							// create a struct to hold the output (or for input structs that were not connected to anything)
+							UK2Node_MakeStruct* OutputStructNode = CompilerContext.SpawnIntermediateNode<UK2Node_MakeStruct>(this, SourceGraph);
+							CompilerContext.MessageLog.NotifyIntermediateObjectCreation(OutputStructNode, this);
+							OutputStructNode->StructType = static_cast<UScriptStruct*>(StructContext.Struct);
+							OutputStructNode->AllocateDefaultPins();
+							OutputStructNode->FindPin(StructContext.Struct->GetFName(), EGPD_Output)->MakeLinkTo(AddStructPin);
+						}
+
+						if (StructContext.Direction != EContextObjectDirection::Read)
+						{
+							
+							if (UEdGraphPin* Pin = FindPin(StructContext.Struct.GetFName(), EGPD_Output))
+							{
+								if (Pin->HasAnyConnections())
+								{
+									// set up struct output pin
+									UK2Node_CallFunction* GetStructFunction = CompilerContext.SpawnIntermediateNode<UK2Node_CallFunction>(this, SourceGraph);
+									CompilerContext.MessageLog.NotifyIntermediateObjectCreation(GetStructFunction, this);
+									GetStructFunction->SetFromFunction(UChooserFunctionLibrary::StaticClass()->FindFunctionByName(GET_MEMBER_NAME_CHECKED(UChooserFunctionLibrary, GetChooserStructOutput)));
+									GetStructFunction->AllocateDefaultPins();
+									GetStructFunction->FindPin(FName("Index"))->DefaultValue = FString::FromInt(ContextDataIndex);
+									
+									ContextStructPin->MakeLinkTo(GetStructFunction->FindPin(FName("Context")));
+									UEdGraphPin* ValuePin = GetStructFunction->FindPin(FName("Value"));
+									// not sure why this isn't automatically happening:
+									ValuePin->PinType.PinCategory = UEdGraphSchema_K2::PC_Struct;
+									ValuePin->PinType.PinSubCategoryObject = StructContext.Struct;
+					
+									CompilerContext.MovePinLinksToIntermediate(*Pin, *ValuePin);
+								}
+							}
+						}
+					}
+				}
+			}
+
+			if (PreviousNodeExecOutput)
+			{
+				PreviousNodeExecOutput->MakeLinkTo(CallFunction->GetExecPin());
+			}
+			else
+			{
+				CompilerContext.MovePinLinksToIntermediate(*ExecInput, *CallFunction->GetExecPin());
+			}
+			
+			CompilerContext.MovePinLinksToIntermediate(*ExecOutput, *CallFunction->GetThenPin());
+		}
+
+		UK2Node_CallFunction* MakeLookupProxyFunction = CompilerContext.SpawnIntermediateNode<UK2Node_CallFunction>(this, SourceGraph);
+		CompilerContext.MessageLog.NotifyIntermediateObjectCreation(MakeLookupProxyFunction, this);
+		MakeLookupProxyFunction->SetFromFunction(UProxyTableFunctionLibrary::StaticClass()->FindFunctionByName(GET_MEMBER_NAME_CHECKED(UProxyTableFunctionLibrary, MakeLookupProxy)));
+		MakeLookupProxyFunction->AllocateDefaultPins();
+		
+		UEdGraphPin* ChooserTablePin = MakeLookupProxyFunction->FindPin(FName("Proxy"));
+		CallFunction->GetSchema()->TrySetDefaultObject(*ChooserTablePin, Proxy);
+
+		MakeLookupProxyFunction->GetReturnValuePin()->MakeLinkTo(CallFunction->FindPin(FName("ObjectChooser")));
+
+		UEdGraphPin* OutputPin = CallFunction->GetReturnValuePin();
+
+		if (Proxy && Proxy->Type)
+		{
+			UEdGraphPin* OutputClassPin = CallFunction->FindPin(TEXT("ObjectClass"));
+			CallFunction->GetSchema()->TrySetDefaultObject(*OutputClassPin, Proxy->Type);
+		}
+
+		CompilerContext.MovePinLinksToIntermediate(*ResultPin, *OutputPin);
+	}
+	
+	BreakAllNodeLinks();
+}
+
+UK2Node::ERedirectType UK2Node_EvaluateProxy2::DoPinsMatchForReconstruction(const UEdGraphPin* NewPin, int32 NewPinIndex, const UEdGraphPin* OldPin, int32 OldPinIndex) const
+{
+	ERedirectType RedirectType = ERedirectType_None;
+
+	// if the pin names do match
+	if (NewPin->PinName.ToString().Equals(OldPin->PinName.ToString(), ESearchCase::CaseSensitive))
+	{
+		// Make sure we're not dealing with a menu node
+		UEdGraph* OuterGraph = GetGraph();
+		if( OuterGraph && OuterGraph->Schema )
+		{
+			const UEdGraphSchema_K2* K2Schema = Cast<const UEdGraphSchema_K2>(GetSchema());
+			if( !K2Schema || K2Schema->IsSelfPin(*NewPin) || K2Schema->ArePinTypesCompatible(OldPin->PinType, NewPin->PinType) )
+			{
+				RedirectType = ERedirectType_Name;
+			}
+			else
+			{
+				RedirectType = ERedirectType_None;
+			}
+		}
+	}
+	else
+	{
+		// try looking for a redirect if it's a K2 node
+		if (UK2Node* Node = Cast<UK2Node>(NewPin->GetOwningNode()))
+		{	
+			// if you don't have matching pin, now check if there is any redirect param set
+			TArray<FString> OldPinNames;
+			GetRedirectPinNames(*OldPin, OldPinNames);
+
+			FName NewPinName;
+			RedirectType = ShouldRedirectParam(OldPinNames, /*out*/ NewPinName, Node);
+
+			// make sure they match
+			if ((RedirectType != ERedirectType_None) && (!NewPin->PinName.ToString().Equals(NewPinName.ToString(), ESearchCase::CaseSensitive)))
+			{
+				RedirectType = ERedirectType_None;
+			}
+		}
+	}
+
+	return RedirectType;
+}
+
+bool UK2Node_EvaluateProxy2::IsConnectionDisallowed(const UEdGraphPin* MyPin, const UEdGraphPin* OtherPin, FString& OutReason) const
+{
+	return Super::IsConnectionDisallowed(MyPin, OtherPin, OutReason);
+}
+
+void UK2Node_EvaluateProxy2::GetMenuActions(FBlueprintActionDatabaseRegistrar& ActionRegistrar) const
+{
+	// actions get registered under specific object-keys; the idea is that 
+	// actions might have to be updated (or deleted) if their object-key is  
+	// mutated (or removed)... here we use the node's class (so if the node 
+	// type disappears, then the action should go with it)
+	UClass* ActionKey = GetClass();
+	// to keep from needlessly instantiating a UBlueprintNodeSpawner, first   
+	// check to make sure that the registrar is looking for actions of this type
+	// (could be regenerating actions for a specific asset, and therefore the 
+	// registrar would only accept actions corresponding to that asset)
+	if (ActionRegistrar.IsOpenForRegistration(ActionKey))
+	{
+		UBlueprintNodeSpawner* NodeSpawner = UBlueprintNodeSpawner::Create(GetClass());
+		check(NodeSpawner != nullptr);
+
+		ActionRegistrar.AddBlueprintAction(ActionKey, NodeSpawner);
+	}
+}
+
+FText UK2Node_EvaluateProxy2::GetMenuCategory() const
+{
+	return FEditorCategoryUtils::GetCommonCategory(FCommonEditorCategory::Animation);
+}
+
+
+
 #undef LOCTEXT_NAMESPACE
diff --git a/Engine/Plugins/Experimental/Chooser/Source/ProxyTableUncooked/Private/EvaluateProxyNode.h b/Engine/Plugins/Experimental/Chooser/Source/ProxyTableUncooked/Private/EvaluateProxyNode.h
--- a/Engine/Plugins/Experimental/Chooser/Source/ProxyTableUncooked/Private/EvaluateProxyNode.h
+++ b/Engine/Plugins/Experimental/Chooser/Source/ProxyTableUncooked/Private/EvaluateProxyNode.h
@@ -14,7 +14,7 @@
 
 #include "EvaluateProxyNode.generated.h"
 
-UCLASS(MinimalAPI)
+UCLASS(MinimalAPI, Hidden)
 class UK2Node_EvaluateProxy : public UK2Node
 {
 	GENERATED_UCLASS_BODY()
@@ -67,3 +67,67 @@ private:
 	FText NodeTooltip;
 };
 
+UENUM()
+enum class EEvaluateProxyMode
+{
+	FirstResult,
+	AllResults
+};
+
+//----------------------------------------------------------------------------------------------
+// UK2Node_EvaluateChooser2
+// New Implementation of EvaluateChooser with support for passing in/out multiple objects and structs
+UCLASS(MinimalAPI)
+class UK2Node_EvaluateProxy2 : public UK2Node
+{
+	GENERATED_UCLASS_BODY()
+	
+	//~ Begin UObject Interface
+	virtual void BeginDestroy() override;
+	virtual void PostEditUndo() override;
+	virtual void PostEditChangeProperty(struct FPropertyChangedEvent& PropertyChangedEvent) override;
+	virtual void PostLoad() override;
+	//~ End UObject Interface
+
+	//~ Begin UEdGraphNode Interface.
+	virtual void AllocateDefaultPins() override;
+	virtual FText GetNodeTitle(ENodeTitleType::Type TitleType) const override;
+	virtual bool ShouldShowNodeProperties() const override { return true; }
+	virtual void PinConnectionListChanged(UEdGraphPin* Pin) override;
+	virtual void PinDefaultValueChanged(UEdGraphPin* Pin) override;
+	virtual void PinTypeChanged(UEdGraphPin* Pin) override;
+	virtual FText GetTooltipText() const override;
+	virtual FText GetPinDisplayName(const UEdGraphPin* Pin) const override;
+	//~ End UEdGraphNode Interface.
+
+	//~ Begin UK2Node Interface.
+	virtual bool IsNodePure() const override { return true; }
+	virtual bool NodeCausesStructuralBlueprintChange() const override { return true; }
+	virtual void PostReconstructNode() override;
+	virtual void ExpandNode(class FKismetCompilerContext& CompilerContext, UEdGraph* SourceGraph) override;
+	virtual ERedirectType DoPinsMatchForReconstruction(const UEdGraphPin* NewPin, int32 NewPinIndex, const UEdGraphPin* OldPin, int32 OldPinIndex) const override;
+	virtual bool IsConnectionDisallowed(const UEdGraphPin* MyPin, const UEdGraphPin* OtherPin, FString& OutReason) const override;
+	virtual void GetMenuActions(FBlueprintActionDatabaseRegistrar& ActionRegistrar) const override;
+	virtual FText GetMenuCategory() const override;
+	virtual int32 GetNodeRefreshPriority() const override { return EBaseNodeRefreshPriority::Low_UsesDependentWildcard; }
+	virtual void PreloadRequiredAssets() override;
+
+	virtual void DestroyNode() override;
+    
+	//~ End UK2Node Interface.
+
+private:
+	void ProxyChanged();
+	void ResultTypeChanged(const UClass*);
+	void UnregisterProxyCallback();
+	UProxyAsset* CurrentCallbackProxy = nullptr;
+
+	UPROPERTY(EditAnywhere, Category = "Proxy")
+	TObjectPtr<UProxyAsset> Proxy;
+
+	UPROPERTY(EditAnywhere, Category = "Chooser")
+	EEvaluateProxyMode Mode;
+
+	/** Tooltip text for this node. */
+	FText NodeTooltip;
+};
diff --git a/Engine/Plugins/Experimental/Chooser/Source/ProxyTableUncooked/ProxyTableUncooked.Build.cs b/Engine/Plugins/Experimental/Chooser/Source/ProxyTableUncooked/ProxyTableUncooked.Build.cs
--- a/Engine/Plugins/Experimental/Chooser/Source/ProxyTableUncooked/ProxyTableUncooked.Build.cs
+++ b/Engine/Plugins/Experimental/Chooser/Source/ProxyTableUncooked/ProxyTableUncooked.Build.cs
@@ -12,6 +12,7 @@ namespace UnrealBuildTool.Rules
 				{
 					"Core",
 					"CoreUObject",
+					"Chooser",
 					"Engine",
 					"AssetTools",
 					"UnrealEd",
