Proxy Table struct output feature

diff --git a/Engine/Plugins/Experimental/Chooser/Source/Chooser/Private/OutputStructColumn.cpp b/Engine/Plugins/Experimental/Chooser/Source/Chooser/Private/OutputStructColumn.cpp
--- a/Engine/Plugins/Experimental/Chooser/Source/Chooser/Private/OutputStructColumn.cpp
+++ b/Engine/Plugins/Experimental/Chooser/Source/Chooser/Private/OutputStructColumn.cpp
@@ -77,8 +77,11 @@ void FStructContextProperty::SetBinding(const UObject* OuterObject, const TArray
 			if (InBindingChain.Num() > 2)
 			{
 				// add the parent struct name to the display name if there is one
-				const FField* NextToLastField = InBindingChain[InBindingChain.Num()-2].Field.ToField();
-				Binding.DisplayName = NextToLastField->GetAuthoredName() + "." + Binding.DisplayName;
+				
+				if (const FField* NextToLastField = InBindingChain[InBindingChain.Num()-2].Field.ToField())
+				{
+					Binding.DisplayName = NextToLastField->GetAuthoredName() + "." + Binding.DisplayName;
+				}
 			}
 		}
 	}
diff --git a/Engine/Plugins/Experimental/Chooser/Source/ChooserEditor/Public/SPropertyAccessChainWidget.h b/Engine/Plugins/Experimental/Chooser/Source/ChooserEditor/Public/SPropertyAccessChainWidget.h
--- a/Engine/Plugins/Experimental/Chooser/Source/ChooserEditor/Public/SPropertyAccessChainWidget.h
+++ b/Engine/Plugins/Experimental/Chooser/Source/ChooserEditor/Public/SPropertyAccessChainWidget.h
@@ -13,7 +13,7 @@ namespace UE::ChooserEditor
 {
 		
 // Wrapper widget for Property access widget, which can update when the target Class changes
-class SPropertyAccessChainWidget : public SCompoundWidget
+class CHOOSEREDITOR_API SPropertyAccessChainWidget : public SCompoundWidget
 {
 public:
 	SLATE_BEGIN_ARGS(SPropertyAccessChainWidget)
diff --git a/Engine/Plugins/Experimental/Chooser/Source/ProxyTable/Private/ProxyAsset.cpp b/Engine/Plugins/Experimental/Chooser/Source/ProxyTable/Private/ProxyAsset.cpp
--- a/Engine/Plugins/Experimental/Chooser/Source/ProxyTable/Private/ProxyAsset.cpp
+++ b/Engine/Plugins/Experimental/Chooser/Source/ProxyTable/Private/ProxyAsset.cpp
@@ -93,7 +93,7 @@ UProxyAsset::UProxyAsset(const FObjectInitializer& Initializer)
 	ProxyTable.InitializeAs(FProxyTableContextProperty::StaticStruct());
 }
 
-UObject* UProxyAsset::FindProxyObject(FChooserEvaluationContext Context) const
+UObject* UProxyAsset::FindProxyObject(FChooserEvaluationContext& Context) const
 {
 	if (ProxyTable.IsValid())
 	{
diff --git a/Engine/Plugins/Experimental/Chooser/Source/ProxyTable/Private/ProxyTable.cpp b/Engine/Plugins/Experimental/Chooser/Source/ProxyTable/Private/ProxyTable.cpp
--- a/Engine/Plugins/Experimental/Chooser/Source/ProxyTable/Private/ProxyTable.cpp
+++ b/Engine/Plugins/Experimental/Chooser/Source/ProxyTable/Private/ProxyTable.cpp
@@ -125,16 +125,11 @@ void UProxyTable::BuildRuntimeData()
 	Keys.Empty(EntryCount);
 	Values.Empty();
 	
-	TArray<FConstStructView> Views;
-	Views.Reserve(EntryCount);
-	
 	for(const FProxyEntry& Entry : RuntimeEntries) 
 	{
 		Keys.Add(Entry.GetGuid());
-   		Views.Add(Entry.ValueStruct);
+   		Values.Add({Entry.ValueStruct, Entry.OutputStructData});
 	}
-
-	Values = Views;
 	
 	// register callbacks on updated dependencies
 	for (TWeakObjectPtr<UProxyTable> Dependency : TableDependencies)
@@ -173,8 +168,46 @@ UObject* UProxyTable::FindProxyObject(const FGuid& Key, FChooserEvaluationContex
 	const int FoundIndex = Algo::BinarySearch(Keys, Key);
 	if (FoundIndex != INDEX_NONE)
 	{
-		const FObjectChooserBase &EntryValue = Values[FoundIndex].Get<const FObjectChooserBase>();
-		return EntryValue.ChooseObject(Context);
+		const FRuntimeProxyValue& EntryValueData = Values[FoundIndex];
+		const FObjectChooserBase &EntryValue = EntryValueData.Value.Get<const FObjectChooserBase>();
+
+		UObject* Result = EntryValue.ChooseObject(Context);
+
+		for (const FProxyStructOutput& StructOutput : EntryValueData.OutputStructData)
+		{
+			const void* Container = nullptr;
+			const UStruct* StructType;
+
+			// copy each struct output value
+
+			if (StructOutput.Binding.PropertyBindingChain.IsEmpty())
+			{
+				if(Context.Params.IsValidIndex((StructOutput.Binding.ContextIndex)))
+				{
+					// directly bound to context struct
+					if (Context.Params[StructOutput.Binding.ContextIndex].GetScriptStruct() == StructOutput.Value.GetScriptStruct())
+					{
+						void* TargetData = Context.Params[StructOutput.Binding.ContextIndex].GetMutableMemory();
+						StructOutput.Value.GetScriptStruct()->CopyScriptStruct(TargetData, StructOutput.Value.GetMemory());
+					}
+				}
+			}
+			else if (UE::Chooser::ResolvePropertyChain(Context, StructOutput.Binding, Container, StructType))
+			{
+				if (FStructProperty* Property = FindFProperty<FStructProperty>(StructType, StructOutput.Binding.PropertyBindingChain.Last()))
+				{
+					// const cast is here just because ResolvePropertyChain expects a const void*&
+					void* TargetData = Property->ContainerPtrToValuePtr<void>(const_cast<void*>(Container));
+					
+					if (Property->Struct == StructOutput.Value.GetScriptStruct())
+					{
+						Property->Struct->CopyScriptStruct(TargetData, StructOutput.Value.GetMemory());
+					}
+				}
+			}
+		}
+
+		return Result;
 	}
 	
 	return nullptr;
diff --git a/Engine/Plugins/Experimental/Chooser/Source/ProxyTable/Public/ProxyAsset.h b/Engine/Plugins/Experimental/Chooser/Source/ProxyTable/Public/ProxyAsset.h
--- a/Engine/Plugins/Experimental/Chooser/Source/ProxyTable/Public/ProxyAsset.h
+++ b/Engine/Plugins/Experimental/Chooser/Source/ProxyTable/Public/ProxyAsset.h
@@ -43,7 +43,7 @@ public:
 	FGuid Guid;
 
 	virtual TConstArrayView<FInstancedStruct> GetContextData() const override { return ContextData; }
-	UObject* FindProxyObject(FChooserEvaluationContext Context) const;
+	UObject* FindProxyObject(FChooserEvaluationContext& Context) const;
 
 #if WITH_EDITORONLY_DATA
 	UPROPERTY()
diff --git a/Engine/Plugins/Experimental/Chooser/Source/ProxyTable/Public/ProxyTable.h b/Engine/Plugins/Experimental/Chooser/Source/ProxyTable/Public/ProxyTable.h
--- a/Engine/Plugins/Experimental/Chooser/Source/ProxyTable/Public/ProxyTable.h
+++ b/Engine/Plugins/Experimental/Chooser/Source/ProxyTable/Public/ProxyTable.h
@@ -3,7 +3,6 @@
 
 #include "CoreMinimal.h"
 #include "UObject/Object.h"
-#include "ChooserPropertyAccess.h"
 #include "IChooserParameterProxyTable.h"
 #include "InstancedStruct.h"
 #include "InstancedStructContainer.h"
@@ -11,6 +10,18 @@
 #include "ProxyAsset.h"
 #include "ProxyTable.generated.h"
 
+USTRUCT()
+struct PROXYTABLE_API FProxyStructOutput
+{
+	GENERATED_BODY()
+	
+	UPROPERTY(EditAnywhere, Category = "Data")
+	FChooserStructPropertyBinding Binding;
+	
+	UPROPERTY(EditAnywhere, Category = "Data")
+	FInstancedStruct Value;
+};
+
 USTRUCT()
 struct PROXYTABLE_API FProxyEntry
 {
@@ -29,6 +40,9 @@ struct PROXYTABLE_API FProxyEntry
 	UPROPERTY(DisplayName="Value", EditAnywhere, Meta = (ExcludeBaseStruct, BaseStruct = "/Script/Chooser.ObjectChooserBase"), Category = "Data")
 	FInstancedStruct ValueStruct;
 
+	UPROPERTY(EditAnywhere, Category = "Data")
+	TArray<FProxyStructOutput> OutputStructData;
+
 	const FGuid GetGuid() const;
 };
 
@@ -38,6 +52,18 @@ inline uint32 GetTypeHash(const FProxyEntry& Entry) { return GetTypeHash(Entry.G
 
 DECLARE_MULTICAST_DELEGATE(FProxyTableChanged);
 
+USTRUCT()
+struct PROXYTABLE_API FRuntimeProxyValue
+{
+	GENERATED_BODY()
+
+	UPROPERTY()
+	FInstancedStruct Value;
+
+	UPROPERTY()
+	TArray<FProxyStructOutput> OutputStructData;
+};
+
 UCLASS(MinimalAPI,BlueprintType)
 class UProxyTable : public UObject
 {
@@ -49,7 +75,7 @@ public:
 	TArray<FGuid> Keys;
 
 	UPROPERTY()
-	FInstancedStructContainer Values;
+	TArray<FRuntimeProxyValue> Values;
 	
 	UObject* FindProxyObject(const FGuid& Key, FChooserEvaluationContext& Context) const;
 
diff --git a/Engine/Plugins/Experimental/Chooser/Source/ProxyTableEditor/Private/ProxyTableEditor.cpp b/Engine/Plugins/Experimental/Chooser/Source/ProxyTableEditor/Private/ProxyTableEditor.cpp
--- a/Engine/Plugins/Experimental/Chooser/Source/ProxyTableEditor/Private/ProxyTableEditor.cpp
+++ b/Engine/Plugins/Experimental/Chooser/Source/ProxyTableEditor/Private/ProxyTableEditor.cpp
@@ -19,11 +19,13 @@
 #include "SourceCodeNavigation.h"
 #include "ProxyTable.h"
 #include "ClassViewerFilter.h"
+#include "DetailLayoutBuilder.h"
 #include "IPropertyAccessEditor.h"
 #include "LandscapeRender.h"
 #include "Widgets/Layout/SWidgetSwitcher.h"
 #include "DragAndDrop/DecoratedDragDropOp.h"
 #include "GraphEditorSettings.h"
+#include "IDetailCustomization.h"
 #include "ObjectChooserClassFilter.h"
 #include "ScopedTransaction.h"
 #include "ObjectChooserWidgetFactories.h"
@@ -42,7 +44,48 @@ namespace UE::ProxyTableEditor
 const FName FProxyTableEditor::ToolkitFName( TEXT( "GenericAssetEditor" ) );
 const FName FProxyTableEditor::PropertiesTabId( TEXT( "ProxyEditor_Properties" ) );
 const FName FProxyTableEditor::TableTabId( TEXT( "ProxyEditor_Table" ) );
+	
+
+class FProxyRowDetails : public IDetailCustomization
+{
+public:
+	FProxyRowDetails() {};
+	virtual ~FProxyRowDetails() override {};
+
+	static TSharedRef<IDetailCustomization> MakeInstance()
+	{
+		return MakeShareable( new FProxyRowDetails() );
+	}
+
+	// IDetailCustomization interface
+	virtual void CustomizeDetails(class IDetailLayoutBuilder& DetailBuilder) override;
+};
+
+// Make the details panel show the values for the selected row, showing each column value
+void FProxyRowDetails::CustomizeDetails(IDetailLayoutBuilder& DetailBuilder)
+{
+	TArray<TWeakObjectPtr<UObject>> Objects;
+	DetailBuilder.GetObjectsBeingCustomized(Objects);
+
+	UProxyRowDetails* Row = Cast<UProxyRowDetails>(Objects[0]);
+	UProxyTable* ProxyTable = Row->ProxyTable;
+	
+	if (ProxyTable->Entries.Num() > Row->Row)
+	{
+		IDetailCategoryBuilder& PropertiesCategory = DetailBuilder.EditCategory("Row Properties");
 
+		TSharedPtr<IPropertyHandle> ProxyTableProperty = DetailBuilder.GetProperty("ProxyTable", Row->StaticClass());
+		DetailBuilder.HideProperty(ProxyTableProperty);
+	
+		TSharedPtr<IPropertyHandle> EntriesArrayProperty = ProxyTableProperty->GetChildHandle("Entries");
+		TSharedPtr<IPropertyHandle> CurrentEntryProperty = EntriesArrayProperty->AsArray()->GetElement(Row->Row);
+		IDetailPropertyRow& NewEntryProperty = PropertiesCategory.AddProperty(CurrentEntryProperty);
+		NewEntryProperty.DisplayName(LOCTEXT("Entry","Selected Entry"));
+		NewEntryProperty.ShowPropertyButtons(false); // hide array add button
+		NewEntryProperty.ShouldAutoExpand(true);
+	}
+}
+	
 void FProxyTableEditor::RegisterTabSpawners(const TSharedRef<class FTabManager>& InTabManager)
 {
 	WorkspaceMenuCategory = InTabManager->AddLocalWorkspaceMenuCategory(LOCTEXT("WorkspaceMenu_GenericAssetEditor", "Asset Editor"));
@@ -730,28 +773,24 @@ void FProxyTableEditor::UpdateTableColumns()
 					.DefaultLabel(LOCTEXT("KeyColumnName", "Proxy"))
 					.ManualWidth(500));
 
-
-		if (UProxyTable* Table = GetProxyTable())
-    	{
-    		if (Table->Entries.Num() > 0)
-    		{
-    			// if the first entry has a non-none key, assume this is an old table and make an extra column with the old FName Key property
-    			if (Table->Entries[0].Key != NAME_None)
-    			{
-    					HeaderRow->AddColumn(SHeaderRow::Column("OldKey")
-                    					.DefaultLabel(LOCTEXT("OldKeyColumnName", "Key (Deprecated)"))
-                    					.ManualWidth(500));
-    			}
-    		}
-    	}
+  // Code for allowing editing of deprecated Key data
+  //   if (UProxyTable* Table = GetProxyTable())
+  //   	{
+  //   		if (Table->Entries.Num() > 0)
+  //   		{
+  //   			// if the first entry has a non-none key, assume this is an old table and make an extra column with the old FName Key property
+  //   			if (Table->Entries[0].Key != NAME_None)
+  //   			{
+  //   					HeaderRow->AddColumn(SHeaderRow::Column("OldKey")
+  //                   					.DefaultLabel(LOCTEXT("OldKeyColumnName", "Key (Deprecated)"))
+  //                   					.ManualWidth(500));
+  //   			}
+  //   		}
+  //   	}
 
 	HeaderRow->AddColumn(SHeaderRow::Column("Value")
 					.DefaultLabel(LOCTEXT("ValueColumnName", "Value"))
 					.ManualWidth(500));
-
-	
-
-
 }
 
 TSharedRef<SDockTab> FProxyTableEditor::SpawnTableTab( const FSpawnTabArgs& Args )
@@ -810,20 +849,30 @@ TSharedRef<SDockTab> FProxyTableEditor::SpawnTableTab( const FSpawnTabArgs& Args
 				{
 					if (SelectedItem)
 					{
+						for (UObject* SelectedRow : SelectedRows)
+                     	{
+                     		SelectedRow->ClearFlags(RF_Standalone);
+                     	}
 						SelectedRows.SetNum(0);
 						UProxyTable* ProxyTable = Cast<UProxyTable>(EditingObjects[0]);
 						// Get the list of objects to edit the details of
 						TObjectPtr<UProxyRowDetails> Selection = NewObject<UProxyRowDetails>();
 						Selection->ProxyTable = ProxyTable;
 						Selection->Row = SelectedItem->RowIndex;
-						Selection->SetFlags(RF_Transactional); 
+						Selection->SetFlags(RF_Standalone);
 						SelectedRows.Add(Selection);
-											
+						
+						TArray<UObject*> DetailsObjects;
+						for(auto& Item : SelectedRows)
+						{
+							DetailsObjects.Add(Item.Get());
+						}
+
 						if( DetailsView.IsValid() )
 						{
 							// Make sure details window is pointing to our object
-							//DetailsView->SetObjects( ... )
-						}
+							DetailsView->SetObjects( DetailsObjects );
+						}		
 					}
 				})
     			.OnGenerateRow_Raw(this, &FProxyTableEditor::GenerateTableRow)
@@ -1009,6 +1058,9 @@ TSharedRef<SWidget> CreateLookupProxyWidget(bool bReadOnly, UObject* Transaction
 void FProxyTableEditor::RegisterWidgets()
 {
 	UE::ChooserEditor::FObjectChooserWidgetFactories::RegisterWidgetCreator(FLookupProxy::StaticStruct(), CreateLookupProxyWidget);
+	
+	FPropertyEditorModule& PropertyModule = FModuleManager::LoadModuleChecked<FPropertyEditorModule>("PropertyEditor");
+	PropertyModule.RegisterCustomClassLayout("ProxyRowDetails", FOnGetDetailCustomizationInstance::CreateStatic(&FProxyRowDetails::MakeInstance));	
 }
 	
 }
diff --git a/Engine/Plugins/Experimental/Chooser/Source/ProxyTableEditor/Private/ProxyTableEditorModule.cpp b/Engine/Plugins/Experimental/Chooser/Source/ProxyTableEditor/Private/ProxyTableEditorModule.cpp
--- a/Engine/Plugins/Experimental/Chooser/Source/ProxyTableEditor/Private/ProxyTableEditorModule.cpp
+++ b/Engine/Plugins/Experimental/Chooser/Source/ProxyTableEditor/Private/ProxyTableEditorModule.cpp
@@ -3,6 +3,8 @@
 #include "ProxyTableEditorModule.h"
 #include "ProxyTableEditor.h"
 #include "ProxyTableEditorCommands.h"
+#include "StructOutputDataCustomization.h"
+#include "PropertyEditorModule.h"
 
 #define LOCTEXT_NAMESPACE "ProxyTableEditorModule"
 
@@ -14,10 +16,15 @@ void FModule::StartupModule()
 	FProxyTableEditor::RegisterWidgets();
 
 	FProxyTableEditorCommands::Register();
+	
+	FPropertyEditorModule& PropertyModule = FModuleManager::LoadModuleChecked<FPropertyEditorModule>("PropertyEditor");
+	PropertyModule.RegisterCustomPropertyTypeLayout(FProxyStructOutput::StaticStruct()->GetFName(), FOnGetPropertyTypeCustomizationInstance::CreateLambda([] { return MakeShared<FStructOutputDataCustomization>(); }));
 }
 
 void FModule::ShutdownModule()
 {
+	FProxyTableEditorCommands::Unregister();
+	FPropertyEditorModule& PropertyModule = FModuleManager::LoadModuleChecked<FPropertyEditorModule>("PropertyEditor");
 }
 
 }
diff --git a/Engine/Plugins/Experimental/Chooser/Source/ProxyTableEditor/Private/StructOutputDataCustomization.cpp b/Engine/Plugins/Experimental/Chooser/Source/ProxyTableEditor/Private/StructOutputDataCustomization.cpp
new file mode 100644
--- /dev/null
+++ b/Engine/Plugins/Experimental/Chooser/Source/ProxyTableEditor/Private/StructOutputDataCustomization.cpp
@@ -0,0 +1,158 @@
+﻿// Copyright Epic Games, Inc. All Rights Reserved.
+
+#include "StructOutputDataCustomization.h"
+
+#include "Chooser.h"
+#include "ChooserPropertyAccess.h"
+#include "PropertyHandle.h"
+#include "DetailWidgetRow.h"
+#include "GraphEditorSettings.h"
+#include "SClassViewer.h"
+#include "IDetailChildrenBuilder.h"
+#include "IPropertyAccessEditor.h"
+#include "ScopedTransaction.h"
+#include "SPropertyAccessChainWidget.h"
+#include "InstancedStructDetails.h"
+#include "ProxyTable.h"
+
+#define LOCTEXT_NAMESPACE "StructOutputCustomization"
+
+namespace UE::ProxyTableEditor
+{
+	
+void FStructOutputDataCustomization::CustomizeHeader(TSharedRef<IPropertyHandle> PropertyHandle, FDetailWidgetRow& HeaderRow, IPropertyTypeCustomizationUtils& CustomizationUtils)
+{
+	FProperty* Property = PropertyHandle->GetProperty();
+
+	UClass* ContextClass = nullptr;
+
+	static FName ProxyPropertyName = "Proxy";
+
+	IHasContextClass* HasContext = nullptr;
+	// this details customization is hard coded to work in a ProxyTable where the StructOutputs array will be the parent
+	if (TSharedPtr<IPropertyHandle> ArrayHandle = PropertyHandle->GetParentHandle())
+	{
+		// and then the parent of that is the FProxyEntry
+		if (TSharedPtr<IPropertyHandle> ProxyEntryHandle = ArrayHandle->GetParentHandle())
+		{
+			// and the FProxyEntry contains a UProxyAsset reference
+			if (TSharedPtr<IPropertyHandle> ProxyHandle = ProxyEntryHandle->GetChildHandle(ProxyPropertyName))
+			{
+				if (FProperty* ProxyProperty = ProxyHandle->GetProperty())
+				{
+					if (FObjectProperty* ProxyObjectProperty = CastField<FObjectProperty>(ProxyProperty))
+					{
+						void* ValuePtr = nullptr;
+						ProxyHandle->GetValueData(ValuePtr);
+						const TObjectPtr<UProxyAsset>& ProxyAssetReference = *reinterpret_cast<TObjectPtr<UProxyAsset>*>(ValuePtr);
+
+						if (ProxyAssetReference)
+						{
+							// and the UProxyAsset has has a context on it that we will use for the binding
+							HasContext = static_cast<IHasContextClass*>(ProxyAssetReference->GetInterfaceAddress(UHasContextClass::StaticClass()));
+						}
+					}
+				}
+			}
+		}
+	}
+
+	static const FName BindingName = "Binding";
+	TSharedPtr<IPropertyHandle> BindingHandle = PropertyHandle->GetChildHandle(BindingName);
+	
+	static const FName ValueName = "Value";
+	TSharedPtr<IPropertyHandle> ValueHandle = PropertyHandle->GetChildHandle(ValueName);
+	
+	HeaderRow
+	.NameContent()
+	[
+		PropertyHandle->CreatePropertyNameWidget()
+	]
+	.ValueContent()
+	[
+		SNew(UE::ChooserEditor::SPropertyAccessChainWidget).ContextClassOwner(HasContext).TypeFilter("struct").BindingColor("StructPinTypeColor").AllowFunctions(false)
+		.PropertyBindingValue_Lambda([BindingHandle]()
+		{
+			void* data;
+			BindingHandle->GetValueData(data);
+			return reinterpret_cast<const FChooserPropertyBinding*>(data);
+		})
+		.OnAddBinding_Lambda([HasContext, ValueHandle, BindingHandle](FName InPropertyName, const TArray<FBindingChainElement>& InBindingChain)
+         	 	{
+         	 		TArray<UObject*> OuterObjects;
+         	 		BindingHandle->GetOuterObjects(OuterObjects);
+         	 		const FScopedTransaction Transaction(LOCTEXT("Change Property Binding", "Change Property Binding"));
+         	 	
+         			for (uint32 i=0; i<BindingHandle->GetNumOuterObjects(); i++)
+         	 		{
+         				// Set the binding
+         	 			void* ValuePtr = nullptr;
+         	 			BindingHandle->GetValueData(ValuePtr);
+         	 			FChooserPropertyBinding* PropertyValue = reinterpret_cast<FChooserPropertyBinding*>(ValuePtr);
+         	 		
+         	 			if (PropertyValue != nullptr)
+         	 			{
+         	 				BindingHandle->NotifyPreChange();
+         	 				
+         	 				OuterObjects[i]->Modify(true);
+         	 				Chooser::CopyPropertyChain(InBindingChain, *PropertyValue);
+         	
+							FChooserStructPropertyBinding* StructPropertyBinding = static_cast<FChooserStructPropertyBinding*>(PropertyValue);
+
+         	 				if (InBindingChain.Num() > 1)
+         	 				{
+         	 					FField* Property = InBindingChain.Last().Field.ToField();
+								if (const FStructProperty* StructProperty = CastField<const FStructProperty>(Property))
+								{
+									StructPropertyBinding->StructType = StructProperty->Struct;
+								}
+         	 				}
+         	 				else if (InBindingChain.Num() ==1)
+         	 				{
+         	 					// directly bound to context struct
+         	 					int Index = InBindingChain[0].ArrayIndex;
+								TConstArrayView<FInstancedStruct> ContextData = HasContext->GetContextData();
+								if (const FContextObjectTypeStruct* StructContextData = ContextData[Index].GetPtr<FContextObjectTypeStruct>())
+								{
+									StructPropertyBinding->StructType = StructContextData->Struct;
+									StructPropertyBinding->DisplayName = StructContextData->Struct->GetAuthoredName();
+								}
+         	 				}
+         	 				
+							// reset the type of the FInstancedStruct storing the value
+							ValueHandle->GetValueData(ValuePtr);
+							FInstancedStruct* ValueStruct = reinterpret_cast<FInstancedStruct*>(ValuePtr);
+							ValueStruct->InitializeAs(StructPropertyBinding->StructType);
+	
+         	 				BindingHandle->NotifyPostChange(EPropertyChangeType::ValueSet);
+							ValueHandle->NotifyPostChange(EPropertyChangeType::ValueSet);
+         	 			}
+
+         	 		}
+         	 	})
+	];
+}
+	
+void FStructOutputDataCustomization::CustomizeChildren(TSharedRef<IPropertyHandle> StructPropertyHandle, IDetailChildrenBuilder& ChildBuilder, IPropertyTypeCustomizationUtils& CustomizationUtils)
+{
+	static FName ValueName = "Value";
+	
+	uint32 NumChildren = 0;
+	StructPropertyHandle->GetNumChildren(NumChildren);
+
+	for (uint32 ChildIndex = 0; ChildIndex < NumChildren; ++ChildIndex)
+	{
+		TSharedPtr<IPropertyHandle> CurrentProperty = StructPropertyHandle->GetChildHandle(ChildIndex);
+		if (CurrentProperty->GetProperty()->GetFName() == ValueName)
+		{
+			// show the child properties of "Value" as our child properties
+			TSharedRef<FInstancedStructDataDetails> DataDetails = MakeShared<FInstancedStructDataDetails>(CurrentProperty);
+			ChildBuilder.AddCustomBuilder(DataDetails);
+		}
+	}
+}
+	
+}
+
+
+#undef LOCTEXT_NAMESPACE
\ No newline at end of file
diff --git a/Engine/Plugins/Experimental/Chooser/Source/ProxyTableEditor/Private/StructOutputDataCustomization.h b/Engine/Plugins/Experimental/Chooser/Source/ProxyTableEditor/Private/StructOutputDataCustomization.h
new file mode 100644
--- /dev/null
+++ b/Engine/Plugins/Experimental/Chooser/Source/ProxyTableEditor/Private/StructOutputDataCustomization.h
@@ -0,0 +1,20 @@
+﻿// Copyright Epic Games, Inc. All Rights Reserved.
+
+#pragma once
+
+#include "IPropertyTypeCustomization.h"
+
+namespace UE::ProxyTableEditor
+{
+
+class FStructOutputDataCustomization : public IPropertyTypeCustomization
+{
+public:
+	FStructOutputDataCustomization(){}
+private:
+	// IPropertyTypeCustomization interface
+	virtual void CustomizeHeader(TSharedRef<IPropertyHandle> PropertyHandle, FDetailWidgetRow& HeaderRow, IPropertyTypeCustomizationUtils& CustomizationUtils) override;
+	virtual void CustomizeChildren(TSharedRef<IPropertyHandle> PropertyHandle, IDetailChildrenBuilder& ChildBuilder, IPropertyTypeCustomizationUtils& CustomizationUtils) override;
+};
+
+}
\ No newline at end of file
diff --git a/Engine/Plugins/Experimental/Chooser/Source/ProxyTableEditor/ProxyTableEditor.Build.cs b/Engine/Plugins/Experimental/Chooser/Source/ProxyTableEditor/ProxyTableEditor.Build.cs
--- a/Engine/Plugins/Experimental/Chooser/Source/ProxyTableEditor/ProxyTableEditor.Build.cs
+++ b/Engine/Plugins/Experimental/Chooser/Source/ProxyTableEditor/ProxyTableEditor.Build.cs
@@ -31,7 +31,8 @@ namespace UnrealBuildTool.Rules
 					"GraphEditor",
 					"GameplayTags",
 					"GameplayTagsEditor",
-					"StructUtils"
+					"StructUtils",
+					"StructUtilsEditor"
 					// ... add private dependencies that you statically link with here ...
 				}
 			);
