Proxy Table struct output feature - fix for some warnings caused by property type change

diff --git a/Engine/Plugins/Experimental/Chooser/Source/ProxyTable/Private/ProxyTable.cpp b/Engine/Plugins/Experimental/Chooser/Source/ProxyTable/Private/ProxyTable.cpp
--- a/Engine/Plugins/Experimental/Chooser/Source/ProxyTable/Private/ProxyTable.cpp
+++ b/Engine/Plugins/Experimental/Chooser/Source/ProxyTable/Private/ProxyTable.cpp
@@ -123,12 +123,12 @@ void UProxyTable::BuildRuntimeData()
 	int EntryCount = RuntimeEntries.Num();
 	
 	Keys.Empty(EntryCount);
-	Values.Empty();
+	RuntimeValues.Empty();
 	
 	for(const FProxyEntry& Entry : RuntimeEntries) 
 	{
 		Keys.Add(Entry.GetGuid());
-   		Values.Add({Entry.ValueStruct, Entry.OutputStructData});
+   		RuntimeValues.Add({Entry.ValueStruct, Entry.OutputStructData});
 	}
 	
 	// register callbacks on updated dependencies
@@ -168,7 +168,7 @@ UObject* UProxyTable::FindProxyObject(const FGuid& Key, FChooserEvaluationContex
 	const int FoundIndex = Algo::BinarySearch(Keys, Key);
 	if (FoundIndex != INDEX_NONE)
 	{
-		const FRuntimeProxyValue& EntryValueData = Values[FoundIndex];
+		const FRuntimeProxyValue& EntryValueData = RuntimeValues[FoundIndex];
 		const FObjectChooserBase &EntryValue = EntryValueData.Value.Get<const FObjectChooserBase>();
 
 		UObject* Result = EntryValue.ChooseObject(Context);
diff --git a/Engine/Plugins/Experimental/Chooser/Source/ProxyTable/Public/ProxyTable.h b/Engine/Plugins/Experimental/Chooser/Source/ProxyTable/Public/ProxyTable.h
--- a/Engine/Plugins/Experimental/Chooser/Source/ProxyTable/Public/ProxyTable.h
+++ b/Engine/Plugins/Experimental/Chooser/Source/ProxyTable/Public/ProxyTable.h
@@ -75,7 +75,7 @@ public:
 	TArray<FGuid> Keys;
 
 	UPROPERTY()
-	TArray<FRuntimeProxyValue> Values;
+	TArray<FRuntimeProxyValue> RuntimeValues;
 	
 	UObject* FindProxyObject(const FGuid& Key, FChooserEvaluationContext& Context) const;
 
