Fix for crashing in EvaluateChooser node when used in Event Graphs

diff --git a/Engine/Plugins/Experimental/Chooser/Source/Chooser/Private/ChooserFunctionLibrary.cpp b/Engine/Plugins/Experimental/Chooser/Source/Chooser/Private/ChooserFunctionLibrary.cpp
--- a/Engine/Plugins/Experimental/Chooser/Source/Chooser/Private/ChooserFunctionLibrary.cpp
+++ b/Engine/Plugins/Experimental/Chooser/Source/Chooser/Private/ChooserFunctionLibrary.cpp
@@ -167,6 +167,12 @@ FInstancedStruct UChooserFunctionLibrary::MakeEvaluateChooser(UChooserTable* Cho
 	return Struct;
 }
 
+FChooserEvaluationContext UChooserFunctionLibrary::MakeChooserEvaluationContext()
+{
+	FChooserEvaluationContext Context;
+	return Context;
+}
+
 TArray<UObject*> UChooserFunctionLibrary::EvaluateChooserMulti(const UObject* ContextObject, const UChooserTable* Chooser, TSubclassOf<UObject> ObjectClass)
 {
 	TArray<UObject*> Result;
@@ -189,4 +195,4 @@ TArray<UObject*> UChooserFunctionLibrary::EvaluateChooserMulti(const UObject* Co
 	}
 
 	return Result;
-}
+}
\ No newline at end of file
diff --git a/Engine/Plugins/Experimental/Chooser/Source/Chooser/Public/ChooserFunctionLibrary.h b/Engine/Plugins/Experimental/Chooser/Source/Chooser/Public/ChooserFunctionLibrary.h
--- a/Engine/Plugins/Experimental/Chooser/Source/Chooser/Public/ChooserFunctionLibrary.h
+++ b/Engine/Plugins/Experimental/Chooser/Source/Chooser/Public/ChooserFunctionLibrary.h
@@ -96,4 +96,8 @@ public:
 	*/
 	UFUNCTION(BlueprintPure, Category = "Animation", meta = (BlueprintThreadSafe, BlueprintInternalUseOnly="true", NativeMakeFunc))
 	static FInstancedStruct MakeEvaluateChooser(UChooserTable* Chooser);
+	
+	
+	UFUNCTION(BlueprintCallable, Category = "Animation", meta = (BlueprintThreadSafe, BlueprintInternalUseOnly="true"))
+	static FChooserEvaluationContext MakeChooserEvaluationContext();
 };
diff --git a/Engine/Plugins/Experimental/Chooser/Source/Chooser/Public/IObjectChooser.h b/Engine/Plugins/Experimental/Chooser/Source/Chooser/Public/IObjectChooser.h
--- a/Engine/Plugins/Experimental/Chooser/Source/Chooser/Public/IObjectChooser.h
+++ b/Engine/Plugins/Experimental/Chooser/Source/Chooser/Public/IObjectChooser.h
@@ -44,7 +44,7 @@ struct CHOOSER_API FChooserEvaluationContext
     #endif
 	
 	GENERATED_BODY()
-	TArray<FInstancedStruct, TFixedAllocator<4>> Params;
+	TArray<FInstancedStruct, TInlineAllocator<4>> Params;
 };
 
 USTRUCT()
diff --git a/Engine/Plugins/Experimental/Chooser/Source/ChooserUncooked/Private/EvaluateChooserNode.cpp b/Engine/Plugins/Experimental/Chooser/Source/ChooserUncooked/Private/EvaluateChooserNode.cpp
--- a/Engine/Plugins/Experimental/Chooser/Source/ChooserUncooked/Private/EvaluateChooserNode.cpp
+++ b/Engine/Plugins/Experimental/Chooser/Source/ChooserUncooked/Private/EvaluateChooserNode.cpp
@@ -598,14 +598,15 @@ void UK2Node_EvaluateChooser2::ExpandNode(class FKismetCompilerContext& Compiler
 		}
 		CallFunction->AllocateDefaultPins();
 		
-		UK2Node_MakeStruct* ContextStructNode = CompilerContext.SpawnIntermediateNode<UK2Node_MakeStruct>(this, SourceGraph);
-		ContextStructNode->PostPlacedNewNode();
+		UK2Node_CallFunction* ContextStructNode = CompilerContext.SpawnIntermediateNode<UK2Node_CallFunction>(this, SourceGraph);
+		ContextStructNode->SetFromFunction(UChooserFunctionLibrary::StaticClass()->FindFunctionByName(GET_MEMBER_NAME_CHECKED(UChooserFunctionLibrary, MakeChooserEvaluationContext)));
         CompilerContext.MessageLog.NotifyIntermediateObjectCreation(ContextStructNode, this);
-		ContextStructNode->StructType = FChooserEvaluationContext::StaticStruct();
         ContextStructNode->AllocateDefaultPins();
-		UEdGraphPin* ContextStructPin = ContextStructNode->FindPin(FChooserEvaluationContext::StaticStruct()->GetFName(), EGPD_Output);
-
+		CompilerContext.MovePinLinksToIntermediate(*ExecInput, *ContextStructNode->GetExecPin());
+		UEdGraphPin* ContextStructPin = ContextStructNode->GetReturnValuePin();
+		
 		ContextStructPin->MakeLinkTo(CallFunction->FindPin(FName("Context")));
+		PreviousNodeExecOutput = ContextStructNode->GetThenPin();
 
 		UK2Node_Self* SelfNode = CompilerContext.SpawnIntermediateNode<UK2Node_Self>(this, SourceGraph);
 		CompilerContext.MessageLog.NotifyIntermediateObjectCreation(SelfNode, this);
@@ -632,14 +633,7 @@ void UK2Node_EvaluateChooser2::ExpandNode(class FKismetCompilerContext& Compiler
 							AddObjectFunction->AllocateDefaultPins();
 							ContextStructPin->MakeLinkTo(AddObjectFunction->FindPin(FName("Context")));
 
-							if (PreviousNodeExecOutput == nullptr)
-							{
-								CompilerContext.MovePinLinksToIntermediate(*ExecInput, *AddObjectFunction->GetExecPin());
-							}
-							else
-							{
-								PreviousNodeExecOutput->MakeLinkTo(AddObjectFunction->GetExecPin());
-							}
+							PreviousNodeExecOutput->MakeLinkTo(AddObjectFunction->GetExecPin());
 							PreviousNodeExecOutput = AddObjectFunction->GetThenPin();
 							
 							UEdGraphPin* AddObjectPin = AddObjectFunction->FindPin(FName("Object"));
@@ -665,14 +659,7 @@ void UK2Node_EvaluateChooser2::ExpandNode(class FKismetCompilerContext& Compiler
 						AddStructFunction->AllocateDefaultPins();
 						ContextStructPin->MakeLinkTo(AddStructFunction->FindPin(FName("Context")));
 
-						if (PreviousNodeExecOutput == nullptr)
-						{
-							CompilerContext.MovePinLinksToIntermediate(*ExecInput, *AddStructFunction->GetExecPin());
-						}
-						else
-						{
-							PreviousNodeExecOutput->MakeLinkTo(AddStructFunction->GetExecPin());
-						}
+						PreviousNodeExecOutput->MakeLinkTo(AddStructFunction->GetExecPin());
 						PreviousNodeExecOutput = AddStructFunction->GetThenPin();
 							
 						UEdGraphPin* AddStructPin = AddStructFunction->FindPin(FName("Value"));
diff --git a/Engine/Plugins/Experimental/Chooser/Source/ProxyTableUncooked/Private/EvaluateProxyNode.cpp b/Engine/Plugins/Experimental/Chooser/Source/ProxyTableUncooked/Private/EvaluateProxyNode.cpp
--- a/Engine/Plugins/Experimental/Chooser/Source/ProxyTableUncooked/Private/EvaluateProxyNode.cpp
+++ b/Engine/Plugins/Experimental/Chooser/Source/ProxyTableUncooked/Private/EvaluateProxyNode.cpp
@@ -581,12 +581,12 @@ void UK2Node_EvaluateProxy2::ExpandNode(class FKismetCompilerContext& CompilerCo
 		}
 		CallFunction->AllocateDefaultPins();
 		
-		UK2Node_MakeStruct* ContextStructNode = CompilerContext.SpawnIntermediateNode<UK2Node_MakeStruct>(this, SourceGraph);
-		ContextStructNode->PostPlacedNewNode();
+		UK2Node_CallFunction* ContextStructNode = CompilerContext.SpawnIntermediateNode<UK2Node_CallFunction>(this, SourceGraph);
+		ContextStructNode->SetFromFunction(UChooserFunctionLibrary::StaticClass()->FindFunctionByName(GET_MEMBER_NAME_CHECKED(UChooserFunctionLibrary, MakeChooserEvaluationContext)));
         CompilerContext.MessageLog.NotifyIntermediateObjectCreation(ContextStructNode, this);
-		ContextStructNode->StructType = FChooserEvaluationContext::StaticStruct();
         ContextStructNode->AllocateDefaultPins();
-		UEdGraphPin* ContextStructPin = ContextStructNode->FindPin(FChooserEvaluationContext::StaticStruct()->GetFName(), EGPD_Output);
+		CompilerContext.MovePinLinksToIntermediate(*ExecInput, *ContextStructNode->GetExecPin());
+		UEdGraphPin* ContextStructPin = ContextStructNode->GetReturnValuePin();
 
 		ContextStructPin->MakeLinkTo(CallFunction->FindPin(FName("Context")));
 
@@ -615,14 +615,7 @@ void UK2Node_EvaluateProxy2::ExpandNode(class FKismetCompilerContext& CompilerCo
 							AddObjectFunction->AllocateDefaultPins();
 							ContextStructPin->MakeLinkTo(AddObjectFunction->FindPin(FName("Context")));
 
-							if (PreviousNodeExecOutput == nullptr)
-							{
-								CompilerContext.MovePinLinksToIntermediate(*ExecInput, *AddObjectFunction->GetExecPin());
-							}
-							else
-							{
-								PreviousNodeExecOutput->MakeLinkTo(AddObjectFunction->GetExecPin());
-							}
+							PreviousNodeExecOutput->MakeLinkTo(AddObjectFunction->GetExecPin());
 							PreviousNodeExecOutput = AddObjectFunction->GetThenPin();
 							
 							UEdGraphPin* AddObjectPin = AddObjectFunction->FindPin(FName("Object"));
@@ -648,14 +641,7 @@ void UK2Node_EvaluateProxy2::ExpandNode(class FKismetCompilerContext& CompilerCo
 						AddStructFunction->AllocateDefaultPins();
 						ContextStructPin->MakeLinkTo(AddStructFunction->FindPin(FName("Context")));
 
-						if (PreviousNodeExecOutput == nullptr)
-						{
-							CompilerContext.MovePinLinksToIntermediate(*ExecInput, *AddStructFunction->GetExecPin());
-						}
-						else
-						{
-							PreviousNodeExecOutput->MakeLinkTo(AddStructFunction->GetExecPin());
-						}
+						PreviousNodeExecOutput->MakeLinkTo(AddStructFunction->GetExecPin());
 						PreviousNodeExecOutput = AddStructFunction->GetThenPin();
 							
 						UEdGraphPin* AddStructPin = AddStructFunction->FindPin(FName("Value"));
