Crash fixes for EvaluateChooser and EvaluateProxy blueprint node Expansion

diff --git a/Engine/Plugins/Experimental/Chooser/Source/ChooserUncooked/Private/EvaluateChooserNode.cpp b/Engine/Plugins/Experimental/Chooser/Source/ChooserUncooked/Private/EvaluateChooserNode.cpp
--- a/Engine/Plugins/Experimental/Chooser/Source/ChooserUncooked/Private/EvaluateChooserNode.cpp
+++ b/Engine/Plugins/Experimental/Chooser/Source/ChooserUncooked/Private/EvaluateChooserNode.cpp
@@ -625,92 +625,98 @@ void UK2Node_EvaluateChooser2::ExpandNode(class FKismetCompilerContext& Compiler
 					if (EntryType == FContextObjectTypeClass::StaticStruct())
 					{
 						const FContextObjectTypeClass& ClassContext = ContextDataEntry.Get<FContextObjectTypeClass>();
-						if (UEdGraphPin* Pin = FindPin(ClassContext.Class.GetFName(), EGPD_Input))
+						if (ClassContext.Class)
 						{
-							UK2Node_CallFunction* AddObjectFunction = CompilerContext.SpawnIntermediateNode<UK2Node_CallFunction>(this, SourceGraph);
-                       		CompilerContext.MessageLog.NotifyIntermediateObjectCreation(AddObjectFunction, this);
-							AddObjectFunction->SetFromFunction(UChooserFunctionLibrary::StaticClass()->FindFunctionByName(GET_MEMBER_NAME_CHECKED(UChooserFunctionLibrary, AddChooserObjectInput)));
-							AddObjectFunction->AllocateDefaultPins();
-							ContextStructPin->MakeLinkTo(AddObjectFunction->FindPin(FName("Context")));
-
-							PreviousNodeExecOutput->MakeLinkTo(AddObjectFunction->GetExecPin());
-							PreviousNodeExecOutput = AddObjectFunction->GetThenPin();
-							
-							UEdGraphPin* AddObjectPin = AddObjectFunction->FindPin(FName("Object"));
-							
-							if (Pin->HasAnyConnections())
-							{
-								CompilerContext.MovePinLinksToIntermediate(*Pin, *AddObjectPin);
-							}
-							else // would be nice to check that self is the same type as this Object parameter
+							if (UEdGraphPin* Pin = FindPin(ClassContext.Class.GetFName(), EGPD_Input))
 							{
-								// auto connect self node to any disconnected object pin
-								SelfPin->MakeLinkTo(AddObjectPin);
+								UK2Node_CallFunction* AddObjectFunction = CompilerContext.SpawnIntermediateNode<UK2Node_CallFunction>(this, SourceGraph);
+								CompilerContext.MessageLog.NotifyIntermediateObjectCreation(AddObjectFunction, this);
+								AddObjectFunction->SetFromFunction(UChooserFunctionLibrary::StaticClass()->FindFunctionByName(GET_MEMBER_NAME_CHECKED(UChooserFunctionLibrary, AddChooserObjectInput)));
+								AddObjectFunction->AllocateDefaultPins();
+								ContextStructPin->MakeLinkTo(AddObjectFunction->FindPin(FName("Context")));
+								PreviousNodeExecOutput = ContextStructNode->GetThenPin();
+
+								PreviousNodeExecOutput->MakeLinkTo(AddObjectFunction->GetExecPin());
+								PreviousNodeExecOutput = AddObjectFunction->GetThenPin();
+								
+								UEdGraphPin* AddObjectPin = AddObjectFunction->FindPin(FName("Object"));
+								
+								if (Pin->HasAnyConnections())
+								{
+									CompilerContext.MovePinLinksToIntermediate(*Pin, *AddObjectPin);
+								}
+								else // would be nice to check that self is the same type as this Object parameter
+								{
+									// auto connect self node to any disconnected object pin
+									SelfPin->MakeLinkTo(AddObjectPin);
+								}
 							}
 						}
 					}
 					else if (EntryType == FContextObjectTypeStruct::StaticStruct())
 					{
 						const FContextObjectTypeStruct& StructContext = ContextDataEntry.Get<FContextObjectTypeStruct>();
-						
-						UK2Node_CallFunction* AddStructFunction = CompilerContext.SpawnIntermediateNode<UK2Node_CallFunction>(this, SourceGraph);
-						CompilerContext.MessageLog.NotifyIntermediateObjectCreation(AddStructFunction, this);
-						AddStructFunction->SetFromFunction(UChooserFunctionLibrary::StaticClass()->FindFunctionByName(GET_MEMBER_NAME_CHECKED(UChooserFunctionLibrary, AddChooserStructInput)));
-						AddStructFunction->AllocateDefaultPins();
-						ContextStructPin->MakeLinkTo(AddStructFunction->FindPin(FName("Context")));
-
-						PreviousNodeExecOutput->MakeLinkTo(AddStructFunction->GetExecPin());
-						PreviousNodeExecOutput = AddStructFunction->GetThenPin();
-							
-						UEdGraphPin* AddStructPin = AddStructFunction->FindPin(FName("Value"));
-						// not sure why this isn't automatically happening:
-						AddStructPin->PinType.PinCategory = UEdGraphSchema_K2::PC_Struct;
-						AddStructPin->PinType.PinSubCategoryObject = StructContext.Struct;
-
-						bool bLinked = false;
-						if (StructContext.Direction != EContextObjectDirection::Write)
+						if (StructContext.Struct)
 						{
-							if (UEdGraphPin* Pin = FindPin(StructContext.Struct.GetFName(), EGPD_Input))
+							UK2Node_CallFunction* AddStructFunction = CompilerContext.SpawnIntermediateNode<UK2Node_CallFunction>(this, SourceGraph);
+							CompilerContext.MessageLog.NotifyIntermediateObjectCreation(AddStructFunction, this);
+							AddStructFunction->SetFromFunction(UChooserFunctionLibrary::StaticClass()->FindFunctionByName(GET_MEMBER_NAME_CHECKED(UChooserFunctionLibrary, AddChooserStructInput)));
+							AddStructFunction->AllocateDefaultPins();
+							ContextStructPin->MakeLinkTo(AddStructFunction->FindPin(FName("Context")));
+
+							PreviousNodeExecOutput->MakeLinkTo(AddStructFunction->GetExecPin());
+							PreviousNodeExecOutput = AddStructFunction->GetThenPin();
+								
+							UEdGraphPin* AddStructPin = AddStructFunction->FindPin(FName("Value"));
+							// not sure why this isn't automatically happening:
+							AddStructPin->PinType.PinCategory = UEdGraphSchema_K2::PC_Struct;
+							AddStructPin->PinType.PinSubCategoryObject = StructContext.Struct;
+
+							bool bLinked = false;
+							if (StructContext.Direction != EContextObjectDirection::Write)
 							{
-								if (Pin->HasAnyConnections())
+								if (UEdGraphPin* Pin = FindPin(StructContext.Struct.GetFName(), EGPD_Input))
 								{
-									CompilerContext.MovePinLinksToIntermediate(*Pin, *AddStructPin);
-									bLinked = true;
+									if (Pin->HasAnyConnections())
+									{
+										CompilerContext.MovePinLinksToIntermediate(*Pin, *AddStructPin);
+										bLinked = true;
+									}
 								}
 							}
-						}
-						
-						if (!bLinked)
-						{
-							// create a struct to hold the output (or for input structs that were not connected to anything)
-							UK2Node_MakeStruct* OutputStructNode = CompilerContext.SpawnIntermediateNode<UK2Node_MakeStruct>(this, SourceGraph);
-							CompilerContext.MessageLog.NotifyIntermediateObjectCreation(OutputStructNode, this);
-							OutputStructNode->StructType = static_cast<UScriptStruct*>(StructContext.Struct);
-							OutputStructNode->AllocateDefaultPins();
-							OutputStructNode->FindPin(StructContext.Struct->GetFName(), EGPD_Output)->MakeLinkTo(AddStructPin);
-						}
-
-						if (StructContext.Direction != EContextObjectDirection::Read)
-						{
 							
-							if (UEdGraphPin* Pin = FindPin(StructContext.Struct.GetFName(), EGPD_Output))
+							if (!bLinked)
 							{
-								if (Pin->HasAnyConnections())
+								// create a struct to hold the output (or for input structs that were not connected to anything)
+								UK2Node_MakeStruct* OutputStructNode = CompilerContext.SpawnIntermediateNode<UK2Node_MakeStruct>(this, SourceGraph);
+								CompilerContext.MessageLog.NotifyIntermediateObjectCreation(OutputStructNode, this);
+								OutputStructNode->StructType = static_cast<UScriptStruct*>(StructContext.Struct);
+								OutputStructNode->AllocateDefaultPins();
+								OutputStructNode->FindPin(StructContext.Struct->GetFName(), EGPD_Output)->MakeLinkTo(AddStructPin);
+							}
+
+							if (StructContext.Direction != EContextObjectDirection::Read)
+							{
+								
+								if (UEdGraphPin* Pin = FindPin(StructContext.Struct.GetFName(), EGPD_Output))
 								{
-									// set up struct output pin
-									UK2Node_CallFunction* GetStructFunction = CompilerContext.SpawnIntermediateNode<UK2Node_CallFunction>(this, SourceGraph);
-									CompilerContext.MessageLog.NotifyIntermediateObjectCreation(GetStructFunction, this);
-									GetStructFunction->SetFromFunction(UChooserFunctionLibrary::StaticClass()->FindFunctionByName(GET_MEMBER_NAME_CHECKED(UChooserFunctionLibrary, GetChooserStructOutput)));
-									GetStructFunction->AllocateDefaultPins();
-									GetStructFunction->FindPin(FName("Index"))->DefaultValue = FString::FromInt(ContextDataIndex);
-									
-									ContextStructPin->MakeLinkTo(GetStructFunction->FindPin(FName("Context")));
-									UEdGraphPin* ValuePin = GetStructFunction->FindPin(FName("Value"));
-									// not sure why this isn't automatically happening:
-									ValuePin->PinType.PinCategory = UEdGraphSchema_K2::PC_Struct;
-									ValuePin->PinType.PinSubCategoryObject = StructContext.Struct;
-					
-									CompilerContext.MovePinLinksToIntermediate(*Pin, *ValuePin);
+									if (Pin->HasAnyConnections())
+									{
+										// set up struct output pin
+										UK2Node_CallFunction* GetStructFunction = CompilerContext.SpawnIntermediateNode<UK2Node_CallFunction>(this, SourceGraph);
+										CompilerContext.MessageLog.NotifyIntermediateObjectCreation(GetStructFunction, this);
+										GetStructFunction->SetFromFunction(UChooserFunctionLibrary::StaticClass()->FindFunctionByName(GET_MEMBER_NAME_CHECKED(UChooserFunctionLibrary, GetChooserStructOutput)));
+										GetStructFunction->AllocateDefaultPins();
+										GetStructFunction->FindPin(FName("Index"))->DefaultValue = FString::FromInt(ContextDataIndex);
+										
+										ContextStructPin->MakeLinkTo(GetStructFunction->FindPin(FName("Context")));
+										UEdGraphPin* ValuePin = GetStructFunction->FindPin(FName("Value"));
+										// not sure why this isn't automatically happening:
+										ValuePin->PinType.PinCategory = UEdGraphSchema_K2::PC_Struct;
+										ValuePin->PinType.PinSubCategoryObject = StructContext.Struct;
+						
+										CompilerContext.MovePinLinksToIntermediate(*Pin, *ValuePin);
+									}
 								}
 							}
 						}
diff --git a/Engine/Plugins/Experimental/Chooser/Source/ProxyTableUncooked/Private/EvaluateProxyNode.cpp b/Engine/Plugins/Experimental/Chooser/Source/ProxyTableUncooked/Private/EvaluateProxyNode.cpp
--- a/Engine/Plugins/Experimental/Chooser/Source/ProxyTableUncooked/Private/EvaluateProxyNode.cpp
+++ b/Engine/Plugins/Experimental/Chooser/Source/ProxyTableUncooked/Private/EvaluateProxyNode.cpp
@@ -464,7 +464,7 @@ void UK2Node_EvaluateProxy2::AllocateDefaultPins()
    	{
    		ResultType = Proxy->Type;
    	}
-
+	
 	static const FName ResultPinName = "Result";
 	if (UEdGraphPin* ResultPin = FindPin(ResultPinName, EGPD_Output))
 	{
@@ -607,92 +607,96 @@ void UK2Node_EvaluateProxy2::ExpandNode(class FKismetCompilerContext& CompilerCo
 					if (EntryType == FContextObjectTypeClass::StaticStruct())
 					{
 						const FContextObjectTypeClass& ClassContext = ContextDataEntry.Get<FContextObjectTypeClass>();
-						if (UEdGraphPin* Pin = FindPin(ClassContext.Class.GetFName(), EGPD_Input))
+						if (ClassContext.Class)
 						{
-							UK2Node_CallFunction* AddObjectFunction = CompilerContext.SpawnIntermediateNode<UK2Node_CallFunction>(this, SourceGraph);
-                       		CompilerContext.MessageLog.NotifyIntermediateObjectCreation(AddObjectFunction, this);
-							AddObjectFunction->SetFromFunction(UChooserFunctionLibrary::StaticClass()->FindFunctionByName(GET_MEMBER_NAME_CHECKED(UChooserFunctionLibrary, AddChooserObjectInput)));
-							AddObjectFunction->AllocateDefaultPins();
-							ContextStructPin->MakeLinkTo(AddObjectFunction->FindPin(FName("Context")));
-
-							PreviousNodeExecOutput->MakeLinkTo(AddObjectFunction->GetExecPin());
-							PreviousNodeExecOutput = AddObjectFunction->GetThenPin();
-							
-							UEdGraphPin* AddObjectPin = AddObjectFunction->FindPin(FName("Object"));
-							
-							if (Pin->HasAnyConnections())
-							{
-								CompilerContext.MovePinLinksToIntermediate(*Pin, *AddObjectPin);
-							}
-							else // would be nice to check that self is the same type as this Object parameter
+							if (UEdGraphPin* Pin = FindPin(ClassContext.Class.GetFName(), EGPD_Input))
 							{
-								// auto connect self node to any disconnected object pin
-								SelfPin->MakeLinkTo(AddObjectPin);
+								UK2Node_CallFunction* AddObjectFunction = CompilerContext.SpawnIntermediateNode<UK2Node_CallFunction>(this, SourceGraph);
+								CompilerContext.MessageLog.NotifyIntermediateObjectCreation(AddObjectFunction, this);
+								AddObjectFunction->SetFromFunction(UChooserFunctionLibrary::StaticClass()->FindFunctionByName(GET_MEMBER_NAME_CHECKED(UChooserFunctionLibrary, AddChooserObjectInput)));
+								AddObjectFunction->AllocateDefaultPins();
+								ContextStructPin->MakeLinkTo(AddObjectFunction->FindPin(FName("Context")));
+
+								PreviousNodeExecOutput->MakeLinkTo(AddObjectFunction->GetExecPin());
+								PreviousNodeExecOutput = AddObjectFunction->GetThenPin();
+								
+								UEdGraphPin* AddObjectPin = AddObjectFunction->FindPin(FName("Object"));
+								
+								if (Pin->HasAnyConnections())
+								{
+									CompilerContext.MovePinLinksToIntermediate(*Pin, *AddObjectPin);
+								}
+								else // would be nice to check that self is the same type as this Object parameter
+									{
+									// auto connect self node to any disconnected object pin
+									SelfPin->MakeLinkTo(AddObjectPin);
+									}
 							}
 						}
 					}
 					else if (EntryType == FContextObjectTypeStruct::StaticStruct())
 					{
 						const FContextObjectTypeStruct& StructContext = ContextDataEntry.Get<FContextObjectTypeStruct>();
-						
-						UK2Node_CallFunction* AddStructFunction = CompilerContext.SpawnIntermediateNode<UK2Node_CallFunction>(this, SourceGraph);
-						CompilerContext.MessageLog.NotifyIntermediateObjectCreation(AddStructFunction, this);
-						AddStructFunction->SetFromFunction(UChooserFunctionLibrary::StaticClass()->FindFunctionByName(GET_MEMBER_NAME_CHECKED(UChooserFunctionLibrary, AddChooserStructInput)));
-						AddStructFunction->AllocateDefaultPins();
-						ContextStructPin->MakeLinkTo(AddStructFunction->FindPin(FName("Context")));
-
-						PreviousNodeExecOutput->MakeLinkTo(AddStructFunction->GetExecPin());
-						PreviousNodeExecOutput = AddStructFunction->GetThenPin();
+						if (StructContext.Struct)
+						{
+							UK2Node_CallFunction* AddStructFunction = CompilerContext.SpawnIntermediateNode<UK2Node_CallFunction>(this, SourceGraph);
+							CompilerContext.MessageLog.NotifyIntermediateObjectCreation(AddStructFunction, this);
+							AddStructFunction->SetFromFunction(UChooserFunctionLibrary::StaticClass()->FindFunctionByName(GET_MEMBER_NAME_CHECKED(UChooserFunctionLibrary, AddChooserStructInput)));
+							AddStructFunction->AllocateDefaultPins();
+							ContextStructPin->MakeLinkTo(AddStructFunction->FindPin(FName("Context")));
+
+							PreviousNodeExecOutput->MakeLinkTo(AddStructFunction->GetExecPin());
+							PreviousNodeExecOutput = AddStructFunction->GetThenPin();
 							
-						UEdGraphPin* AddStructPin = AddStructFunction->FindPin(FName("Value"));
-						// not sure why this isn't automatically happening:
-						AddStructPin->PinType.PinCategory = UEdGraphSchema_K2::PC_Struct;
-						AddStructPin->PinType.PinSubCategoryObject = StructContext.Struct;
+							UEdGraphPin* AddStructPin = AddStructFunction->FindPin(FName("Value"));
+							// not sure why this isn't automatically happening:
+							AddStructPin->PinType.PinCategory = UEdGraphSchema_K2::PC_Struct;
+							AddStructPin->PinType.PinSubCategoryObject = StructContext.Struct;
 
-						bool bLinked = false;
-						if (StructContext.Direction != EContextObjectDirection::Write)
-						{
-							if (UEdGraphPin* Pin = FindPin(StructContext.Struct.GetFName(), EGPD_Input))
+							bool bLinked = false;
+							if (StructContext.Direction != EContextObjectDirection::Write)
 							{
-								if (Pin->HasAnyConnections())
+								if (UEdGraphPin* Pin = FindPin(StructContext.Struct.GetFName(), EGPD_Input))
 								{
-									CompilerContext.MovePinLinksToIntermediate(*Pin, *AddStructPin);
-									bLinked = true;
+									if (Pin->HasAnyConnections())
+									{
+										CompilerContext.MovePinLinksToIntermediate(*Pin, *AddStructPin);
+										bLinked = true;
+									}
 								}
 							}
-						}
 						
-						if (!bLinked)
-						{
-							// create a struct to hold the output (or for input structs that were not connected to anything)
-							UK2Node_MakeStruct* OutputStructNode = CompilerContext.SpawnIntermediateNode<UK2Node_MakeStruct>(this, SourceGraph);
-							CompilerContext.MessageLog.NotifyIntermediateObjectCreation(OutputStructNode, this);
-							OutputStructNode->StructType = static_cast<UScriptStruct*>(StructContext.Struct);
-							OutputStructNode->AllocateDefaultPins();
-							OutputStructNode->FindPin(StructContext.Struct->GetFName(), EGPD_Output)->MakeLinkTo(AddStructPin);
-						}
+							if (!bLinked)
+							{
+								// create a struct to hold the output (or for input structs that were not connected to anything)
+								UK2Node_MakeStruct* OutputStructNode = CompilerContext.SpawnIntermediateNode<UK2Node_MakeStruct>(this, SourceGraph);
+								CompilerContext.MessageLog.NotifyIntermediateObjectCreation(OutputStructNode, this);
+								OutputStructNode->StructType = static_cast<UScriptStruct*>(StructContext.Struct);
+								OutputStructNode->AllocateDefaultPins();
+								OutputStructNode->FindPin(StructContext.Struct->GetFName(), EGPD_Output)->MakeLinkTo(AddStructPin);
+							}
 
-						if (StructContext.Direction != EContextObjectDirection::Read)
-						{
-							
-							if (UEdGraphPin* Pin = FindPin(StructContext.Struct.GetFName(), EGPD_Output))
+							if (StructContext.Direction != EContextObjectDirection::Read)
 							{
-								if (Pin->HasAnyConnections())
+								if (UEdGraphPin* Pin = FindPin(StructContext.Struct.GetFName(), EGPD_Output))
 								{
-									// set up struct output pin
-									UK2Node_CallFunction* GetStructFunction = CompilerContext.SpawnIntermediateNode<UK2Node_CallFunction>(this, SourceGraph);
-									CompilerContext.MessageLog.NotifyIntermediateObjectCreation(GetStructFunction, this);
-									GetStructFunction->SetFromFunction(UChooserFunctionLibrary::StaticClass()->FindFunctionByName(GET_MEMBER_NAME_CHECKED(UChooserFunctionLibrary, GetChooserStructOutput)));
-									GetStructFunction->AllocateDefaultPins();
-									GetStructFunction->FindPin(FName("Index"))->DefaultValue = FString::FromInt(ContextDataIndex);
+									if (Pin->HasAnyConnections())
+									{
+										// set up struct output pin
+										UK2Node_CallFunction* GetStructFunction = CompilerContext.SpawnIntermediateNode<UK2Node_CallFunction>(this, SourceGraph);
+										CompilerContext.MessageLog.NotifyIntermediateObjectCreation(GetStructFunction, this);
+										GetStructFunction->SetFromFunction(UChooserFunctionLibrary::StaticClass()->FindFunctionByName(GET_MEMBER_NAME_CHECKED(UChooserFunctionLibrary, GetChooserStructOutput)));
+										GetStructFunction->AllocateDefaultPins();
+										GetStructFunction->FindPin(FName("Index"))->DefaultValue = FString::FromInt(ContextDataIndex);
 									
-									ContextStructPin->MakeLinkTo(GetStructFunction->FindPin(FName("Context")));
-									UEdGraphPin* ValuePin = GetStructFunction->FindPin(FName("Value"));
-									// not sure why this isn't automatically happening:
-									ValuePin->PinType.PinCategory = UEdGraphSchema_K2::PC_Struct;
-									ValuePin->PinType.PinSubCategoryObject = StructContext.Struct;
+										ContextStructPin->MakeLinkTo(GetStructFunction->FindPin(FName("Context")));
+										UEdGraphPin* ValuePin = GetStructFunction->FindPin(FName("Value"));
+										// not sure why this isn't automatically happening:
+										ValuePin->PinType.PinCategory = UEdGraphSchema_K2::PC_Struct;
+										ValuePin->PinType.PinSubCategoryObject = StructContext.Struct;
 					
-									CompilerContext.MovePinLinksToIntermediate(*Pin, *ValuePin);
+										CompilerContext.MovePinLinksToIntermediate(*Pin, *ValuePin);
+									}
 								}
 							}
 						}
@@ -711,14 +715,14 @@ void UK2Node_EvaluateProxy2::ExpandNode(class FKismetCompilerContext& CompilerCo
 			
 			CompilerContext.MovePinLinksToIntermediate(*ExecOutput, *CallFunction->GetThenPin());
 		}
-
+	
 		UK2Node_CallFunction* MakeLookupProxyFunction = CompilerContext.SpawnIntermediateNode<UK2Node_CallFunction>(this, SourceGraph);
 		CompilerContext.MessageLog.NotifyIntermediateObjectCreation(MakeLookupProxyFunction, this);
 		MakeLookupProxyFunction->SetFromFunction(UProxyTableFunctionLibrary::StaticClass()->FindFunctionByName(GET_MEMBER_NAME_CHECKED(UProxyTableFunctionLibrary, MakeLookupProxy)));
 		MakeLookupProxyFunction->AllocateDefaultPins();
 		
-		UEdGraphPin* ChooserTablePin = MakeLookupProxyFunction->FindPin(FName("Proxy"));
-		CallFunction->GetSchema()->TrySetDefaultObject(*ChooserTablePin, Proxy);
+		UEdGraphPin* ProxyPin = MakeLookupProxyFunction->FindPin(FName("Proxy"));
+		MakeLookupProxyFunction->GetSchema()->TrySetDefaultObject(*ProxyPin, Proxy);
 
 		MakeLookupProxyFunction->GetReturnValuePin()->MakeLinkTo(CallFunction->FindPin(FName("ObjectChooser")));
 
