Fix some uninitialized properties in ChooserPropertyAccess

diff --git a/Engine/Plugins/Experimental/Chooser/Source/Chooser/Public/ChooserPropertyAccess.h b/Engine/Plugins/Experimental/Chooser/Source/Chooser/Public/ChooserPropertyAccess.h
--- a/Engine/Plugins/Experimental/Chooser/Source/Chooser/Public/ChooserPropertyAccess.h
+++ b/Engine/Plugins/Experimental/Chooser/Source/Chooser/Public/ChooserPropertyAccess.h
@@ -38,7 +38,7 @@ struct FChooserPropertyBinding
 	TArray<FName> PropertyBindingChain;
 	
 	UPROPERTY()
-	int ContextIndex;
+	int ContextIndex = 0;
 
 #if WITH_EDITORONLY_DATA
 	UPROPERTY()
@@ -94,7 +94,7 @@ struct FContextObjectTypeBase
 	GENERATED_BODY()
 	
 	UPROPERTY(EditAnywhere, Category="Type")
-	EContextObjectDirection Direction;
+	EContextObjectDirection Direction = EContextObjectDirection::Read;
 };
 
 
