Make EvaluateProxy blueprint node take in a Proxy Table reference

diff --git a/Engine/Plugins/Experimental/Chooser/Source/ProxyTable/Internal/LookupProxy.h b/Engine/Plugins/Experimental/Chooser/Source/ProxyTable/Internal/LookupProxy.h
--- a/Engine/Plugins/Experimental/Chooser/Source/ProxyTable/Internal/LookupProxy.h
+++ b/Engine/Plugins/Experimental/Chooser/Source/ProxyTable/Internal/LookupProxy.h
@@ -42,8 +42,22 @@ struct PROXYTABLE_API FLookupProxy : public FObjectChooserBase
 	virtual UObject* ChooseObject(FChooserEvaluationContext& Context) const final override;
 	
 	public:
-	FLookupProxy();
 	
 	UPROPERTY(EditAnywhere, Category="Parameters")
 	TObjectPtr<UProxyAsset> Proxy;
+};
+
+USTRUCT(meta=(Hidden))
+struct PROXYTABLE_API FLookupProxyWithOverrideTable : public FObjectChooserBase
+{
+	GENERATED_BODY()
+	virtual UObject* ChooseObject(FChooserEvaluationContext& Context) const final override;
+	
+	public:
+	
+	UPROPERTY(EditAnywhere, Category="Parameters")
+	TObjectPtr<UProxyAsset> Proxy;
+	
+	UPROPERTY(EditAnywhere, Category="Parameters")
+	TObjectPtr<UProxyTable> OverrideProxyTable;
 };
\ No newline at end of file
diff --git a/Engine/Plugins/Experimental/Chooser/Source/ProxyTable/Private/LookupProxy.cpp b/Engine/Plugins/Experimental/Chooser/Source/ProxyTable/Private/LookupProxy.cpp
--- a/Engine/Plugins/Experimental/Chooser/Source/ProxyTable/Private/LookupProxy.cpp
+++ b/Engine/Plugins/Experimental/Chooser/Source/ProxyTable/Private/LookupProxy.cpp
@@ -5,10 +5,6 @@
 #include "ChooserPropertyAccess.h"
 #include "UObject/Package.h"
 
-FLookupProxy::FLookupProxy()
-{
-}
-
 UObject* FLookupProxy::ChooseObject(FChooserEvaluationContext& Context) const
 {
 	if (Proxy)
@@ -18,6 +14,15 @@ UObject* FLookupProxy::ChooseObject(FChooserEvaluationContext& Context) const
 	return nullptr;
 }
 
+UObject* FLookupProxyWithOverrideTable::ChooseObject(FChooserEvaluationContext& Context) const
+{
+	if (Proxy && OverrideProxyTable)
+	{
+		return OverrideProxyTable->FindProxyObject(Proxy->Guid, Context);
+	}
+	return nullptr;
+}
+
 bool FProxyTableContextProperty::GetValue(FChooserEvaluationContext& Context, const UProxyTable*& OutResult) const
 {
 	const UStruct* StructType = nullptr;
diff --git a/Engine/Plugins/Experimental/Chooser/Source/ProxyTable/Private/ProxyTableFunctionLibrary.cpp b/Engine/Plugins/Experimental/Chooser/Source/ProxyTable/Private/ProxyTableFunctionLibrary.cpp
--- a/Engine/Plugins/Experimental/Chooser/Source/ProxyTable/Private/ProxyTableFunctionLibrary.cpp
+++ b/Engine/Plugins/Experimental/Chooser/Source/ProxyTable/Private/ProxyTableFunctionLibrary.cpp
@@ -49,9 +49,20 @@ UObject* UProxyTableFunctionLibrary::EvaluateProxyTable(const UObject* ContextOb
 
 
 FInstancedStruct UProxyTableFunctionLibrary::MakeLookupProxy(UProxyAsset* Proxy)
- {
+{
  	FInstancedStruct Struct;
  	Struct.InitializeAs(FLookupProxy::StaticStruct());
  	Struct.GetMutable<FLookupProxy>().Proxy = Proxy;
  	return Struct;
- }
+}
+
+FInstancedStruct UProxyTableFunctionLibrary::MakeLookupProxyWithOverrideTable(UProxyAsset* Proxy, UProxyTable* ProxyTable)
+{
+ 	FInstancedStruct Struct;
+ 	Struct.InitializeAs(FLookupProxyWithOverrideTable::StaticStruct());
+ 	FLookupProxyWithOverrideTable& Data = Struct.GetMutable<FLookupProxyWithOverrideTable>();
+	Data.Proxy = Proxy;
+	Data.OverrideProxyTable = ProxyTable;
+ 	return Struct;
+}
+
diff --git a/Engine/Plugins/Experimental/Chooser/Source/ProxyTable/Public/ProxyTableFunctionLibrary.h b/Engine/Plugins/Experimental/Chooser/Source/ProxyTable/Public/ProxyTableFunctionLibrary.h
--- a/Engine/Plugins/Experimental/Chooser/Source/ProxyTable/Public/ProxyTableFunctionLibrary.h
+++ b/Engine/Plugins/Experimental/Chooser/Source/ProxyTable/Public/ProxyTableFunctionLibrary.h
@@ -41,6 +41,9 @@ public:
 	*
 	* @param Chooser				(in) the ChooserTable asset to evaluate
 	*/
-	UFUNCTION(BlueprintPure, Category = "Animation", meta = (BlueprintThreadSafe, BlueprintInternalUseOnly="true", NativeMakeFunc))
+	UFUNCTION(BlueprintPure, Category = "Animation", meta = (BlueprintThreadSafe, BlueprintInternalUseOnly="true"))
 	static FInstancedStruct MakeLookupProxy(UProxyAsset* Proxy);
+	
+	UFUNCTION(BlueprintPure, Category = "Animation", meta = (BlueprintThreadSafe, BlueprintInternalUseOnly="true"))
+	static FInstancedStruct MakeLookupProxyWithOverrideTable(UProxyAsset* Proxy, UProxyTable* ProxyTable);
 };
\ No newline at end of file
diff --git a/Engine/Plugins/Experimental/Chooser/Source/ProxyTableUncooked/Private/EvaluateProxyNode.cpp b/Engine/Plugins/Experimental/Chooser/Source/ProxyTableUncooked/Private/EvaluateProxyNode.cpp
--- a/Engine/Plugins/Experimental/Chooser/Source/ProxyTableUncooked/Private/EvaluateProxyNode.cpp
+++ b/Engine/Plugins/Experimental/Chooser/Source/ProxyTableUncooked/Private/EvaluateProxyNode.cpp
@@ -465,6 +465,12 @@ void UK2Node_EvaluateProxy2::AllocateDefaultPins()
    		ResultType = Proxy->Type;
    	}
 	
+	static const FName ProxyTablePinName = "ProxyTable";
+	if (FindPin(ProxyTablePinName, EGPD_Input) == nullptr)
+	{
+		CreatePin(EGPD_Input, UEdGraphSchema_K2::PC_Object, UProxyTable::StaticClass(), ProxyTablePinName);
+	}
+
 	static const FName ResultPinName = "Result";
 	if (UEdGraphPin* ResultPin = FindPin(ResultPinName, EGPD_Output))
 	{
@@ -589,6 +595,7 @@ void UK2Node_EvaluateProxy2::ExpandNode(class FKismetCompilerContext& CompilerCo
 		UEdGraphPin* ContextStructPin = ContextStructNode->GetReturnValuePin();
 
 		ContextStructPin->MakeLinkTo(CallFunction->FindPin(FName("Context")));
+		PreviousNodeExecOutput = ContextStructNode->GetThenPin();
 
 		UK2Node_Self* SelfNode = CompilerContext.SpawnIntermediateNode<UK2Node_Self>(this, SourceGraph);
 		CompilerContext.MessageLog.NotifyIntermediateObjectCreation(SelfNode, this);
@@ -715,11 +722,27 @@ void UK2Node_EvaluateProxy2::ExpandNode(class FKismetCompilerContext& CompilerCo
 			
 			CompilerContext.MovePinLinksToIntermediate(*ExecOutput, *CallFunction->GetThenPin());
 		}
-	
-		UK2Node_CallFunction* MakeLookupProxyFunction = CompilerContext.SpawnIntermediateNode<UK2Node_CallFunction>(this, SourceGraph);
-		CompilerContext.MessageLog.NotifyIntermediateObjectCreation(MakeLookupProxyFunction, this);
-		MakeLookupProxyFunction->SetFromFunction(UProxyTableFunctionLibrary::StaticClass()->FindFunctionByName(GET_MEMBER_NAME_CHECKED(UProxyTableFunctionLibrary, MakeLookupProxy)));
-		MakeLookupProxyFunction->AllocateDefaultPins();
+		
+		UK2Node_CallFunction* MakeLookupProxyFunction = nullptr;
+
+		UEdGraphPin* OverrideTablePin = FindPinChecked(TEXT("ProxyTable"));
+		if (OverrideTablePin->HasAnyConnections() || OverrideTablePin->DefaultObject)
+		{
+			MakeLookupProxyFunction = CompilerContext.SpawnIntermediateNode<UK2Node_CallFunction>(this, SourceGraph);
+   			MakeLookupProxyFunction->SetFromFunction(UProxyTableFunctionLibrary::StaticClass()->FindFunctionByName(GET_MEMBER_NAME_CHECKED(UProxyTableFunctionLibrary, MakeLookupProxyWithOverrideTable)));
+       		MakeLookupProxyFunction->AllocateDefaultPins();
+
+			UEdGraphPin* TargetOverrideTablePin = MakeLookupProxyFunction->FindPinChecked(TEXT("ProxyTable"));
+			CompilerContext.MovePinLinksToIntermediate(*OverrideTablePin, *TargetOverrideTablePin);
+		}
+		else
+		{
+			MakeLookupProxyFunction = CompilerContext.SpawnIntermediateNode<UK2Node_CallFunction>(this, SourceGraph);
+			MakeLookupProxyFunction->SetFromFunction(UProxyTableFunctionLibrary::StaticClass()->FindFunctionByName(GET_MEMBER_NAME_CHECKED(UProxyTableFunctionLibrary, MakeLookupProxy)));
+			MakeLookupProxyFunction->AllocateDefaultPins();
+		}
+		
+   		CompilerContext.MessageLog.NotifyIntermediateObjectCreation(MakeLookupProxyFunction, this);
 		
 		UEdGraphPin* ProxyPin = MakeLookupProxyFunction->FindPin(FName("Proxy"));
 		MakeLookupProxyFunction->GetSchema()->TrySetDefaultObject(*ProxyPin, Proxy);
