Make choosers and proxy tables support AnimInstances (and other Class types) better.

diff --git a/Engine/Plugins/Experimental/Chooser/Source/Chooser/Internal/Chooser.h b/Engine/Plugins/Experimental/Chooser/Source/Chooser/Internal/Chooser.h
--- a/Engine/Plugins/Experimental/Chooser/Source/Chooser/Internal/Chooser.h
+++ b/Engine/Plugins/Experimental/Chooser/Source/Chooser/Internal/Chooser.h
@@ -39,7 +39,8 @@ public:
 
 private: 
 	// caching the OutputObjectType and ContextObjectType so that on Undo, we can tell if we should fire the changed delegate
-	UClass* CachedPreviousOutputObjectType;
+	UClass* CachedPreviousOutputObjectType = nullptr;
+	EObjectChooserResultType CachedPreviousResultType = EObjectChooserResultType::ObjectResult;
 	
 	// objects this chooser has been recently evaluated on
 	mutable TSet<TWeakObjectPtr<const UObject>> RecentContextObjects;
@@ -78,6 +79,8 @@ public:
 
 	UPROPERTY(EditAnywhere, Category="Output", Meta = (AllowAbstract=true))
 	TObjectPtr<UClass> OutputObjectType;
+	UPROPERTY(EditAnywhere, Category="Output")
+	EObjectChooserResultType ResultType = EObjectChooserResultType::ObjectResult;
 
 	virtual TConstArrayView<FInstancedStruct> GetContextData() const override { return ContextData; }
 	
diff --git a/Engine/Plugins/Experimental/Chooser/Source/Chooser/Private/Chooser.cpp b/Engine/Plugins/Experimental/Chooser/Source/Chooser/Private/Chooser.cpp
--- a/Engine/Plugins/Experimental/Chooser/Source/Chooser/Private/Chooser.cpp
+++ b/Engine/Plugins/Experimental/Chooser/Source/Chooser/Private/Chooser.cpp
@@ -16,10 +16,11 @@ void UChooserTable::PostEditUndo()
 {
 	UObject::PostEditUndo();
 
-	if (CachedPreviousOutputObjectType != OutputObjectType)
+	if (CachedPreviousOutputObjectType != OutputObjectType || CachedPreviousResultType != ResultType)
 	{
 		OnOutputObjectTypeChanged.Broadcast(OutputObjectType);
 		CachedPreviousOutputObjectType = OutputObjectType;
+		CachedPreviousResultType = ResultType;
 	}
 	OnContextClassChanged.Broadcast();
 }
@@ -29,6 +30,7 @@ void UChooserTable::PostEditChangeProperty(FPropertyChangedEvent& PropertyChange
 	UObject::PostEditChangeProperty(PropertyChangedEvent);
 	
 	static FName OutputObjectTypeName = "OutputObjectType";
+	static FName ResultTypeName = "ResultType";
 	if (PropertyChangedEvent.Property->GetName() == OutputObjectTypeName)
 	{
 		if (CachedPreviousOutputObjectType != OutputObjectType)
@@ -37,6 +39,14 @@ void UChooserTable::PostEditChangeProperty(FPropertyChangedEvent& PropertyChange
 			CachedPreviousOutputObjectType = OutputObjectType;
 		}
 	}
+	else if (PropertyChangedEvent.Property->GetName() == ResultTypeName)
+	{
+		if (CachedPreviousResultType != ResultType)
+		{
+			OnOutputObjectTypeChanged.Broadcast(OutputObjectType);
+			CachedPreviousResultType = ResultType;
+		}
+	}
 	else
 	{
 		OnContextClassChanged.Broadcast();
@@ -48,6 +58,7 @@ void UChooserTable::PostLoad()
 	Super::PostLoad();
 
 	CachedPreviousOutputObjectType = OutputObjectType;
+	CachedPreviousResultType = ResultType;
 
 	// convert old data if it exists
 
diff --git a/Engine/Plugins/Experimental/Chooser/Source/Chooser/Private/ChooserFunctionLibrary.cpp b/Engine/Plugins/Experimental/Chooser/Source/Chooser/Private/ChooserFunctionLibrary.cpp
--- a/Engine/Plugins/Experimental/Chooser/Source/Chooser/Private/ChooserFunctionLibrary.cpp
+++ b/Engine/Plugins/Experimental/Chooser/Source/Chooser/Private/ChooserFunctionLibrary.cpp
@@ -118,18 +118,36 @@ DEFINE_FUNCTION(UChooserFunctionLibrary::execGetChooserStructOutput)
 	}
 }
 
-UObject* UChooserFunctionLibrary::EvaluateObjectChooserBase(FChooserEvaluationContext& Context, const FInstancedStruct& ObjectChooser, TSubclassOf<UObject> ObjectClass)
+UObject* UChooserFunctionLibrary::EvaluateObjectChooserBase(FChooserEvaluationContext& Context, const FInstancedStruct& ObjectChooser, TSubclassOf<UObject> ObjectClass, bool bResultIsClass)
  {
  	UObject* Result = nullptr;
 
 	if (const FObjectChooserBase* ObjectChooserPtr = ObjectChooser.GetPtr<FObjectChooserBase>())
 	{
-		ObjectChooserPtr->ChooseMulti(Context, FObjectChooserBase::FObjectChooserIteratorCallback::CreateLambda([&Result, ObjectClass](UObject* InResult)
+		ObjectChooserPtr->ChooseMulti(Context, FObjectChooserBase::FObjectChooserIteratorCallback::CreateLambda([&Result, bResultIsClass, ObjectClass](UObject* InResult)
 		{
-			if (InResult && (ObjectClass == nullptr || InResult->IsA(ObjectClass)))
+			if (InResult)
 			{
-				Result = InResult;
-				return FObjectChooserBase::EIteratorStatus::Stop;
+				bool bTypesMatch = true;
+
+				if (ObjectClass && InResult)
+				{
+					if (bResultIsClass)
+					{
+						UClass* ResultClass =  Cast<UClass>(InResult);
+						bTypesMatch = ResultClass && ResultClass->IsChildOf(ObjectClass);
+					}
+					else
+					{
+						bTypesMatch = InResult->IsA(ObjectClass);
+					}
+				}
+
+				if (bTypesMatch)
+				{
+					Result = InResult;
+					return FObjectChooserBase::EIteratorStatus::Stop;
+				}
 			}
 			return FObjectChooserBase::EIteratorStatus::Continue;
 		}));
@@ -139,15 +157,30 @@ UObject* UChooserFunctionLibrary::EvaluateObjectChooserBase(FChooserEvaluationCo
  }
 
 
-TArray<UObject*> UChooserFunctionLibrary::EvaluateObjectChooserBaseMulti(FChooserEvaluationContext& Context, const FInstancedStruct& ObjectChooser, TSubclassOf<UObject> ObjectClass)
+TArray<UObject*> UChooserFunctionLibrary::EvaluateObjectChooserBaseMulti(FChooserEvaluationContext& Context, const FInstancedStruct& ObjectChooser, TSubclassOf<UObject> ObjectClass, bool bResultIsClass)
 {
  	TArray<UObject*> Result;
 
 	if (const FObjectChooserBase* ObjectChooserPtr = ObjectChooser.GetPtr<FObjectChooserBase>())
 	{
-		ObjectChooserPtr->ChooseMulti(Context, FObjectChooserBase::FObjectChooserIteratorCallback::CreateLambda([&Result, ObjectClass](UObject* InResult)
+		ObjectChooserPtr->ChooseMulti(Context, FObjectChooserBase::FObjectChooserIteratorCallback::CreateLambda([&Result, bResultIsClass, ObjectClass](UObject* InResult)
 		{
-			if (InResult && (ObjectClass == nullptr || InResult->IsA(ObjectClass)))
+			bool bTypesMatch = true;
+            
+			if (ObjectClass && InResult)
+			{
+				if (bResultIsClass)
+				{
+					UClass* ResultClass =  Cast<UClass>(InResult);
+					bTypesMatch = ResultClass && ResultClass->IsChildOf(ObjectClass);
+				}
+				else
+				{
+					bTypesMatch = InResult->IsA(ObjectClass);
+				}
+			}
+            
+			if (bTypesMatch)
 			{
 				Result.Add(InResult);
 			}
diff --git a/Engine/Plugins/Experimental/Chooser/Source/Chooser/Public/ChooserFunctionLibrary.h b/Engine/Plugins/Experimental/Chooser/Source/Chooser/Public/ChooserFunctionLibrary.h
--- a/Engine/Plugins/Experimental/Chooser/Source/Chooser/Public/ChooserFunctionLibrary.h
+++ b/Engine/Plugins/Experimental/Chooser/Source/Chooser/Public/ChooserFunctionLibrary.h
@@ -42,20 +42,22 @@ public:
 	*
 	* @param Context			(in) A struct reference to the chooser evaluation context
 	* @param ObjectChooser		(in) An Instanced struct containing an ObjectChooserBase implementation, such as EvaluateChooser, or EvaluateProxyAsset
-	* @param ObjectClass		(in) Expected type of result object
+	* @param ObjectClass		(in) Expected type of result object (or the type of UClass if bResultIsClass is true)
+	* @param bResultIsClass		(in) The Object being returned is a UClass, and the ObjectClass parameter indicates what it must be a subclass of
 	*/
 	UFUNCTION(BlueprintCallable, meta = (BlueprintThreadSafe, BlueprintInternalUseOnly = "true", DeterminesOutputType = "ObjectClass"))
-	static UObject* EvaluateObjectChooserBase(UPARAM(Ref) FChooserEvaluationContext& Context,UPARAM(Ref) const FInstancedStruct& ObjectChooser, TSubclassOf<UObject> ObjectClass);
+	static UObject* EvaluateObjectChooserBase(UPARAM(Ref) FChooserEvaluationContext& Context,UPARAM(Ref) const FInstancedStruct& ObjectChooser, TSubclassOf<UObject> ObjectClass, bool bResultIsClass = false);
 
 	/**
 	* Evaluate a chooser table and return all selected UObjects
 	*
 	* @param Context			(in) A struct reference to the chooser evaluation context
 	* @param ObjectChooser		(in) An Instanced struct containing an ObjectChooserBase implementation, such as EvaluateChooser, or EvaluateProxyAsset
-	* @param ObjectClass		(in) Expected type of result object
+	* @param ObjectClass		(in) Expected type of result object (or the type of UClass if bResultIsClass is true)
+	* @param bResultIsClass		(in) The Object being returned is a UClass, and the ObjectClass parameter indicates what it must be a subclass of
 	*/
 	UFUNCTION(BlueprintCallable, meta = (BlueprintThreadSafe, BlueprintInternalUseOnly = "true", DeterminesOutputType = "ObjectClass"))
-	static TArray<UObject*> EvaluateObjectChooserBaseMulti(UPARAM(Ref) FChooserEvaluationContext& Context,UPARAM(Ref) const FInstancedStruct& ObjectChooser, TSubclassOf<UObject> ObjectClass);
+	static TArray<UObject*> EvaluateObjectChooserBaseMulti(UPARAM(Ref) FChooserEvaluationContext& Context,UPARAM(Ref) const FInstancedStruct& ObjectChooser, TSubclassOf<UObject> ObjectClass, bool bResultIsClass = false);
 
 
 	/**
diff --git a/Engine/Plugins/Experimental/Chooser/Source/Chooser/Public/ChooserPropertyAccess.h b/Engine/Plugins/Experimental/Chooser/Source/Chooser/Public/ChooserPropertyAccess.h
--- a/Engine/Plugins/Experimental/Chooser/Source/Chooser/Public/ChooserPropertyAccess.h
+++ b/Engine/Plugins/Experimental/Chooser/Source/Chooser/Public/ChooserPropertyAccess.h
@@ -21,6 +21,14 @@ class UHasContextClass : public UInterface
 
 DECLARE_MULTICAST_DELEGATE(FContextClassChanged)
 
+
+UENUM()
+enum class EObjectChooserResultType
+{
+	ObjectResult,
+	ClassResult,
+};
+
 class IHasContextClass
 {
 	GENERATED_BODY()
@@ -38,7 +46,7 @@ struct FChooserPropertyBinding
 	TArray<FName> PropertyBindingChain;
 	
 	UPROPERTY()
-	int ContextIndex = 0;
+	int ContextIndex;
 
 #if WITH_EDITORONLY_DATA
 	UPROPERTY()
@@ -94,7 +102,7 @@ struct FContextObjectTypeBase
 	GENERATED_BODY()
 	
 	UPROPERTY(EditAnywhere, Category="Type")
-	EContextObjectDirection Direction = EContextObjectDirection::Read;
+	EContextObjectDirection Direction;
 };
 
 
diff --git a/Engine/Plugins/Experimental/Chooser/Source/ChooserEditor/Private/ChooserTableEditor.cpp b/Engine/Plugins/Experimental/Chooser/Source/ChooserEditor/Private/ChooserTableEditor.cpp
--- a/Engine/Plugins/Experimental/Chooser/Source/ChooserEditor/Private/ChooserTableEditor.cpp
+++ b/Engine/Plugins/Experimental/Chooser/Source/ChooserEditor/Private/ChooserTableEditor.cpp
@@ -35,6 +35,8 @@
 #include "ScopedTransaction.h"
 #include "Widgets/Layout/SSeparator.h"
 #include "ChooserTableEditorCommands.h"
+#include "Animation/AnimBlueprintGeneratedClass.h"
+#include "Animation/AnimInstance.h"
 #include "Widgets/Layout/SScrollBox.h"
 
 #define LOCTEXT_NAMESPACE "ChooserEditor"
diff --git a/Engine/Plugins/Experimental/Chooser/Source/ChooserUncooked/Private/EvaluateChooserNode.cpp b/Engine/Plugins/Experimental/Chooser/Source/ChooserUncooked/Private/EvaluateChooserNode.cpp
--- a/Engine/Plugins/Experimental/Chooser/Source/ChooserUncooked/Private/EvaluateChooserNode.cpp
+++ b/Engine/Plugins/Experimental/Chooser/Source/ChooserUncooked/Private/EvaluateChooserNode.cpp
@@ -387,7 +387,7 @@ void UK2Node_EvaluateChooser2::AllocateDefaultPins()
 {
 	Super::AllocateDefaultPins();
 	
-	UClass* ChooserResultType = UObject::StaticClass();
+	UClass* ResultType = UObject::StaticClass();
 
 	if (!FindPin(UEdGraphSchema_K2::PN_Execute, EGPD_Input))
 	{
@@ -477,24 +477,26 @@ void UK2Node_EvaluateChooser2::AllocateDefaultPins()
 		}
 	}
 
+	bool bResultIsClass = false;
    	if (Chooser && Chooser->OutputObjectType)
    	{
-   		ChooserResultType = Chooser->OutputObjectType;
+   		ResultType = Chooser->OutputObjectType;
+   		bResultIsClass = Chooser->ResultType == EObjectChooserResultType::ClassResult;
    	}
 
-	static const FName ResultPinName = "Result";
-	if (UEdGraphPin* ResultPin = FindPin(ResultPinName, EGPD_Output))
+	if (UEdGraphPin* ResultPin = FindPin(TEXT("Result"), EGPD_Output))
 	{
-		ResultPin->PinType.PinSubCategoryObject = ChooserResultType;
+		ResultPin->PinType.PinCategory = bResultIsClass ? UEdGraphSchema_K2::PC_Class : UEdGraphSchema_K2::PC_Object;
+		ResultPin->PinType.PinSubCategoryObject = ResultType;
 		ResultPin->PinType.ContainerType = (Mode == EEvaluateChooserMode::AllResults) ? EPinContainerType::Array : EPinContainerType::None;
 	}
 	else
 	{
 		UEdGraphNode::FCreatePinParams PinParams;
 		PinParams.ContainerType = (Mode == EEvaluateChooserMode::AllResults) ? EPinContainerType::Array : EPinContainerType::None;
-		PinParams.ValueTerminalType.TerminalCategory = UEdGraphSchema_K2::PC_Object;
+		PinParams.ValueTerminalType.TerminalCategory =  bResultIsClass ? UEdGraphSchema_K2::PC_Class : UEdGraphSchema_K2::PC_Object;
 				
-		CreatePin(EGPD_Output, UEdGraphSchema_K2::PC_Object, ChooserResultType, TEXT("Result"), PinParams);
+		CreatePin(EGPD_Output, PinParams.ValueTerminalType.TerminalCategory, ResultType, TEXT("Result"), PinParams);
 	}
 }
 
@@ -748,12 +750,24 @@ void UK2Node_EvaluateChooser2::ExpandNode(class FKismetCompilerContext& Compiler
 
 		UEdGraphPin* OutputPin = CallFunction->GetReturnValuePin();
 
-		if (Chooser && Chooser->OutputObjectType)
+		if (Chooser) 
 		{
-			UEdGraphPin* OutputClassPin = CallFunction->FindPin(TEXT("ObjectClass"));
-			CallFunction->GetSchema()->TrySetDefaultObject(*OutputClassPin, Chooser->OutputObjectType);
+			if (Chooser->OutputObjectType)
+			{
+				if (UEdGraphPin* OutputClassPin = CallFunction->FindPin(TEXT("ObjectClass")))
+				{
+					CallFunction->GetSchema()->TrySetDefaultObject(*OutputClassPin, Chooser->OutputObjectType);
+				}
+			}
+
+			if (UEdGraphPin* ResultIsClassPin = CallFunction->FindPin(TEXT("bResultIsClass")))
+			{
+				CallFunction->GetSchema()->TrySetDefaultValue(*ResultIsClassPin, Chooser->ResultType == EObjectChooserResultType::ClassResult ? TEXT("true") : TEXT("false"));
+			}
 		}
 
+		
+
 		CompilerContext.MovePinLinksToIntermediate(*ResultPin, *OutputPin);
 	}
 	
diff --git a/Engine/Plugins/Experimental/Chooser/Source/ProxyTable/Private/ProxyAsset.cpp b/Engine/Plugins/Experimental/Chooser/Source/ProxyTable/Private/ProxyAsset.cpp
--- a/Engine/Plugins/Experimental/Chooser/Source/ProxyTable/Private/ProxyAsset.cpp
+++ b/Engine/Plugins/Experimental/Chooser/Source/ProxyTable/Private/ProxyAsset.cpp
@@ -11,10 +11,11 @@ void UProxyAsset::PostEditUndo()
 {
 	UObject::PostEditUndo();
 
-	if (CachedPreviousType != Type)
+	if (CachedPreviousType != Type || CachedPreviousResultType != ResultType)
 	{
 		OnTypeChanged.Broadcast(Type);
 		CachedPreviousType = Type;
+		CachedPreviousResultType = ResultType;
 	}
 	
 	OnContextClassChanged.Broadcast();
@@ -25,7 +26,7 @@ void UProxyAsset::PostEditChangeProperty(FPropertyChangedEvent& PropertyChangedE
 	UObject::PostEditChangeProperty(PropertyChangedEvent);
 	
 	static FName TypeName = "Type";
-	static FName ContextClassName = "ContextData";
+	static FName ResultTypeName = "ResultType";
 	if (PropertyChangedEvent.Property->GetName() == TypeName)
 	{
 		if (CachedPreviousType != Type)
@@ -34,6 +35,14 @@ void UProxyAsset::PostEditChangeProperty(FPropertyChangedEvent& PropertyChangedE
 		}
 		CachedPreviousType = Type;
 	}
+	else if (PropertyChangedEvent.Property->GetName() == ResultTypeName)
+	{
+		if (CachedPreviousResultType != ResultType)
+		{
+			OnTypeChanged.Broadcast(Type);
+			CachedPreviousResultType = ResultType;
+		}
+	}
 	else
 	{
 		OnContextClassChanged.Broadcast();
@@ -53,6 +62,7 @@ void UProxyAsset::PostLoad()
 
 #if WITH_EDITOR
 	CachedPreviousType = Type;
+	CachedPreviousResultType = ResultType;
 #endif
 	
 	if (ContextClass_DEPRECATED)
diff --git a/Engine/Plugins/Experimental/Chooser/Source/ProxyTable/Public/ProxyAsset.h b/Engine/Plugins/Experimental/Chooser/Source/ProxyTable/Public/ProxyAsset.h
--- a/Engine/Plugins/Experimental/Chooser/Source/ProxyTable/Public/ProxyAsset.h
+++ b/Engine/Plugins/Experimental/Chooser/Source/ProxyTable/Public/ProxyAsset.h
@@ -12,6 +12,7 @@
 
 DECLARE_MULTICAST_DELEGATE_OneParam(FProxyTypeChanged, const UClass* OutputObjectType);
 
+
 UCLASS(MinimalAPI,BlueprintType)
 class UProxyAsset : public UObject, public IHasContextClass
 {
@@ -26,11 +27,14 @@ public:
 
 	// caching the type so that on Undo, we can tell if we should fire the changed delegate
 	UClass* CachedPreviousType = nullptr;
+	EObjectChooserResultType CachedPreviousResultType = EObjectChooserResultType::ObjectResult;
 #endif
 	
 
 	UPROPERTY(EditAnywhere, Category = "Proxy", Meta = (AllowAbstract=true))
 	TObjectPtr<UClass> Type;
+	UPROPERTY(EditAnywhere, Category = "Proxy")
+	EObjectChooserResultType ResultType = EObjectChooserResultType::ObjectResult;
 	
 	UPROPERTY(EditAnywhere, Meta = (ExcludeBaseStruct, BaseStruct = "/Script/Chooser.ContextObjectTypeBase"), Category = "Input")
 	TArray<FInstancedStruct> ContextData;
diff --git a/Engine/Plugins/Experimental/Chooser/Source/ProxyTableEditor/Private/ProxyTableEditor.cpp b/Engine/Plugins/Experimental/Chooser/Source/ProxyTableEditor/Private/ProxyTableEditor.cpp
--- a/Engine/Plugins/Experimental/Chooser/Source/ProxyTableEditor/Private/ProxyTableEditor.cpp
+++ b/Engine/Plugins/Experimental/Chooser/Source/ProxyTableEditor/Private/ProxyTableEditor.cpp
@@ -492,7 +492,7 @@ public:
 						ChooserEditor::FObjectChooserWidgetFactories::CreateWidget(false, Row->ProxyTable, FObjectChooserBase::StaticStruct(),
 								Row->ProxyTable->Entries[Row->RowIndex].ValueStruct.GetMutableMemory(),
 								Row->ProxyTable->Entries[Row->RowIndex].ValueStruct.GetScriptStruct(),
-								Row->ProxyTable->Entries[Row->RowIndex].Proxy->Type,
+								Row->ProxyTable->Entries[Row->RowIndex].Proxy ? Row->ProxyTable->Entries[Row->RowIndex].Proxy->Type.Get() : UObject::StaticClass(),
 								FOnStructPicked(), &CacheBorder);
 					}),
 					&CacheBorder
diff --git a/Engine/Plugins/Experimental/Chooser/Source/ProxyTableUncooked/Private/EvaluateProxyNode.cpp b/Engine/Plugins/Experimental/Chooser/Source/ProxyTableUncooked/Private/EvaluateProxyNode.cpp
--- a/Engine/Plugins/Experimental/Chooser/Source/ProxyTableUncooked/Private/EvaluateProxyNode.cpp
+++ b/Engine/Plugins/Experimental/Chooser/Source/ProxyTableUncooked/Private/EvaluateProxyNode.cpp
@@ -460,9 +460,11 @@ void UK2Node_EvaluateProxy2::AllocateDefaultPins()
 		}
 	}
 
+	bool bResultIsClass = false;
    	if (Proxy && Proxy->Type)
    	{
    		ResultType = Proxy->Type;
+   		bResultIsClass = Proxy->ResultType == EObjectChooserResultType::ClassResult;
    	}
 	
 	static const FName ProxyTablePinName = "ProxyTable";
@@ -471,9 +473,9 @@ void UK2Node_EvaluateProxy2::AllocateDefaultPins()
 		CreatePin(EGPD_Input, UEdGraphSchema_K2::PC_Object, UProxyTable::StaticClass(), ProxyTablePinName);
 	}
 
-	static const FName ResultPinName = "Result";
-	if (UEdGraphPin* ResultPin = FindPin(ResultPinName, EGPD_Output))
+	if (UEdGraphPin* ResultPin = FindPin(TEXT("Result"), EGPD_Output))
 	{
+		ResultPin->PinType.PinCategory = bResultIsClass ? UEdGraphSchema_K2::PC_Class : UEdGraphSchema_K2::PC_Object;
 		ResultPin->PinType.PinSubCategoryObject = ResultType;
 		ResultPin->PinType.ContainerType = (Mode == EEvaluateProxyMode::AllResults) ? EPinContainerType::Array : EPinContainerType::None;
 	}
@@ -481,9 +483,9 @@ void UK2Node_EvaluateProxy2::AllocateDefaultPins()
 	{
 		UEdGraphNode::FCreatePinParams PinParams;
 		PinParams.ContainerType = (Mode == EEvaluateProxyMode::AllResults) ? EPinContainerType::Array : EPinContainerType::None;
-		PinParams.ValueTerminalType.TerminalCategory = UEdGraphSchema_K2::PC_Object;
+		PinParams.ValueTerminalType.TerminalCategory =  bResultIsClass ? UEdGraphSchema_K2::PC_Class : UEdGraphSchema_K2::PC_Object;
 				
-		CreatePin(EGPD_Output, UEdGraphSchema_K2::PC_Object, ResultType, TEXT("Result"), PinParams);
+		CreatePin(EGPD_Output, PinParams.ValueTerminalType.TerminalCategory, ResultType, TEXT("Result"), PinParams);
 	}
 }
 
@@ -634,10 +636,10 @@ void UK2Node_EvaluateProxy2::ExpandNode(class FKismetCompilerContext& CompilerCo
 									CompilerContext.MovePinLinksToIntermediate(*Pin, *AddObjectPin);
 								}
 								else // would be nice to check that self is the same type as this Object parameter
-									{
+								{
 									// auto connect self node to any disconnected object pin
 									SelfPin->MakeLinkTo(AddObjectPin);
-									}
+								}
 							}
 						}
 					}
@@ -751,10 +753,20 @@ void UK2Node_EvaluateProxy2::ExpandNode(class FKismetCompilerContext& CompilerCo
 
 		UEdGraphPin* OutputPin = CallFunction->GetReturnValuePin();
 
-		if (Proxy && Proxy->Type)
+		if (Proxy) 
 		{
-			UEdGraphPin* OutputClassPin = CallFunction->FindPin(TEXT("ObjectClass"));
-			CallFunction->GetSchema()->TrySetDefaultObject(*OutputClassPin, Proxy->Type);
+			if (Proxy->Type)
+			{
+				if (UEdGraphPin* OutputClassPin = CallFunction->FindPin(TEXT("ObjectClass")))
+				{
+					CallFunction->GetSchema()->TrySetDefaultObject(*OutputClassPin, Proxy->Type);
+				}
+			}
+
+			if (UEdGraphPin* ResultIsClassPin = CallFunction->FindPin(TEXT("bResultIsClass")))
+			{
+				CallFunction->GetSchema()->TrySetDefaultValue(*ResultIsClassPin, Proxy->ResultType == EObjectChooserResultType::ClassResult ? TEXT("true") : TEXT("false"));
+			}
 		}
 
 		CompilerContext.MovePinLinksToIntermediate(*ResultPin, *OutputPin);
