add NoClear to FInstancedStructs in the chooser system to prevent users from setting column inputs etc to null

diff --git a/Engine/Plugins/Experimental/Chooser/Source/Chooser/Internal/BoolColumn.h b/Engine/Plugins/Experimental/Chooser/Source/Chooser/Internal/BoolColumn.h
--- a/Engine/Plugins/Experimental/Chooser/Source/Chooser/Internal/BoolColumn.h
+++ b/Engine/Plugins/Experimental/Chooser/Source/Chooser/Internal/BoolColumn.h
@@ -71,7 +71,7 @@ struct CHOOSER_API FBoolColumn : public FChooserColumnBase
 	public:
 	FBoolColumn();
 	
-	UPROPERTY(EditAnywhere, Meta = (ExcludeBaseStruct, BaseStruct = "/Script/Chooser.ChooserParameterBoolBase"), Category = "Data")
+	UPROPERTY(EditAnywhere, NoClear, Meta = (ExcludeBaseStruct, BaseStruct = "/Script/Chooser.ChooserParameterBoolBase"), Category = "Data")
 	FInstancedStruct InputValue;
 
 #if WITH_EDITORONLY_DATA
diff --git a/Engine/Plugins/Experimental/Chooser/Source/Chooser/Internal/Chooser.h b/Engine/Plugins/Experimental/Chooser/Source/Chooser/Internal/Chooser.h
--- a/Engine/Plugins/Experimental/Chooser/Source/Chooser/Internal/Chooser.h
+++ b/Engine/Plugins/Experimental/Chooser/Source/Chooser/Internal/Chooser.h
@@ -67,14 +67,14 @@ public:
 #endif
 
 	// Each possible Result (Rows of chooser table)
-	UPROPERTY(EditAnywhere, DisplayName = "Results", Meta = (ExcludeBaseStruct, BaseStruct = "/Script/Chooser.ObjectChooserBase"), Category = "Hidden")
+	UPROPERTY(EditAnywhere, NoClear, DisplayName = "Results", Meta = (ExcludeBaseStruct, BaseStruct = "/Script/Chooser.ObjectChooserBase"), Category = "Hidden")
 	TArray<FInstancedStruct> ResultsStructs;
 
 	// Columns which filter Results
-	UPROPERTY(EditAnywhere, DisplayName = "Columns", Category = Hidden, meta = (ExcludeBaseStruct, BaseStruct = "/Script/Chooser.ChooserColumnBase"))
+	UPROPERTY(EditAnywhere, NoClear, DisplayName = "Columns", Category = Hidden, meta = (ExcludeBaseStruct, BaseStruct = "/Script/Chooser.ChooserColumnBase"))
 	TArray<FInstancedStruct> ColumnsStructs;
 	
-	UPROPERTY(EditAnywhere, Meta = (ExcludeBaseStruct, BaseStruct = "/Script/Chooser.ContextObjectTypeBase"), Category = "Input")
+	UPROPERTY(EditAnywhere, NoClear, Meta = (ExcludeBaseStruct, BaseStruct = "/Script/Chooser.ContextObjectTypeBase"), Category = "Input")
 	TArray<FInstancedStruct> ContextData;
 
 	UPROPERTY(EditAnywhere, Category="Output", Meta = (AllowAbstract=true))
diff --git a/Engine/Plugins/Experimental/Chooser/Source/Chooser/Internal/EnumColumn.h b/Engine/Plugins/Experimental/Chooser/Source/Chooser/Internal/EnumColumn.h
--- a/Engine/Plugins/Experimental/Chooser/Source/Chooser/Internal/EnumColumn.h
+++ b/Engine/Plugins/Experimental/Chooser/Source/Chooser/Internal/EnumColumn.h
@@ -96,7 +96,7 @@ struct CHOOSER_API FEnumColumn : public FChooserColumnBase
 public:
 	FEnumColumn();
 
-	UPROPERTY(EditAnywhere, Meta = (ExcludeBaseStruct, BaseStruct = "/Script/Chooser.ChooserParameterEnumBase"), Category = "Data")
+	UPROPERTY(EditAnywhere, NoClear, Meta = (ExcludeBaseStruct, BaseStruct = "/Script/Chooser.ChooserParameterEnumBase"), Category = "Data")
 	FInstancedStruct InputValue;
 
 #if WITH_EDITORONLY_DATA
diff --git a/Engine/Plugins/Experimental/Chooser/Source/Chooser/Internal/FloatRangeColumn.h b/Engine/Plugins/Experimental/Chooser/Source/Chooser/Internal/FloatRangeColumn.h
--- a/Engine/Plugins/Experimental/Chooser/Source/Chooser/Internal/FloatRangeColumn.h
+++ b/Engine/Plugins/Experimental/Chooser/Source/Chooser/Internal/FloatRangeColumn.h
@@ -75,7 +75,7 @@ struct CHOOSER_API FFloatRangeColumn : public FChooserColumnBase
 	public:
 	FFloatRangeColumn();
 		
-	UPROPERTY(EditAnywhere, Meta = (ExcludeBaseStruct, BaseStruct = "/Script/Chooser.ChooserParameterFloatBase"), Category = "Data")
+	UPROPERTY(EditAnywhere, NoClear, Meta = (ExcludeBaseStruct, BaseStruct = "/Script/Chooser.ChooserParameterFloatBase"), Category = "Data")
 	FInstancedStruct InputValue;
 	
 #if WITH_EDITORONLY_DATA
diff --git a/Engine/Plugins/Experimental/Chooser/Source/Chooser/Internal/GameplayTagColumn.h b/Engine/Plugins/Experimental/Chooser/Source/Chooser/Internal/GameplayTagColumn.h
--- a/Engine/Plugins/Experimental/Chooser/Source/Chooser/Internal/GameplayTagColumn.h
+++ b/Engine/Plugins/Experimental/Chooser/Source/Chooser/Internal/GameplayTagColumn.h
@@ -20,7 +20,7 @@ public:
 	UPROPERTY()
 	TArray<FName> PropertyBindingChain_DEPRECATED;
 	
-	UPROPERTY(EditAnywhere, Meta = (BindingType = "FGameplayTagContainer", BindingColor = "StructPinTypeColor"), Category = "Binding")
+	UPROPERTY(EditAnywhere, NoClear, Meta = (BindingType = "FGameplayTagContainer", BindingColor = "StructPinTypeColor"), Category = "Binding")
 	FChooserPropertyBinding Binding;
 	
 	virtual bool GetValue(FChooserEvaluationContext& Context, const FGameplayTagContainer*& OutResult) const override;
diff --git a/Engine/Plugins/Experimental/Chooser/Source/Chooser/Internal/ObjectColumn.h b/Engine/Plugins/Experimental/Chooser/Source/Chooser/Internal/ObjectColumn.h
--- a/Engine/Plugins/Experimental/Chooser/Source/Chooser/Internal/ObjectColumn.h
+++ b/Engine/Plugins/Experimental/Chooser/Source/Chooser/Internal/ObjectColumn.h
@@ -71,7 +71,7 @@ struct CHOOSER_API FObjectColumn : public FChooserColumnBase
 
 	FObjectColumn();
 
-	UPROPERTY(EditAnywhere, Meta = (ExcludeBaseStruct, BaseStruct = "/Script/Chooser.ChooserParameterObjectBase"), Category = "Data")
+	UPROPERTY(EditAnywhere, NoClear, Meta = (ExcludeBaseStruct, BaseStruct = "/Script/Chooser.ChooserParameterObjectBase"), Category = "Data")
 	FInstancedStruct InputValue;
 
 #if WITH_EDITORONLY_DATA
diff --git a/Engine/Plugins/Experimental/Chooser/Source/Chooser/Internal/OutputBoolColumn.h b/Engine/Plugins/Experimental/Chooser/Source/Chooser/Internal/OutputBoolColumn.h
--- a/Engine/Plugins/Experimental/Chooser/Source/Chooser/Internal/OutputBoolColumn.h
+++ b/Engine/Plugins/Experimental/Chooser/Source/Chooser/Internal/OutputBoolColumn.h
@@ -18,7 +18,7 @@ struct CHOOSER_API FOutputBoolColumn : public FChooserColumnBase
 	virtual bool HasOutputs() const override { return true; }
 	virtual void SetOutputs(FChooserEvaluationContext& Context, int RowIndex) const override;
 	
-	UPROPERTY(EditAnywhere, Meta = (ExcludeBaseStruct, BaseStruct = "/Script/Chooser.ChooserParameterBoolBase"), Category = "Hidden")
+	UPROPERTY(EditAnywhere, NoClear, Meta = (ExcludeBaseStruct, BaseStruct = "/Script/Chooser.ChooserParameterBoolBase"), Category = "Hidden")
 	FInstancedStruct InputValue;
 
 #if WITH_EDITOR
diff --git a/Engine/Plugins/Experimental/Chooser/Source/Chooser/Internal/OutputEnumColumn.h b/Engine/Plugins/Experimental/Chooser/Source/Chooser/Internal/OutputEnumColumn.h
--- a/Engine/Plugins/Experimental/Chooser/Source/Chooser/Internal/OutputEnumColumn.h
+++ b/Engine/Plugins/Experimental/Chooser/Source/Chooser/Internal/OutputEnumColumn.h
@@ -28,7 +28,7 @@ public:
 	virtual bool HasOutputs() const override { return true; }
 	virtual void SetOutputs(FChooserEvaluationContext& Context, int RowIndex) const override;
 
-	UPROPERTY(EditAnywhere, Meta = (ExcludeBaseStruct, BaseStruct = "/Script/Chooser.ChooserParameterEnumBase"), Category = "Data")
+	UPROPERTY(EditAnywhere, NoClear, Meta = (ExcludeBaseStruct, BaseStruct = "/Script/Chooser.ChooserParameterEnumBase"), Category = "Data")
 	FInstancedStruct InputValue;
 
 #if WITH_EDITORONLY_DATA
diff --git a/Engine/Plugins/Experimental/Chooser/Source/Chooser/Internal/OutputFloatColumn.h b/Engine/Plugins/Experimental/Chooser/Source/Chooser/Internal/OutputFloatColumn.h
--- a/Engine/Plugins/Experimental/Chooser/Source/Chooser/Internal/OutputFloatColumn.h
+++ b/Engine/Plugins/Experimental/Chooser/Source/Chooser/Internal/OutputFloatColumn.h
@@ -18,7 +18,7 @@ struct CHOOSER_API FOutputFloatColumn : public FChooserColumnBase
 	virtual bool HasOutputs() const override { return true; }
 	virtual void SetOutputs(FChooserEvaluationContext& Context, int RowIndex) const override;
 	
-	UPROPERTY(EditAnywhere, Meta = (ExcludeBaseStruct, BaseStruct = "/Script/Chooser.ChooserParameterFloatBase"), Category = "Hidden")
+	UPROPERTY(EditAnywhere, Meta = (ExcludeBaseStruct, NoClear, BaseStruct = "/Script/Chooser.ChooserParameterFloatBase"), Category = "Hidden")
 	FInstancedStruct InputValue;
 
 #if WITH_EDITOR
diff --git a/Engine/Plugins/Experimental/Chooser/Source/Chooser/Internal/OutputStructColumn.h b/Engine/Plugins/Experimental/Chooser/Source/Chooser/Internal/OutputStructColumn.h
--- a/Engine/Plugins/Experimental/Chooser/Source/Chooser/Internal/OutputStructColumn.h
+++ b/Engine/Plugins/Experimental/Chooser/Source/Chooser/Internal/OutputStructColumn.h
@@ -49,7 +49,7 @@ struct CHOOSER_API FOutputStructColumn : public FChooserColumnBase
 	virtual bool HasOutputs() const override { return true; }
 	virtual void SetOutputs(FChooserEvaluationContext& Context, int RowIndex) const override;
 	
-	UPROPERTY(EditAnywhere, Meta = (ExcludeBaseStruct, BaseStruct = "/Script/Chooser.ChooserParameterStructBase"), Category = "Hidden")
+	UPROPERTY(EditAnywhere, NoClear, Meta = (ExcludeBaseStruct, BaseStruct = "/Script/Chooser.ChooserParameterStructBase"), Category = "Hidden")
 	FInstancedStruct InputValue;
 
 #if WITH_EDITOR
diff --git a/Engine/Plugins/Experimental/Chooser/Source/Chooser/Internal/RandomizeColumn.h b/Engine/Plugins/Experimental/Chooser/Source/Chooser/Internal/RandomizeColumn.h
--- a/Engine/Plugins/Experimental/Chooser/Source/Chooser/Internal/RandomizeColumn.h
+++ b/Engine/Plugins/Experimental/Chooser/Source/Chooser/Internal/RandomizeColumn.h
@@ -48,7 +48,7 @@ struct CHOOSER_API FRandomizeColumn : public FChooserColumnBase
 	public:
 	FRandomizeColumn();
 	
-	UPROPERTY(EditAnywhere, Meta = (ExcludeBaseStruct, BaseStruct = "/Script/Chooser.RandomizeContextProperty"), Category = "Data")
+	UPROPERTY(EditAnywhere, NoClear, Meta = (ExcludeBaseStruct, BaseStruct = "/Script/Chooser.RandomizeContextProperty"), Category = "Data")
 	FInstancedStruct InputValue;
 	
 	
diff --git a/Engine/Plugins/Experimental/Chooser/Source/ChooserEditor/Private/ChooserTableEditor.cpp b/Engine/Plugins/Experimental/Chooser/Source/ChooserEditor/Private/ChooserTableEditor.cpp
--- a/Engine/Plugins/Experimental/Chooser/Source/ChooserEditor/Private/ChooserTableEditor.cpp
+++ b/Engine/Plugins/Experimental/Chooser/Source/ChooserEditor/Private/ChooserTableEditor.cpp
@@ -853,6 +853,8 @@ void FChooserTableEditor::UpdateTableColumns()
 					{
 						FStructViewerInitializationOptions Options;
 						Options.StructFilter = MakeShared<FStructFilter>(Chooser->ColumnsStructs[ColumnIndex].Get<FChooserColumnBase>().GetInputBaseType());
+						Options.bAllowViewOptions = false;
+						Options.bShowNoneOption = false;
 						Options.NameTypeToDisplay = EStructViewerNameTypeToDisplay::DisplayName;
 					
 						// Add class filter for columns here
@@ -977,6 +979,8 @@ TSharedRef<SDockTab> FChooserTableEditor::SpawnTableTab( const FSpawnTabArgs& Ar
 	{
 		FStructViewerInitializationOptions Options;
 		Options.StructFilter = MakeShared<FStructFilter>(FChooserColumnBase::StaticStruct());
+		Options.bAllowViewOptions = false;
+		Options.bShowNoneOption = false;
 		Options.NameTypeToDisplay = EStructViewerNameTypeToDisplay::DisplayName;
 		
 		// Add class filter for columns here
@@ -1001,6 +1005,8 @@ TSharedRef<SDockTab> FChooserTableEditor::SpawnTableTab( const FSpawnTabArgs& Ar
 		{
 			FStructViewerInitializationOptions Options;
 			Options.StructFilter = MakeShared<FStructFilter>(FObjectChooserBase::StaticStruct());
+			Options.bAllowViewOptions = false;
+			Options.bShowNoneOption = false;
 			Options.NameTypeToDisplay = EStructViewerNameTypeToDisplay::DisplayName;
 			
 			TSharedRef<SWidget> Widget = FModuleManager::LoadModuleChecked<FStructViewerModule>("StructViewer").CreateStructViewer(Options, FOnStructPicked::CreateLambda([this](const UScriptStruct* ChosenStruct)
diff --git a/Engine/Plugins/Experimental/Chooser/Source/ChooserEditor/Public/ObjectChooserClassFilter.h b/Engine/Plugins/Experimental/Chooser/Source/ChooserEditor/Public/ObjectChooserClassFilter.h
--- a/Engine/Plugins/Experimental/Chooser/Source/ChooserEditor/Public/ObjectChooserClassFilter.h
+++ b/Engine/Plugins/Experimental/Chooser/Source/ChooserEditor/Public/ObjectChooserClassFilter.h
@@ -26,7 +26,6 @@ namespace UE::ChooserEditor
 		UClass* InterfaceType;
 	};
 	
-
 	class FStructFilter : public IStructViewerFilter
 	{
 	public:
@@ -34,6 +33,11 @@ namespace UE::ChooserEditor
 		
 		virtual bool IsStructAllowed(const FStructViewerInitializationOptions& InInitOptions, const UScriptStruct* InStruct, TSharedRef<FStructViewerFilterFuncs> InFilterFuncs) override
 		{
+			if (InStruct->HasMetaData(TEXT("Hidden")))
+			{
+				return false;
+			}
+			
 			return InStruct->IsChildOf(BaseType) && InStruct !=BaseType;
 		}
 		virtual bool IsUnloadedStructAllowed(const FStructViewerInitializationOptions& InInitOptions, const FSoftObjectPath& InStructPath, TSharedRef<class FStructViewerFilterFuncs> InFilterFuncs)
diff --git a/Engine/Plugins/Experimental/Chooser/Source/ProxyTable/Public/ProxyTable.h b/Engine/Plugins/Experimental/Chooser/Source/ProxyTable/Public/ProxyTable.h
--- a/Engine/Plugins/Experimental/Chooser/Source/ProxyTable/Public/ProxyTable.h
+++ b/Engine/Plugins/Experimental/Chooser/Source/ProxyTable/Public/ProxyTable.h
@@ -18,7 +18,7 @@ struct PROXYTABLE_API FProxyStructOutput
 	UPROPERTY(EditAnywhere, Category = "Data")
 	FChooserStructPropertyBinding Binding;
 	
-	UPROPERTY(EditAnywhere, Category = "Data")
+	UPROPERTY(EditAnywhere, NoClear, Category = "Data")
 	FInstancedStruct Value;
 };
 
@@ -37,10 +37,10 @@ struct PROXYTABLE_API FProxyEntry
 	UPROPERTY(EditAnywhere, Category = "Data")
 	FName Key;
 	
-	UPROPERTY(DisplayName="Value", EditAnywhere, Meta = (ExcludeBaseStruct, BaseStruct = "/Script/Chooser.ObjectChooserBase"), Category = "Data")
+	UPROPERTY(DisplayName="Value", EditAnywhere, NoClear, Meta = (ExcludeBaseStruct, BaseStruct = "/Script/Chooser.ObjectChooserBase"), Category = "Data")
 	FInstancedStruct ValueStruct;
 
-	UPROPERTY(EditAnywhere, Category = "Data")
+	UPROPERTY(EditAnywhere, NoClear, Category = "Data")
 	TArray<FProxyStructOutput> OutputStructData;
 
 	const FGuid GetGuid() const;
