Make LookUpProxy take in a Proxy Table property binding

diff --git a/Engine/Plugins/Experimental/Chooser/Source/ProxyTable/Internal/LookupProxy.h b/Engine/Plugins/Experimental/Chooser/Source/ProxyTable/Internal/LookupProxy.h
--- a/Engine/Plugins/Experimental/Chooser/Source/ProxyTable/Internal/LookupProxy.h
+++ b/Engine/Plugins/Experimental/Chooser/Source/ProxyTable/Internal/LookupProxy.h
@@ -40,11 +40,16 @@ struct PROXYTABLE_API FLookupProxy : public FObjectChooserBase
 {
 	GENERATED_BODY()
 	virtual UObject* ChooseObject(FChooserEvaluationContext& Context) const final override;
+
+	FLookupProxy();
 	
 	public:
 	
 	UPROPERTY(EditAnywhere, Category="Parameters")
 	TObjectPtr<UProxyAsset> Proxy;
+
+	UPROPERTY(EditAnywhere, NoClear, Meta = (ExcludeBaseStruct, BaseStruct = "/Script/ProxyTable.ChooserParameterProxyTableBase"), Category = "Parameters")
+	FInstancedStruct ProxyTable;
 };
 
 USTRUCT(meta=(Hidden))
diff --git a/Engine/Plugins/Experimental/Chooser/Source/ProxyTable/Private/LookupProxy.cpp b/Engine/Plugins/Experimental/Chooser/Source/ProxyTable/Private/LookupProxy.cpp
--- a/Engine/Plugins/Experimental/Chooser/Source/ProxyTable/Private/LookupProxy.cpp
+++ b/Engine/Plugins/Experimental/Chooser/Source/ProxyTable/Private/LookupProxy.cpp
@@ -3,12 +3,28 @@
 #include "ProxyTableFunctionLibrary.h"
 #include "AssetRegistry/AssetRegistryModule.h"
 #include "ChooserPropertyAccess.h"
-#include "UObject/Package.h"
+
+FLookupProxy::FLookupProxy()
+{
+	ProxyTable.InitializeAs(FProxyTableContextProperty::StaticStruct());
+}
 
 UObject* FLookupProxy::ChooseObject(FChooserEvaluationContext& Context) const
 {
 	if (Proxy)
 	{
+		if (const FChooserParameterProxyTableBase* ProxyTableParameter = ProxyTable.GetPtr<FChooserParameterProxyTableBase>())
+		{
+			const UProxyTable* Table = nullptr;
+			if (ProxyTableParameter->GetValue(Context, Table))
+			{
+				if (Table)
+				{
+					return Table->FindProxyObject(Proxy->Guid, Context);
+				}
+			}
+		}
+		// fallback codepath will look up the table from the property binding on the proxy asset
 		return Proxy->FindProxyObject(Context);
 	}
 	return nullptr;
diff --git a/Engine/Plugins/Experimental/Chooser/Source/ProxyTableEditor/Private/ProxyTableEditor.cpp b/Engine/Plugins/Experimental/Chooser/Source/ProxyTableEditor/Private/ProxyTableEditor.cpp
--- a/Engine/Plugins/Experimental/Chooser/Source/ProxyTableEditor/Private/ProxyTableEditor.cpp
+++ b/Engine/Plugins/Experimental/Chooser/Source/ProxyTableEditor/Private/ProxyTableEditor.cpp
@@ -34,6 +34,7 @@
 #include "PropertyCustomizationHelpers.h"
 #include "ProxyTableEditorCommands.h"
 #include "LookupProxy.h"
+#include "SPropertyAccessChainWidget.h"
 #include "Widgets/Layout/SScrollBox.h"
 
 #define LOCTEXT_NAMESPACE "ProxyTableEditor"
@@ -1022,42 +1023,77 @@ TSharedRef<FProxyTableEditor> FProxyTableEditor::CreateEditor( const EToolkitMod
 }
 
 /// Result widgets
-
-TSharedRef<SWidget> CreateLookupProxyWidget(bool bReadOnly, UObject* TransactionObject, void* Value, UClass* ResultBaseClass, ChooserEditor::FChooserWidgetValueChanged ValueChanged)
-{
-	FLookupProxy* LookupProxy = static_cast<FLookupProxy*>(Value);
+///
 	
-	return SNew(SObjectPropertyEntryBox)
-		.IsEnabled(!bReadOnly)
-		.AllowedClass(UProxyAsset::StaticClass())
-		.ObjectPath_Lambda([LookupProxy](){ return LookupProxy->Proxy.GetPath();})
-		.OnShouldFilterAsset_Lambda([ResultBaseClass](const FAssetData& AssetData)
-		{
-			if (ResultBaseClass == nullptr)
-			{
-				return false;
-			}
-			if (AssetData.IsInstanceOf(UProxyAsset::StaticClass()))
-			{
-				if (UProxyAsset* Proxy = Cast<UProxyAsset>(AssetData.GetAsset()))
-				{
-					return !(Proxy->Type && Proxy->Type->IsChildOf(ResultBaseClass));
-				}
-			}
-			return true;
-		})
-		.OnObjectChanged_Lambda([TransactionObject, LookupProxy, ValueChanged](const FAssetData& AssetData)
+TSharedRef<SWidget> CreateProxyTablePropertyWidget(bool bReadOnly, UObject* TransactionObject, void* Value, UClass* ResultBaseClass, UE::ChooserEditor::FChooserWidgetValueChanged ValueChanged)
+{
+	IHasContextClass* HasContextClass = Cast<IHasContextClass>(TransactionObject);
+
+	FProxyTableContextProperty* ContextProperty = reinterpret_cast<FProxyTableContextProperty*>(Value);
+
+	return SNew(UE::ChooserEditor::SPropertyAccessChainWidget).ContextClassOwner(HasContextClass).AllowFunctions(false).BindingColor("ClassPinTypeColor").TypeFilter("UProxyTable*")
+	.PropertyBindingValue(&ContextProperty->Binding)
+	.OnAddBinding_Lambda(
+		[ContextProperty, TransactionObject, ValueChanged](FName InPropertyName, const TArray<FBindingChainElement>& InBindingChain)
 		{
-			const FScopedTransaction Transaction(LOCTEXT("Edit Chooser", "Edit Chooser"));
+			const FScopedTransaction Transaction(NSLOCTEXT("ContextPropertyWidget", "Change Property Binding", "Change Property Binding"));
 			TransactionObject->Modify(true);
-			LookupProxy->Proxy = Cast<UProxyAsset>(AssetData.GetAsset());
+			ContextProperty->SetBinding(InBindingChain);
 			ValueChanged.ExecuteIfBound();
 		});
 }
 
+TSharedRef<SWidget> CreateLookupProxyWidget(bool bReadOnly, UObject* TransactionObject, void* Value, UClass* ResultBaseClass, ChooserEditor::FChooserWidgetValueChanged ValueChanged)
+{
+	FLookupProxy* LookupProxy = static_cast<FLookupProxy*>(Value);
+	
+	TSharedPtr<SWidget> ProxyTableWidget = UE::ChooserEditor::FObjectChooserWidgetFactories::CreateWidget(false, TransactionObject, LookupProxy->ProxyTable.GetMutableMemory(),LookupProxy->ProxyTable.GetScriptStruct(), ResultBaseClass);
+
+	
+	TSharedRef<SWidget> ProxyAssetWidget =SNew(SObjectPropertyEntryBox)
+    		.IsEnabled(!bReadOnly)
+    		.AllowedClass(UProxyAsset::StaticClass())
+			.DisplayBrowse(false)
+			.DisplayUseSelected(false)
+    		.ObjectPath_Lambda([LookupProxy](){ return LookupProxy->Proxy.GetPath();})
+    		.OnShouldFilterAsset_Lambda([ResultBaseClass](const FAssetData& AssetData)
+    		{
+    			if (ResultBaseClass == nullptr)
+    			{
+    				return false;
+    			}
+    			if (AssetData.IsInstanceOf(UProxyAsset::StaticClass()))
+    			{
+    				if (UProxyAsset* Proxy = Cast<UProxyAsset>(AssetData.GetAsset()))
+    				{
+    					return !(Proxy->Type && Proxy->Type->IsChildOf(ResultBaseClass));
+    				}
+    			}
+    			return true;
+    		})
+    		.OnObjectChanged_Lambda([TransactionObject, LookupProxy, ValueChanged](const FAssetData& AssetData)
+    		{
+    			const FScopedTransaction Transaction(LOCTEXT("Edit Chooser", "Edit Chooser"));
+    			TransactionObject->Modify(true);
+    			LookupProxy->Proxy = Cast<UProxyAsset>(AssetData.GetAsset());
+    			ValueChanged.ExecuteIfBound();
+    		});
+
+	return SNew(SHorizontalBox)
+		+SHorizontalBox::Slot()
+		[
+			ProxyTableWidget ? ProxyTableWidget.ToSharedRef() : SNullWidget::NullWidget
+		]
+		+SHorizontalBox::Slot()
+		[
+			ProxyAssetWidget
+		];
+}
+
 void FProxyTableEditor::RegisterWidgets()
 {
 	UE::ChooserEditor::FObjectChooserWidgetFactories::RegisterWidgetCreator(FLookupProxy::StaticStruct(), CreateLookupProxyWidget);
+	UE::ChooserEditor::FObjectChooserWidgetFactories::RegisterWidgetCreator(FProxyTableContextProperty::StaticStruct(), CreateProxyTablePropertyWidget);
 	
 	FPropertyEditorModule& PropertyModule = FModuleManager::LoadModuleChecked<FPropertyEditorModule>("PropertyEditor");
 	PropertyModule.RegisterCustomClassLayout("ProxyRowDetails", FOnGetDetailCustomizationInstance::CreateStatic(&FProxyRowDetails::MakeInstance));	
