Fix crash in anim bluperint compilation when a proxy asset is missing

#rnx

diff --git a/Engine/Plugins/Experimental/Chooser/Source/ProxyTableUncooked/Private/EvaluateProxyNode.cpp b/Engine/Plugins/Experimental/Chooser/Source/ProxyTableUncooked/Private/EvaluateProxyNode.cpp
--- a/Engine/Plugins/Experimental/Chooser/Source/ProxyTableUncooked/Private/EvaluateProxyNode.cpp
+++ b/Engine/Plugins/Experimental/Chooser/Source/ProxyTableUncooked/Private/EvaluateProxyNode.cpp
@@ -207,7 +207,7 @@ void UK2Node_EvaluateProxy::ExpandNode(class FKismetCompilerContext& CompilerCon
 
 		UEdGraphPin* OutputPin = CallFunction->GetReturnValuePin();
 
-		if (Proxy->Type)
+		if (Proxy && Proxy->Type)
 		{
 			UEdGraphPin* OutputClassPin = CallFunction->FindPin(TEXT("ObjectClass"));
 			CallFunction->GetSchema()->TrySetDefaultObject(*OutputClassPin, Proxy->Type);
