Fix crash in choser's ResolvePropertyChain when output struct is null

#rnx

diff --git a/Engine/Plugins/Experimental/Chooser/Source/Chooser/Private/ChooserPropertyAccess.cpp b/Engine/Plugins/Experimental/Chooser/Source/Chooser/Private/ChooserPropertyAccess.cpp
--- a/Engine/Plugins/Experimental/Chooser/Source/Chooser/Private/ChooserPropertyAccess.cpp
+++ b/Engine/Plugins/Experimental/Chooser/Source/Chooser/Private/ChooserPropertyAccess.cpp
@@ -13,6 +13,10 @@ namespace UE::Chooser
 {
 	bool ResolvePropertyChain(FChooserEvaluationContext& Context, const FChooserPropertyBinding& PropertyBinding, const void*& OutContainer, const UStruct*& OutStructType)
 	{
+		if (OutStructType == nullptr)
+		{
+			return false;
+		}
 		if (Context.Params.IsValidIndex(PropertyBinding.ContextIndex))
 		{
 			if (FChooserEvaluationInputObject* ObjectParam = Context.Params[PropertyBinding.ContextIndex].GetMutablePtr<FChooserEvaluationInputObject>())
