add multi-output support to proxy table lookups

diff --git a/Engine/Plugins/Experimental/Chooser/Source/ProxyTable/Internal/LookupProxy.h b/Engine/Plugins/Experimental/Chooser/Source/ProxyTable/Internal/LookupProxy.h
--- a/Engine/Plugins/Experimental/Chooser/Source/ProxyTable/Internal/LookupProxy.h
+++ b/Engine/Plugins/Experimental/Chooser/Source/ProxyTable/Internal/LookupProxy.h
@@ -39,6 +39,7 @@ USTRUCT()
 struct PROXYTABLE_API FLookupProxy : public FObjectChooserBase
 {
 	GENERATED_BODY()
+	virtual EIteratorStatus ChooseMulti(FChooserEvaluationContext &Context, FObjectChooserIteratorCallback Callback) const final override;
 	virtual UObject* ChooseObject(FChooserEvaluationContext& Context) const final override;
 
 	FLookupProxy();
@@ -56,6 +57,7 @@ USTRUCT(meta=(Hidden))
 struct PROXYTABLE_API FLookupProxyWithOverrideTable : public FObjectChooserBase
 {
 	GENERATED_BODY()
+	virtual EIteratorStatus ChooseMulti(FChooserEvaluationContext &Context, FObjectChooserIteratorCallback Callback) const final override;
 	virtual UObject* ChooseObject(FChooserEvaluationContext& Context) const final override;
 	
 	public:
diff --git a/Engine/Plugins/Experimental/Chooser/Source/ProxyTable/Private/LookupProxy.cpp b/Engine/Plugins/Experimental/Chooser/Source/ProxyTable/Private/LookupProxy.cpp
--- a/Engine/Plugins/Experimental/Chooser/Source/ProxyTable/Private/LookupProxy.cpp
+++ b/Engine/Plugins/Experimental/Chooser/Source/ProxyTable/Private/LookupProxy.cpp
@@ -9,6 +9,27 @@ FLookupProxy::FLookupProxy()
 	ProxyTable.InitializeAs(FProxyTableContextProperty::StaticStruct());
 }
 
+FObjectChooserBase::EIteratorStatus FLookupProxy::ChooseMulti(FChooserEvaluationContext& Context, FObjectChooserIteratorCallback Callback) const
+{
+	if (Proxy)
+	{
+		if (const FChooserParameterProxyTableBase* ProxyTableParameter = ProxyTable.GetPtr<FChooserParameterProxyTableBase>())
+		{
+			const UProxyTable* Table = nullptr;
+			if (ProxyTableParameter->GetValue(Context, Table))
+			{
+				if (Table)
+				{
+					return Table->FindProxyObjectMulti(Proxy->Guid, Context, Callback);
+				}
+			}
+		}
+		// fallback codepath will look up the table from the property binding on the proxy asset
+		return Proxy->FindProxyObjectMulti(Context, Callback);
+	}
+	return FObjectChooserBase::EIteratorStatus::Continue;
+}
+
 UObject* FLookupProxy::ChooseObject(FChooserEvaluationContext& Context) const
 {
 	if (Proxy)
@@ -39,6 +60,15 @@ UObject* FLookupProxyWithOverrideTable::ChooseObject(FChooserEvaluationContext&
 	return nullptr;
 }
 
+FObjectChooserBase::EIteratorStatus FLookupProxyWithOverrideTable::ChooseMulti(FChooserEvaluationContext& Context, FObjectChooserIteratorCallback Callback) const
+{
+	if (Proxy && OverrideProxyTable)
+	{
+		return OverrideProxyTable->FindProxyObjectMulti(Proxy->Guid, Context, Callback);
+	}
+	return FObjectChooserBase::EIteratorStatus::Continue;
+}
+
 bool FProxyTableContextProperty::GetValue(FChooserEvaluationContext& Context, const UProxyTable*& OutResult) const
 {
 	const UStruct* StructType = nullptr;
diff --git a/Engine/Plugins/Experimental/Chooser/Source/ProxyTable/Private/ProxyAsset.cpp b/Engine/Plugins/Experimental/Chooser/Source/ProxyTable/Private/ProxyAsset.cpp
--- a/Engine/Plugins/Experimental/Chooser/Source/ProxyTable/Private/ProxyAsset.cpp
+++ b/Engine/Plugins/Experimental/Chooser/Source/ProxyTable/Private/ProxyAsset.cpp
@@ -122,3 +122,20 @@ UObject* UProxyAsset::FindProxyObject(FChooserEvaluationContext& Context) const
 	
 	return nullptr;
 }
+
+FObjectChooserBase::EIteratorStatus UProxyAsset::FindProxyObjectMulti(FChooserEvaluationContext &Context, FObjectChooserBase::FObjectChooserIteratorCallback Callback) const
+{
+	if (ProxyTable.IsValid())
+	{
+		const UProxyTable* Table;
+		if (ProxyTable.Get<FChooserParameterProxyTableBase>().GetValue(Context, Table))
+		{
+			if(Table)
+			{
+				return Table->FindProxyObjectMulti(Guid, Context, Callback);
+			}
+		}
+	}
+	
+	return FObjectChooserBase::EIteratorStatus::Continue;
+}
diff --git a/Engine/Plugins/Experimental/Chooser/Source/ProxyTable/Private/ProxyTable.cpp b/Engine/Plugins/Experimental/Chooser/Source/ProxyTable/Private/ProxyTable.cpp
--- a/Engine/Plugins/Experimental/Chooser/Source/ProxyTable/Private/ProxyTable.cpp
+++ b/Engine/Plugins/Experimental/Chooser/Source/ProxyTable/Private/ProxyTable.cpp
@@ -163,49 +163,54 @@ void UProxyTable::PostTransacted(const FTransactionObjectEvent& TransactionEvent
 
 #endif
 
-UObject* UProxyTable::FindProxyObject(const FGuid& Key, FChooserEvaluationContext& Context) const
+static void OutputStructData(const FRuntimeProxyValue& EntryValueData, FChooserEvaluationContext& Context)
 {
-	const int FoundIndex = Algo::BinarySearch(Keys, Key);
-	if (FoundIndex != INDEX_NONE)
+	for (const FProxyStructOutput& StructOutput : EntryValueData.OutputStructData)
 	{
-		const FRuntimeProxyValue& EntryValueData = RuntimeValues[FoundIndex];
-		const FObjectChooserBase &EntryValue = EntryValueData.Value.Get<const FObjectChooserBase>();
+		const void* Container = nullptr;
+		const UStruct* StructType;
 
-		UObject* Result = EntryValue.ChooseObject(Context);
+		// copy each struct output value
 
-		for (const FProxyStructOutput& StructOutput : EntryValueData.OutputStructData)
+		if (StructOutput.Binding.PropertyBindingChain.IsEmpty())
 		{
-			const void* Container = nullptr;
-			const UStruct* StructType;
-
-			// copy each struct output value
-
-			if (StructOutput.Binding.PropertyBindingChain.IsEmpty())
+			if(Context.Params.IsValidIndex((StructOutput.Binding.ContextIndex)))
 			{
-				if(Context.Params.IsValidIndex((StructOutput.Binding.ContextIndex)))
+				// directly bound to context struct
+				if (Context.Params[StructOutput.Binding.ContextIndex].GetScriptStruct() == StructOutput.Value.GetScriptStruct())
 				{
-					// directly bound to context struct
-					if (Context.Params[StructOutput.Binding.ContextIndex].GetScriptStruct() == StructOutput.Value.GetScriptStruct())
-					{
-						void* TargetData = Context.Params[StructOutput.Binding.ContextIndex].GetMutableMemory();
-						StructOutput.Value.GetScriptStruct()->CopyScriptStruct(TargetData, StructOutput.Value.GetMemory());
-					}
+					void* TargetData = Context.Params[StructOutput.Binding.ContextIndex].GetMutableMemory();
+					StructOutput.Value.GetScriptStruct()->CopyScriptStruct(TargetData, StructOutput.Value.GetMemory());
 				}
 			}
-			else if (UE::Chooser::ResolvePropertyChain(Context, StructOutput.Binding, Container, StructType))
+		}
+		else if (UE::Chooser::ResolvePropertyChain(Context, StructOutput.Binding, Container, StructType))
+		{
+			if (FStructProperty* Property = FindFProperty<FStructProperty>(StructType, StructOutput.Binding.PropertyBindingChain.Last()))
 			{
-				if (FStructProperty* Property = FindFProperty<FStructProperty>(StructType, StructOutput.Binding.PropertyBindingChain.Last()))
+				// const cast is here just because ResolvePropertyChain expects a const void*&
+				void* TargetData = Property->ContainerPtrToValuePtr<void>(const_cast<void*>(Container));
+				
+				if (Property->Struct == StructOutput.Value.GetScriptStruct())
 				{
-					// const cast is here just because ResolvePropertyChain expects a const void*&
-					void* TargetData = Property->ContainerPtrToValuePtr<void>(const_cast<void*>(Container));
-					
-					if (Property->Struct == StructOutput.Value.GetScriptStruct())
-					{
-						Property->Struct->CopyScriptStruct(TargetData, StructOutput.Value.GetMemory());
-					}
+					Property->Struct->CopyScriptStruct(TargetData, StructOutput.Value.GetMemory());
 				}
 			}
 		}
+	}
+}
+
+UObject* UProxyTable::FindProxyObject(const FGuid& Key, FChooserEvaluationContext& Context) const
+{
+	const int FoundIndex = Algo::BinarySearch(Keys, Key);
+	if (FoundIndex != INDEX_NONE)
+	{
+		const FRuntimeProxyValue& EntryValueData = RuntimeValues[FoundIndex];
+		const FObjectChooserBase &EntryValue = EntryValueData.Value.Get<const FObjectChooserBase>();
+
+		UObject* Result = EntryValue.ChooseObject(Context);
+
+		OutputStructData(EntryValueData, Context);
 
 		return Result;
 	}
@@ -213,6 +218,24 @@ UObject* UProxyTable::FindProxyObject(const FGuid& Key, FChooserEvaluationContex
 	return nullptr;
 }
 
+FObjectChooserBase::EIteratorStatus UProxyTable::FindProxyObjectMulti(const FGuid& Key, FChooserEvaluationContext &Context, FObjectChooserBase::FObjectChooserIteratorCallback Callback) const
+{
+		const int FoundIndex = Algo::BinarySearch(Keys, Key);
+    	if (FoundIndex != INDEX_NONE)
+    	{
+    		const FRuntimeProxyValue& EntryValueData = RuntimeValues[FoundIndex];
+    		const FObjectChooserBase &EntryValue = EntryValueData.Value.Get<const FObjectChooserBase>();
+    
+    		FObjectChooserBase::EIteratorStatus Result = EntryValue.ChooseMulti(Context, Callback);
+    
+			OutputStructData(EntryValueData, Context);
+    
+    		return Result;
+    	}
+    	
+    	return FObjectChooserBase::EIteratorStatus::Continue;
+}
+
 UProxyTable::UProxyTable(const FObjectInitializer& Initializer)
 	:Super(Initializer)
 {
diff --git a/Engine/Plugins/Experimental/Chooser/Source/ProxyTable/Public/ProxyAsset.h b/Engine/Plugins/Experimental/Chooser/Source/ProxyTable/Public/ProxyAsset.h
--- a/Engine/Plugins/Experimental/Chooser/Source/ProxyTable/Public/ProxyAsset.h
+++ b/Engine/Plugins/Experimental/Chooser/Source/ProxyTable/Public/ProxyAsset.h
@@ -8,6 +8,7 @@
 #include "InstancedStruct.h"
 #include "InstancedStructContainer.h"
 #include "Misc/Guid.h"
+#include "IObjectChooser.h"
 #include "ProxyAsset.generated.h"
 
 DECLARE_MULTICAST_DELEGATE_OneParam(FProxyTypeChanged, const UClass* OutputObjectType);
@@ -47,7 +48,8 @@ public:
 	FGuid Guid;
 
 	virtual TConstArrayView<FInstancedStruct> GetContextData() const override { return ContextData; }
-	UObject* FindProxyObject(FChooserEvaluationContext& Context) const;
+	UObject* FindProxyObject(struct FChooserEvaluationContext& Context) const;
+	FObjectChooserBase::EIteratorStatus FindProxyObjectMulti(FChooserEvaluationContext &Context, FObjectChooserBase::FObjectChooserIteratorCallback Callback) const;
 
 #if WITH_EDITORONLY_DATA
 	UPROPERTY()
diff --git a/Engine/Plugins/Experimental/Chooser/Source/ProxyTable/Public/ProxyTable.h b/Engine/Plugins/Experimental/Chooser/Source/ProxyTable/Public/ProxyTable.h
--- a/Engine/Plugins/Experimental/Chooser/Source/ProxyTable/Public/ProxyTable.h
+++ b/Engine/Plugins/Experimental/Chooser/Source/ProxyTable/Public/ProxyTable.h
@@ -76,7 +76,8 @@ public:
 
 	UPROPERTY()
 	TArray<FRuntimeProxyValue> RuntimeValues;
-	
+
+	FObjectChooserBase::EIteratorStatus FindProxyObjectMulti(const FGuid& Key, FChooserEvaluationContext &Context, FObjectChooserBase::FObjectChooserIteratorCallback Callback) const;
 	UObject* FindProxyObject(const FGuid& Key, FChooserEvaluationContext& Context) const;
 
 #if WITH_EDITORONLY_DATA
