Property binding optimizations for choosers

diff --git a/Engine/Plugins/Experimental/Chooser/Source/Chooser/Internal/BoolColumn.h b/Engine/Plugins/Experimental/Chooser/Source/Chooser/Internal/BoolColumn.h
--- a/Engine/Plugins/Experimental/Chooser/Source/Chooser/Internal/BoolColumn.h
+++ b/Engine/Plugins/Experimental/Chooser/Source/Chooser/Internal/BoolColumn.h
@@ -40,23 +40,19 @@ public:
 		}
 	}
 	
-
-#if WITH_EDITOR
-	
-	static bool CanBind(const FProperty& Property)
+	virtual void Compile(IHasContextClass* Owner, bool bForce) override
 	{
-		static FString BoolTypeName = "bool";
-		return Property.GetCPPType() == BoolTypeName;
-	}
-
-	void SetBinding(const TArray<FBindingChainElement>& InBindingChain)
-	{
-		UE::Chooser::CopyPropertyChain(InBindingChain, Binding);
-	}
+		Binding.Compile(Owner, bForce);
+	};
 
+#if WITH_EDITOR
 	virtual void GetDisplayName(FText& OutName) const override
 	{
-		if (!Binding.PropertyBindingChain.IsEmpty())
+		if (!Binding.DisplayName.IsEmpty())
+		{
+			OutName = FText::FromString(Binding.DisplayName);
+		} 
+		else if (!Binding.PropertyBindingChain.IsEmpty())
 		{
 			OutName = FText::FromName(Binding.PropertyBindingChain.Last());
 		}
diff --git a/Engine/Plugins/Experimental/Chooser/Source/Chooser/Internal/Chooser.h b/Engine/Plugins/Experimental/Chooser/Source/Chooser/Internal/Chooser.h
--- a/Engine/Plugins/Experimental/Chooser/Source/Chooser/Internal/Chooser.h
+++ b/Engine/Plugins/Experimental/Chooser/Source/Chooser/Internal/Chooser.h
@@ -6,7 +6,10 @@
 #include "IObjectChooser.h"
 #include "InstancedStruct.h"
 #include "IChooserColumn.h"
-#include "ChooserPropertyAccess.h"
+#include "IHasContext.h"
+#if WITH_EDITOR
+#include "Kismet2/StructureEditorUtils.h"
+#endif
 
 #include "Chooser.generated.h"
 
@@ -17,12 +20,16 @@ class CHOOSER_API UChooserTable : public UObject, public IHasContextClass
 public:
 	UChooserTable() {}
 
+	virtual void PostLoad() override;
+	virtual void Compile(bool bForce = false) override;
 #if WITH_EDITOR
+	void OnDependentStructChanged(UUserDefinedStruct* Blueprint) { Compile(true); }
+	void OnDependencyCompiled(UBlueprint* Blueprint) { Compile(true); }
+	virtual void AddCompileDependency(const UStruct* Struct) override;
 	FChooserOutputObjectTypeChanged OnOutputObjectTypeChanged;
 	
 	virtual void PostEditUndo() override;
 	virtual void PostEditChangeProperty(FPropertyChangedEvent& PropertyChangedEvent) override;
-	virtual void PostLoad() override;
 	
 	void SetDebugSelectedRow(int32 Index) const { DebugSelectedRow = Index; }
 	int32 GetDebugSelectedRow() const { return DebugSelectedRow; }
@@ -49,6 +56,8 @@ private:
 	TWeakObjectPtr<const UObject> DebugTarget;
 	// Row which was selected last time this chooser was evaluated on DebugTarget
 	mutable int32 DebugSelectedRow = -1;
+
+	TArray<TWeakObjectPtr<UStruct>> CompileDependencies;
 #endif
 
 public:
diff --git a/Engine/Plugins/Experimental/Chooser/Source/Chooser/Internal/EnumColumn.h b/Engine/Plugins/Experimental/Chooser/Source/Chooser/Internal/EnumColumn.h
--- a/Engine/Plugins/Experimental/Chooser/Source/Chooser/Internal/EnumColumn.h
+++ b/Engine/Plugins/Experimental/Chooser/Source/Chooser/Internal/EnumColumn.h
@@ -37,27 +37,19 @@ struct CHOOSER_API FEnumContextProperty : public FChooserParameterEnumBase
 		}
 	}
 
-#if WITH_EDITOR
-	static bool CanBind(const FProperty& Property)
-	{
-		if (Property.IsA<FEnumProperty>())
-		{
-			return true;
-		}
-
-		if (const FByteProperty* ByteProperty = CastField<FByteProperty>(&Property))
-		{
-			return ByteProperty->Enum != nullptr;
-		}
-
-		return false;
-	}
-
-	void SetBinding(const TArray<FBindingChainElement>& InBindingChain);
+	virtual void Compile(IHasContextClass* Owner, bool bForce) override
+    {
+    	Binding.Compile(Owner, bForce);
+    };
 
+#if WITH_EDITOR
 	virtual void GetDisplayName(FText& OutName) const override
 	{
-		if (!Binding.PropertyBindingChain.IsEmpty())
+		if (!Binding.DisplayName.IsEmpty())
+		{
+			OutName = FText::FromString(Binding.DisplayName);
+		} 
+		else if (!Binding.PropertyBindingChain.IsEmpty())
 		{
 			OutName = FText::FromName(Binding.PropertyBindingChain.Last());
 		}
diff --git a/Engine/Plugins/Experimental/Chooser/Source/Chooser/Internal/FloatRangeColumn.h b/Engine/Plugins/Experimental/Chooser/Source/Chooser/Internal/FloatRangeColumn.h
--- a/Engine/Plugins/Experimental/Chooser/Source/Chooser/Internal/FloatRangeColumn.h
+++ b/Engine/Plugins/Experimental/Chooser/Source/Chooser/Internal/FloatRangeColumn.h
@@ -1,6 +1,8 @@
 // Copyright Epic Games, Inc. All Rights Reserved.
 #pragma once
 
+#include <IHasContext.h>
+
 #include "CoreMinimal.h"
 #include "IChooserColumn.h"
 #include "IChooserParameterFloat.h"
@@ -29,25 +31,21 @@ struct CHOOSER_API FFloatContextProperty :  public FChooserParameterFloatBase
 			Binding.PropertyBindingChain = PropertyBindingChain_DEPRECATED;
 			PropertyBindingChain_DEPRECATED.SetNum(0);
 		}
-	}	
-
-#if WITH_EDITOR
-	static bool CanBind(const FProperty& Property)
-	{
-		static FString DoubleTypeName = "double";
-		static FString FloatTypeName = "float";
-		const FString& TypeName = Property.GetCPPType();
-		return TypeName == FloatTypeName || TypeName == DoubleTypeName;
 	}
 	
-	void SetBinding(const TArray<FBindingChainElement>& InBindingChain)
+	virtual void Compile(IHasContextClass* Owner, bool bForce) override
 	{
-		UE::Chooser::CopyPropertyChain(InBindingChain, Binding);
-	}
+		Binding.Compile(Owner, bForce);
+	};
 
+#if WITH_EDITOR
 	virtual void GetDisplayName(FText& OutName) const override
 	{
-		if (!Binding.PropertyBindingChain.IsEmpty())
+		if (!Binding.DisplayName.IsEmpty())
+		{
+			OutName = FText::FromString(Binding.DisplayName);
+		} 
+		else if (!Binding.PropertyBindingChain.IsEmpty())
 		{
 			OutName = FText::FromName(Binding.PropertyBindingChain.Last());
 		}
@@ -105,7 +103,7 @@ struct CHOOSER_API FFloatRangeColumn : public FChooserColumnBase
 			InputValue.GetMutable<FChooserParameterBase>().PostLoad();
 		}
 	}
-	
+
 	CHOOSER_COLUMN_BOILERPLATE(FChooserParameterFloatBase);
 
 };
diff --git a/Engine/Plugins/Experimental/Chooser/Source/Chooser/Internal/GameplayTagColumn.h b/Engine/Plugins/Experimental/Chooser/Source/Chooser/Internal/GameplayTagColumn.h
--- a/Engine/Plugins/Experimental/Chooser/Source/Chooser/Internal/GameplayTagColumn.h
+++ b/Engine/Plugins/Experimental/Chooser/Source/Chooser/Internal/GameplayTagColumn.h
@@ -33,21 +33,20 @@ public:
 			PropertyBindingChain_DEPRECATED.SetNum(0);
 		}
 	}
-#if WITH_EDITOR
-	static bool CanBind(const FProperty& Property)
-	{
-		static FString TypeName = "FGameplayTagContainer";
-		return Property.GetCPPType() == TypeName;
-	}
-
-	void SetBinding(const TArray<FBindingChainElement>& InBindingChain)
+	
+	virtual void Compile(IHasContextClass* Owner, bool bForce) override
 	{
-		UE::Chooser::CopyPropertyChain(InBindingChain, Binding);
-	}
+		Binding.Compile(Owner, bForce);
+	};
 
+#if WITH_EDITOR
 	virtual void GetDisplayName(FText& OutName) const override
 	{
-		if (!Binding.PropertyBindingChain.IsEmpty())
+		if (!Binding.DisplayName.IsEmpty())
+		{
+			OutName = FText::FromString(Binding.DisplayName);
+		} 
+		else if (!Binding.PropertyBindingChain.IsEmpty())
 		{
 			OutName = FText::FromName(Binding.PropertyBindingChain.Last());
 		}
diff --git a/Engine/Plugins/Experimental/Chooser/Source/Chooser/Internal/ObjectColumn.h b/Engine/Plugins/Experimental/Chooser/Source/Chooser/Internal/ObjectColumn.h
--- a/Engine/Plugins/Experimental/Chooser/Source/Chooser/Internal/ObjectColumn.h
+++ b/Engine/Plugins/Experimental/Chooser/Source/Chooser/Internal/ObjectColumn.h
@@ -21,13 +21,6 @@ struct CHOOSER_API FObjectContextProperty : public FChooserParameterObjectBase
 	virtual bool GetValue(FChooserEvaluationContext& Context, FSoftObjectPath& OutResult) const override;
 
 #if WITH_EDITOR
-	static bool CanBind(const FProperty& Property)
-	{
-		return Property.IsA<FObjectPropertyBase>() && !(Property.IsA<FClassProperty>() || Property.IsA<FSoftClassProperty>());
-	}
-
-	void SetBinding(const TArray<FBindingChainElement>& InBindingChain);
-
 	virtual void GetDisplayName(FText& OutName) const override
 	{
 		if (!Binding.PropertyBindingChain.IsEmpty())
diff --git a/Engine/Plugins/Experimental/Chooser/Source/Chooser/Internal/OutputStructColumn.h b/Engine/Plugins/Experimental/Chooser/Source/Chooser/Internal/OutputStructColumn.h
--- a/Engine/Plugins/Experimental/Chooser/Source/Chooser/Internal/OutputStructColumn.h
+++ b/Engine/Plugins/Experimental/Chooser/Source/Chooser/Internal/OutputStructColumn.h
@@ -21,8 +21,6 @@ struct CHOOSER_API FStructContextProperty : public FChooserParameterStructBase
 	virtual bool SetValue(FChooserEvaluationContext& Context, const FInstancedStruct &Value) const override;
 
 #if WITH_EDITOR
-	void SetBinding(const UObject* OuterObject, const TArray<FBindingChainElement>& InBindingChain);
-
 	virtual void GetDisplayName(FText& OutName) const override
 	{
 		if (!Binding.DisplayName.IsEmpty())
diff --git a/Engine/Plugins/Experimental/Chooser/Source/Chooser/Internal/RandomizeColumn.h b/Engine/Plugins/Experimental/Chooser/Source/Chooser/Internal/RandomizeColumn.h
--- a/Engine/Plugins/Experimental/Chooser/Source/Chooser/Internal/RandomizeColumn.h
+++ b/Engine/Plugins/Experimental/Chooser/Source/Chooser/Internal/RandomizeColumn.h
@@ -18,22 +18,20 @@ public:
 	FChooserPropertyBinding Binding;
 	
 	virtual bool GetValue(FChooserEvaluationContext& Context, const FChooserRandomizationContext*& OutResult) const override;
-
-#if WITH_EDITOR
-	static bool CanBind(const FProperty& Property)
-	{
-		static FString RandomizeTypeName = "bool";
-		return Property.GetCPPType() == RandomizeTypeName;
-	}
-
-	void SetBinding(const TArray<FBindingChainElement>& InBindingChain)
+	
+	virtual void Compile(IHasContextClass* Owner, bool bForce) override
 	{
-		UE::Chooser::CopyPropertyChain(InBindingChain, Binding);
-	}
+		Binding.Compile(Owner, bForce);
+	};
 
+#if WITH_EDITOR
 	virtual void GetDisplayName(FText& OutName) const override
 	{
-		if (!Binding.PropertyBindingChain.IsEmpty())
+		if (!Binding.DisplayName.IsEmpty())
+		{
+			OutName = FText::FromString(Binding.DisplayName);
+		} 
+		else if (!Binding.PropertyBindingChain.IsEmpty())
 		{
 			OutName = FText::FromName(Binding.PropertyBindingChain.Last());
 		}
diff --git a/Engine/Plugins/Experimental/Chooser/Source/Chooser/Private/BoolColumn.cpp b/Engine/Plugins/Experimental/Chooser/Source/Chooser/Private/BoolColumn.cpp
--- a/Engine/Plugins/Experimental/Chooser/Source/Chooser/Private/BoolColumn.cpp
+++ b/Engine/Plugins/Experimental/Chooser/Source/Chooser/Private/BoolColumn.cpp
@@ -5,55 +5,12 @@
 
 bool FBoolContextProperty::GetValue(FChooserEvaluationContext& Context, bool& OutResult) const
 {
-	
-	const UStruct* StructType = nullptr;
-	const void* Container = nullptr;
-	
-	if (UE::Chooser::ResolvePropertyChain(Context, Binding, Container, StructType))
-	{
-		if (const FBoolProperty* Property = FindFProperty<FBoolProperty>(StructType, Binding.PropertyBindingChain.Last()))
-		{
-			OutResult = *Property->ContainerPtrToValuePtr<bool>(Container);
-			return true;
-		}
-		
-	    if (const UClass* ClassType = Cast<const UClass>(StructType))
-	    {
-			if (UFunction* Function = ClassType->FindFunctionByName(Binding.PropertyBindingChain.Last()))
-			{
-				UObject* Object = reinterpret_cast<UObject*>(const_cast<void*>(Container));
-				if (Function->IsNative())
-				{
-					FFrame Stack(Object, Function, nullptr, nullptr, Function->ChildProperties);
-					Function->Invoke(Object, Stack, &OutResult);
-				}
-				else
-				{
-					Object->ProcessEvent(Function, &OutResult);
-				}
-			} 
-		}
-	}
-
-	return false;
+	return Binding.GetValue(Context, OutResult);
 }
 
 bool FBoolContextProperty::SetValue(FChooserEvaluationContext& Context, bool InValue) const
 {
-	const UStruct* StructType = nullptr;
-	const void* Container = nullptr;
-	
-	if (UE::Chooser::ResolvePropertyChain(Context, Binding, Container, StructType))
-	{
-		if (FBoolProperty* Property = FindFProperty<FBoolProperty>(StructType, Binding.PropertyBindingChain.Last()))
-		{
-			// const cast is here just because ResolvePropertyChain expects a const void*&
-			*Property->ContainerPtrToValuePtr<bool>(const_cast<void*>(Container)) = InValue;
-			return true;
-		}
-	}
-
-	return false;
+	return Binding.SetValue(Context, InValue);
 }
 
 FBoolColumn::FBoolColumn()
diff --git a/Engine/Plugins/Experimental/Chooser/Source/Chooser/Private/Chooser.cpp b/Engine/Plugins/Experimental/Chooser/Source/Chooser/Private/Chooser.cpp
--- a/Engine/Plugins/Experimental/Chooser/Source/Chooser/Private/Chooser.cpp
+++ b/Engine/Plugins/Experimental/Chooser/Source/Chooser/Private/Chooser.cpp
@@ -2,61 +2,24 @@
 #include "Chooser.h"
 #include "ChooserFunctionLibrary.h"
 #include "ChooserPropertyAccess.h"
+#include "Engine/UserDefinedStruct.h"
+#include "Engine/Blueprint.h"
 
 #include UE_INLINE_GENERATED_CPP_BY_NAME(Chooser)
 
+DEFINE_LOG_CATEGORY(LogChooser)
+
 UChooserTable::UChooserTable(const FObjectInitializer& Initializer)
 	:Super(Initializer)
 {
 
 }
 
-#if WITH_EDITOR
-void UChooserTable::PostEditUndo()
-{
-	UObject::PostEditUndo();
-
-	if (CachedPreviousOutputObjectType != OutputObjectType || CachedPreviousResultType != ResultType)
-	{
-		OnOutputObjectTypeChanged.Broadcast(OutputObjectType);
-		CachedPreviousOutputObjectType = OutputObjectType;
-		CachedPreviousResultType = ResultType;
-	}
-	OnContextClassChanged.Broadcast();
-}
-
-void UChooserTable::PostEditChangeProperty(FPropertyChangedEvent& PropertyChangedEvent)
-{
-	UObject::PostEditChangeProperty(PropertyChangedEvent);
-	
-	static FName OutputObjectTypeName = "OutputObjectType";
-	static FName ResultTypeName = "ResultType";
-	if (PropertyChangedEvent.Property->GetName() == OutputObjectTypeName)
-	{
-		if (CachedPreviousOutputObjectType != OutputObjectType)
-		{
-			OnOutputObjectTypeChanged.Broadcast(OutputObjectType);
-			CachedPreviousOutputObjectType = OutputObjectType;
-		}
-	}
-	else if (PropertyChangedEvent.Property->GetName() == ResultTypeName)
-	{
-		if (CachedPreviousResultType != ResultType)
-		{
-			OnOutputObjectTypeChanged.Broadcast(OutputObjectType);
-			CachedPreviousResultType = ResultType;
-		}
-	}
-	else
-	{
-		OnContextClassChanged.Broadcast();
-	}
-}
-
 void UChooserTable::PostLoad()
 {
 	Super::PostLoad();
 
+#if WITH_EDITOR
 	CachedPreviousOutputObjectType = OutputObjectType;
 	CachedPreviousResultType = ResultType;
 
@@ -101,17 +64,96 @@
 		Results_DEPRECATED.SetNum(0);
 		Columns_DEPRECATED.SetNum(0);
 	}
+#endif
 
-	// call PostLoad on Columns
+	Compile();
+}
+
+#if WITH_EDITOR
+
+void UChooserTable::AddCompileDependency(const UStruct* InStructType)
+{
+	UStruct* StructType = const_cast<UStruct*>(InStructType);
+	if (!CompileDependencies.Contains(StructType))
+	{
+		if (UUserDefinedStruct* UserDefinedStruct = Cast<UUserDefinedStruct>(StructType))
+		{
+			UserDefinedStruct->ChangedEvent.AddUObject(this, &UChooserTable::OnDependentStructChanged);
+		}
+		else if (UClass* Class = Cast<UClass>(StructType))
+		{
+			if(UBlueprint* Blueprint = Cast<UBlueprint>(Class->ClassGeneratedBy))
+			{
+				Blueprint->OnCompiled().AddUObject(this, &UChooserTable::OnDependencyCompiled);
+				CompileDependencies.Add(StructType);
+			}
+		}
+	}
+}
+
+#endif
+
+void UChooserTable::Compile(bool bForce)
+{
 	for (FInstancedStruct& ColumnData : ColumnsStructs)
 	{
 		if (ColumnData.IsValid())
 		{
 			FChooserColumnBase& Column = ColumnData.GetMutable<FChooserColumnBase>();
-			Column.PostLoad();
+			Column.Compile(this, bForce);
 		}
 	}
-		
+
+	for(FInstancedStruct& ResultData : ResultsStructs)
+	{
+		if (ResultData.IsValid())
+		{
+			FObjectChooserBase& Result = ResultData.GetMutable<FObjectChooserBase>();
+			Result.Compile(this, bForce);
+		}
+	}
 }
 
+#if WITH_EDITOR
+void UChooserTable::PostEditUndo()
+{
+	UObject::PostEditUndo();
+
+	if (CachedPreviousOutputObjectType != OutputObjectType || CachedPreviousResultType != ResultType)
+	{
+		OnOutputObjectTypeChanged.Broadcast(OutputObjectType);
+		CachedPreviousOutputObjectType = OutputObjectType;
+		CachedPreviousResultType = ResultType;
+	}
+	OnContextClassChanged.Broadcast();
+}
+
+void UChooserTable::PostEditChangeProperty(FPropertyChangedEvent& PropertyChangedEvent)
+{
+	UObject::PostEditChangeProperty(PropertyChangedEvent);
+	
+	static FName OutputObjectTypeName = "OutputObjectType";
+	static FName ResultTypeName = "ResultType";
+	if (PropertyChangedEvent.Property->GetName() == OutputObjectTypeName)
+	{
+		if (CachedPreviousOutputObjectType != OutputObjectType)
+		{
+			OnOutputObjectTypeChanged.Broadcast(OutputObjectType);
+			CachedPreviousOutputObjectType = OutputObjectType;
+		}
+	}
+	else if (PropertyChangedEvent.Property->GetName() == ResultTypeName)
+	{
+		if (CachedPreviousResultType != ResultType)
+		{
+			OnOutputObjectTypeChanged.Broadcast(OutputObjectType);
+			CachedPreviousResultType = ResultType;
+		}
+	}
+	else
+	{
+		OnContextClassChanged.Broadcast();
+	}
+}
+
 void UChooserTable::IterateRecentContextObjects(TFunction<void(const UObject*)> Callback) const
@@ -150,6 +192,8 @@
 
 FObjectChooserBase::EIteratorStatus UChooserTable::EvaluateChooser(FChooserEvaluationContext& Context, const UChooserTable* Chooser, FObjectChooserBase::FObjectChooserIteratorCallback Callback)
 {
+	TRACE_CPUPROFILER_EVENT_SCOPE(EvaluateChooser);
+	
 	if (Chooser == nullptr)
 	{
 		return FObjectChooserBase::EIteratorStatus::Continue;
@@ -158,7 +202,10 @@
 	// todo validate that parameter types in context data match
 
 #if WITH_EDITOR
-	Chooser->UpdateDebugging(Context);
+	{
+		TRACE_CPUPROFILER_EVENT_SCOPE(EvaluateChooser_Debugging);
+		Chooser->UpdateDebugging(Context);
+	}
 #endif
 
 	TArray<uint32> Indices1;
diff --git a/Engine/Plugins/Experimental/Chooser/Source/Chooser/Private/ChooserPropertyAccess.cpp b/Engine/Plugins/Experimental/Chooser/Source/Chooser/Private/ChooserPropertyAccess.cpp
--- a/Engine/Plugins/Experimental/Chooser/Source/Chooser/Private/ChooserPropertyAccess.cpp
+++ b/Engine/Plugins/Experimental/Chooser/Source/Chooser/Private/ChooserPropertyAccess.cpp
@@ -3,14 +3,329 @@
 #include "ChooserPropertyAccess.h"
 
 #include "IObjectChooser.h"
+#include "Engine/UserDefinedStruct.h"
 #include "UObject/UnrealType.h"
+#include "Logging/LogMacros.h"
 
 #if WITH_EDITOR
+#include "ScopedTransaction.h"
 #include "IPropertyAccessEditor.h"
 #endif
 
+#define LOCTEXT_NAMESPACE "ChooserPropertyAccess"
+		
+FName FChooserPropertyBinding::GetUniqueId() const
+{
+	FString Result;
+	bool First = true;
+	for (const auto& Element : PropertyBindingChain)
+	{
+		if (First)
+		{
+			First = false;
+		}
+		else
+		{
+			Result += TEXT(".");
+		}
+
+		Result += Element.ToString();
+	}
+
+	return FName(Result);
+}
+
+
+struct FCompiledBindingCacheId
+{
+	FName BindingPath;
+	const UStruct* Type;
+	
+	bool operator == (const FCompiledBindingCacheId& Other) const
+	{
+		 return Type == Other.Type
+		  && BindingPath == Other.BindingPath;
+	}
+};
+
+FORCEINLINE uint32 GetTypeHash(const FCompiledBindingCacheId& Id)
+{
+	return HashCombineFast(GetTypeHash(Id.Type), GetTypeHash(Id.BindingPath));
+}
+
+TMap<FCompiledBindingCacheId, TWeakPtr<UE::Chooser::FCompiledBinding>> CompiledBindingCache;
+FCriticalSection CompiledBindingCacheLock;
+
+
+void FChooserPropertyBinding::Compile(IHasContextClass* Owner, bool bForce)
+{
+	TRACE_CPUPROFILER_EVENT_SCOPE(CompilePropertyChain);
+
+#if WITH_EDITORONLY_DATA
+	CompileMessage = FText();
+#endif
+	
+	TConstArrayView<FInstancedStruct> ContextData = Owner->GetContextData();
+
+	int CompiledBindingSerialNumber = 0;
+	
+	if (PropertyBindingChain.IsEmpty() || !ContextData.IsValidIndex(ContextIndex))
+	{
+#if WITH_EDITORONLY_DATA
+		CompileMessage = LOCTEXT("No Property Bound", "No Property Bound");
+#endif
+		CompiledBinding = nullptr;
+		return;
+	}
+
+	const UStruct* StructType = nullptr;
+	if (ContextData.IsValidIndex(ContextIndex))
+	{
+		if (const FContextObjectTypeClass* ClassContext = ContextData[ContextIndex].GetPtr<FContextObjectTypeClass>())
+		{
+			StructType = ClassContext->Class;
+		}
+		else if (const FContextObjectTypeStruct* StructContext = ContextData[ContextIndex].GetPtr<FContextObjectTypeStruct>())
+		{
+			StructType = StructContext->Struct;
+		}
+	}
+
+	if(StructType == nullptr)
+	{
+#if WITH_EDITORONLY_DATA
+		CompileMessage = FText::Format(LOCTEXT("No valid struct", "No valid Context Object/Struct at index: {0}"), FText::FromString(FString::FromInt(ContextIndex)));
+#endif
+		CompiledBinding = nullptr;
+		return;
+	}
+	
+	CompiledBinding->TargetType = StructType;
+	
+	FCompiledBindingCacheId Id;
+	Id.Type = StructType;
+	Id.BindingPath = GetUniqueId();
+
+	{
+		FScopeLock Lock(&CompiledBindingCacheLock);
+		if (TWeakPtr<UE::Chooser::FCompiledBinding>* Binding =  CompiledBindingCache.Find(Id))
+		{
+			if (Binding->IsValid())
+			{
+				bool bUseCachedBinding = true;
+#if WITH_EDITOR
+				if (bForce && CompiledBinding)
+				{
+					// we've triggered a recompile due to dependency changes.
+					// if the one in the cache has a higher serial number, then it has already been recompiled by some other reference, so use the cached one.
+					// otherwise, remove the currently cached one, recompile, and increment the serial number
+					if (Binding->Pin()->SerialNumber <= CompiledBinding->SerialNumber)
+					{
+						bUseCachedBinding = false;
+						CompiledBindingSerialNumber = CompiledBinding->SerialNumber + 1;
+						CompiledBindingCache.Remove(Id);
+					}
+				}
+#endif
+
+				if (bUseCachedBinding)
+				{
+					CompiledBinding = Binding->Pin();
+					return;
+				}
+			}
+		}
+	}
+
+	TSharedPtr<UE::Chooser::FCompiledBinding> NewCompiledBinding = MakeShared<UE::Chooser::FCompiledBinding>();
+	UE::Chooser::FCompiledBinding& OutCompiledBinding = *NewCompiledBinding.Get();
+	
+	OutCompiledBinding.CompiledChain.SetNum(0);
+	OutCompiledBinding.ContextIndex = ContextIndex;
+	
+	const int PropertyChainLength = PropertyBindingChain.Num();
+	int CurrentOffset = 0;
+	for(int PropertyChainIndex = 0; PropertyChainIndex < PropertyChainLength - 1; PropertyChainIndex++)
+	{
+#if WITH_EDITOR
+		Owner->AddCompileDependency(StructType);
+#endif
+
+		bool bFound = false;
+		if (const FStructProperty* StructProperty = FindFProperty<FStructProperty>(StructType, PropertyBindingChain[PropertyChainIndex]))
+		{
+			bFound = true;
+			// accumulate offsets in structs
+			CurrentOffset += StructProperty->GetOffset_ForInternal();
+			StructType = StructProperty->Struct;
+		}
+		else if (const FObjectProperty* ObjectProperty = FindFProperty<FObjectProperty>(StructType, PropertyBindingChain[PropertyChainIndex]))
+		{
+			bFound = true;
+			// when we hit an  object reference, create a chain element with the current offset
+			CurrentOffset += ObjectProperty->GetOffset_ForInternal();
+			OutCompiledBinding.CompiledChain.Add(UE::Chooser::FCompiledBindingElement(CurrentOffset));
+			StructType = ObjectProperty->PropertyClass;
+			// clear the offset, to start accumulating again relative to the new object base
+			CurrentOffset = 0;
+		}
+		// check if it's a member function
+		else if (const UClass* ClassType = Cast<const UClass>(StructType))
+		{
+			if (UFunction* Function = ClassType->FindFunctionByName(PropertyBindingChain[PropertyChainIndex]))
+			{
+				bFound = true;
+				ensure(CurrentOffset == 0);
+				OutCompiledBinding.CompiledChain.Add(UE::Chooser::FCompiledBindingElement(Function));
+				StructType = CastField<FObjectProperty>(Function->GetReturnProperty())->PropertyClass;
+			}
+		}
+
+		if (!bFound)
+		{
+#if WITH_EDITORONLY_DATA
+			CompileMessage = FText::Format(LOCTEXT("Property Not Found", "Property/Function: {0} not Found on Class/Struct: {1}"), FText::FromName(PropertyBindingChain[PropertyChainIndex]), StructType->GetDisplayNameText());
+#endif
+			CompiledBinding = nullptr;
+			return;
+		}
+	}
+	
+	#if WITH_EDITOR
+   	Owner->AddCompileDependency(StructType);
+    #endif
+
+	bool bFound = false;
+	if (const FProperty* BaseProperty = FindFProperty<FProperty>(StructType, PropertyBindingChain.Last()))
+	{
+		bFound = true;
+		
+		// last element should be the actual property - add it's offset to whatever was accumulated from struct offsets
+		CurrentOffset += BaseProperty->GetOffset_ForInternal();
+		OutCompiledBinding.CompiledChain.Add(UE::Chooser::FCompiledBindingElement(CurrentOffset));
+
+		if (BaseProperty->IsA<FFloatProperty>())
+		{
+			OutCompiledBinding.PropertyType = UE::Chooser::EPropertyNumericalType::FLOAT;
+		}
+		else if (BaseProperty->IsA<FDoubleProperty>())
+		{
+			OutCompiledBinding.PropertyType = UE::Chooser::EPropertyNumericalType::DOUBLE;
+		}
+		else if (BaseProperty->IsA<FIntProperty>())
+		{
+			OutCompiledBinding.PropertyType = UE::Chooser::EPropertyNumericalType::INT32;
+		}
+	}
+	else
+	{
+		// handle function calls 
+		if (const UClass* ClassType = Cast<const UClass>(StructType))
+		{
+			if (UFunction* Function = ClassType->FindFunctionByName(PropertyBindingChain.Last()))
+			{
+				bFound = true;
+				
+				const FProperty* ReturnProperty = Function->GetReturnProperty();
+				OutCompiledBinding.CompiledChain.Add(UE::Chooser::FCompiledBindingElement(Function));
+				if (ReturnProperty->IsA<FFloatProperty>())
+				{
+					OutCompiledBinding.PropertyType = UE::Chooser::EPropertyNumericalType::FLOAT;
+				}
+				else if (ReturnProperty->IsA<FDoubleProperty>())
+				{
+					OutCompiledBinding.PropertyType = UE::Chooser::EPropertyNumericalType::DOUBLE;
+				}
+				else if (ReturnProperty->IsA<FIntProperty>())
+				{
+					OutCompiledBinding.PropertyType = UE::Chooser::EPropertyNumericalType::INT32;
+				}
+			}
+		 }
+	}
+
+	if (bFound)
+	{
+		FScopeLock Lock(&CompiledBindingCacheLock);
+#if WITH_EDITORONLY_DATA
+		NewCompiledBinding->SerialNumber = CompiledBindingSerialNumber;
+#endif
+		CompiledBindingCache.Add(Id, NewCompiledBinding);
+		CompiledBinding = NewCompiledBinding;
+	}
+	else
+	{
+#if WITH_EDITORONLY_DATA
+		CompileMessage = FText::Format(LOCTEXT("Property Not Found", "Property/Function: {0} not Found on Class/Struct: {1}"), FText::FromName(PropertyBindingChain.Last()), StructType->GetDisplayNameText());
+#endif
+		CompiledBinding = nullptr;
+ 	}
+}
+
+
 namespace UE::Chooser
 {
+
+	uint8* ResolveCompiledPropertyChain(FChooserEvaluationContext& Context, const FCompiledBinding& CompiledBinding)
+	{
+		uint8* Result = nullptr;
+		const UStruct* InputType = nullptr;
+
+		if(!Context.Params.IsValidIndex(CompiledBinding.ContextIndex))
+		{
+			UE_LOG(LogChooser, Warning, TEXT("Invalid Index {%d} while resolving compiled property chain."), CompiledBinding.ContextIndex);
+			return nullptr;
+		}
+		
+		if (FChooserEvaluationInputObject* ObjectInput = Context.Params[CompiledBinding.ContextIndex].GetMutablePtr<FChooserEvaluationInputObject>())
+		{
+			UObject* Object = ObjectInput->Object.Get();
+			Result = reinterpret_cast<uint8*>(Object);
+			if (Object)
+			{
+				InputType = Object->GetClass();
+			}
+		}
+		else
+		{
+			Result = Context.Params[CompiledBinding.ContextIndex].GetMutableMemory();
+			InputType = Context.Params[CompiledBinding.ContextIndex].GetScriptStruct();
+		}
+
+		if (!InputType->IsChildOf(CompiledBinding.TargetType))
+		{
+			UE_LOG(LogChooser, Warning, TEXT("Property Binding compiled for type: {%s} is being evaluated on incompatible type: {%s}."), ToCStr(CompiledBinding.TargetType->GetName()), ToCStr(InputType->GetName()));
+			return nullptr;
+		}
+		
+		for (int i = 0; i<CompiledBinding.CompiledChain.Num() - 1; i++)
+		{
+			if (Result)
+			{
+				const FCompiledBindingElement& Element = CompiledBinding.CompiledChain[i];
+				if (!Element.bIsFunction)
+				{
+					Result = *reinterpret_cast<uint8**>(Result + CompiledBinding.CompiledChain[i].Offset);
+				}
+				else
+				{
+					UObject* Object = reinterpret_cast<UObject*>(Result);
+					if (Element.Function->IsNative())
+					{
+						FFrame Stack(Object, Element.Function, nullptr, nullptr, Element.Function->ChildProperties);
+						Element.Function->Invoke(Object, Stack, &Result);
+					}
+					else
+					{
+						Object->ProcessEvent(Element.Function, &Result);
+					}
+				}
+			}
+		}
+
+		return Result;
+	}
+	
 	bool ResolvePropertyChain(FChooserEvaluationContext& Context, const FChooserPropertyBinding& PropertyBinding, const void*& OutContainer, const UStruct*& OutStructType)
 	{
 		if (Context.Params.IsValidIndex(PropertyBinding.ContextIndex))
@@ -18,13 +333,25 @@ namespace UE::Chooser
 			if (FChooserEvaluationInputObject* ObjectParam = Context.Params[PropertyBinding.ContextIndex].GetMutablePtr<FChooserEvaluationInputObject>())
 			{
 				OutContainer = ObjectParam->Object;
-				OutStructType = ObjectParam->Object->GetClass();
+				if (OutContainer)
+				{
+					OutStructType = ObjectParam->Object->GetClass();
+				}
+				else
+				{
+					OutStructType = nullptr;
+				}
 			}
 			else
 			{
 				OutContainer = Context.Params[PropertyBinding.ContextIndex].GetMutableMemory();
 				OutStructType = Context.Params[PropertyBinding.ContextIndex].GetScriptStruct();
 			}
+			
+			if (OutContainer == nullptr || OutStructType == nullptr)
+			{
+				return false;
+			}
 
 			return ResolvePropertyChain(OutContainer, OutStructType, PropertyBinding.PropertyBindingChain);
 		}
@@ -120,4 +447,6 @@ namespace UE::Chooser
 #endif
 
 	
-}
\ No newline at end of file
+}
+
+#undef LOCTEXT_NAMESPACE
diff --git a/Engine/Plugins/Experimental/Chooser/Source/Chooser/Private/EnumColumn.cpp b/Engine/Plugins/Experimental/Chooser/Source/Chooser/Private/EnumColumn.cpp
--- a/Engine/Plugins/Experimental/Chooser/Source/Chooser/Private/EnumColumn.cpp
+++ b/Engine/Plugins/Experimental/Chooser/Source/Chooser/Private/EnumColumn.cpp
@@ -8,78 +8,14 @@
 
 bool FEnumContextProperty::GetValue(FChooserEvaluationContext& Context, uint8& OutResult) const
 {
-	const UStruct* StructType = nullptr;
-	const void* Container = nullptr;
-
-	if (UE::Chooser::ResolvePropertyChain(Context, Binding, Container, StructType))
-	{
-		if (const FEnumProperty* EnumProperty = FindFProperty<FEnumProperty>(StructType, Binding.PropertyBindingChain.Last()))
-		{
-			OutResult = *EnumProperty->ContainerPtrToValuePtr<uint8>(Container);
-			return true;
-		}
-
-		if (const FByteProperty* ByteProperty = FindFProperty<FByteProperty>(StructType, Binding.PropertyBindingChain.Last()))
-		{
-			if (ByteProperty->IsEnum())
-			{
-				OutResult = *ByteProperty->ContainerPtrToValuePtr<uint8>(Container);
-				return true;
-			}
-		}
-	}
-
-	return false;
+	return Binding.GetValue(Context, OutResult);
 }
 
 bool FEnumContextProperty::SetValue(FChooserEvaluationContext& Context, uint8 InValue) const
 {
-	const UStruct* StructType = nullptr;
-	const void* Container = nullptr;
-
-	if (UE::Chooser::ResolvePropertyChain(Context, Binding, Container, StructType))
-	{
-		if (const FEnumProperty* EnumProperty = FindFProperty<FEnumProperty>(StructType, Binding.PropertyBindingChain.Last()))
-		{
-			*EnumProperty->ContainerPtrToValuePtr<uint8>(const_cast<void*>(Container)) = InValue;
-			return true;
-		}
-
-		if (const FByteProperty* ByteProperty = FindFProperty<FByteProperty>(StructType, Binding.PropertyBindingChain.Last()))
-		{
-			if (ByteProperty->IsEnum())
-			{
-				*ByteProperty->ContainerPtrToValuePtr<uint8>(const_cast<void*>(Container)) = InValue;
-				return true;
-			}
-		}
-	}
-
-	return false;
-}
-
-#if WITH_EDITOR
-
-void FEnumContextProperty::SetBinding(const TArray<FBindingChainElement>& InBindingChain)
-{
-	const UEnum* PreviousEnum = Binding.Enum;
-	Binding.Enum = nullptr;
-
-	UE::Chooser::CopyPropertyChain(InBindingChain, Binding);
-
-	const FField* Field = InBindingChain.Last().Field.ToField();
-	if (const FEnumProperty* EnumProperty = CastField<FEnumProperty>(Field))
-	{
-		Binding.Enum = EnumProperty->GetEnum();
-	}
-	else if (const FByteProperty* ByteProperty = CastField<FByteProperty>(Field))
-	{
-		Binding.Enum = ByteProperty->Enum;
-	}
+	return Binding.SetValue(Context, InValue);
 }
 
-#endif // WITH_EDITOR
-
 FEnumColumn::FEnumColumn()
 {
 #if WITH_EDITOR
diff --git a/Engine/Plugins/Experimental/Chooser/Source/Chooser/Private/FloatRangeColumn.cpp b/Engine/Plugins/Experimental/Chooser/Source/Chooser/Private/FloatRangeColumn.cpp
--- a/Engine/Plugins/Experimental/Chooser/Source/Chooser/Private/FloatRangeColumn.cpp
+++ b/Engine/Plugins/Experimental/Chooser/Source/Chooser/Private/FloatRangeColumn.cpp
@@ -1,91 +1,17 @@
 // Copyright Epic Games, Inc. All Rights Reserved.
 #include "FloatRangeColumn.h"
-
 #include "ChooserPropertyAccess.h"
 
+
+
 bool FFloatContextProperty::SetValue(FChooserEvaluationContext& Context, double InValue) const
 {
-	const UStruct* StructType = nullptr;
-	const void* Container = nullptr;
-	
-	if (UE::Chooser::ResolvePropertyChain(Context, Binding, Container, StructType))
-	{
-		if (FDoubleProperty* DoubleProperty = FindFProperty<FDoubleProperty>(StructType, Binding.PropertyBindingChain.Last()))
-		{
-			// const cast is here just because ResolvePropertyChain expects a const void*&
-			*DoubleProperty->ContainerPtrToValuePtr<double>(const_cast<void*>(Container)) = InValue;
-			return true;
-		}
-		if (FFloatProperty* FloatProperty = FindFProperty<FFloatProperty>(StructType, Binding.PropertyBindingChain.Last()))
-        {
-			// const cast is here just because ResolvePropertyChain expects a const void*&
-        	*FloatProperty->ContainerPtrToValuePtr<float>(const_cast<void*>(Container)) = InValue;
-        	return true;
-        }
-	}
-
-	return false;
+	return Binding.SetValue(Context, InValue);
 }
 
 bool FFloatContextProperty::GetValue(FChooserEvaluationContext& Context, double& OutResult) const
 {
-	const UStruct* StructType = nullptr;
-	const void* Container = nullptr;
-
-	if (UE::Chooser::ResolvePropertyChain(Context, Binding, Container, StructType))
-	{
-		if (const FDoubleProperty* DoubleProperty = FindFProperty<FDoubleProperty>(StructType, Binding.PropertyBindingChain.Last()))
-		{
-			OutResult = *DoubleProperty->ContainerPtrToValuePtr<double>(Container);
-			return true;
-		}
-		
-		if (const FFloatProperty* FloatProperty = FindFProperty<FFloatProperty>(StructType, Binding.PropertyBindingChain.Last()))
-		{
-			OutResult = *FloatProperty->ContainerPtrToValuePtr<float>(Container);
-			return true;
-		}
-
-	    if (const UClass* ClassType = Cast<const UClass>(StructType))
-	    {
-			if (UFunction* Function = ClassType->FindFunctionByName(Binding.PropertyBindingChain.Last()))
-			{
-				const bool bReturnsDouble = CastField<FDoubleProperty>(Function->GetReturnProperty()) != nullptr;
-					
-				UObject* Object = reinterpret_cast<UObject*>(const_cast<void*>(Container));
-				if (Function->IsNative())
-				{
-					FFrame Stack(Object, Function, nullptr, nullptr, Function->ChildProperties);
-					if (bReturnsDouble)
-					{
-						Function->Invoke(Object, Stack, &OutResult);
-					}
-					else
-					{
-						float result = 0;
-						Function->Invoke(Object, Stack, &result);
-						OutResult = result;
-					}
-				}
-				else
-				{
-					if (bReturnsDouble)
-					{
-						Object->ProcessEvent(Function, &OutResult);
-					}
-					else
-					{
-						float result = 0;
-						Object->ProcessEvent(Function, &result);
-						OutResult = result;
-					}
-				}
-				return true;
-			} 
-		}
-	}
-
-	return false;
+	return Binding.GetValue(Context, OutResult);
 }
 
 FFloatRangeColumn::FFloatRangeColumn()
diff --git a/Engine/Plugins/Experimental/Chooser/Source/Chooser/Private/GameplayTagColumn.cpp b/Engine/Plugins/Experimental/Chooser/Source/Chooser/Private/GameplayTagColumn.cpp
--- a/Engine/Plugins/Experimental/Chooser/Source/Chooser/Private/GameplayTagColumn.cpp
+++ b/Engine/Plugins/Experimental/Chooser/Source/Chooser/Private/GameplayTagColumn.cpp
@@ -4,19 +4,7 @@
 
 bool FGameplayTagContextProperty::GetValue(FChooserEvaluationContext& Context, const FGameplayTagContainer*& OutResult) const
 {
-	const UStruct* StructType = nullptr;
-	const void* Container = nullptr;
-	
-	if (UE::Chooser::ResolvePropertyChain(Context, Binding, Container, StructType))
-	{
-		if (const FStructProperty* Property = FindFProperty<FStructProperty>(StructType, Binding.PropertyBindingChain.Last()))
-		{
-			OutResult = Property->ContainerPtrToValuePtr<FGameplayTagContainer>(Container);
-			return true;
-		}
-	}
-
-	return false;
+	return Binding.GetValuePtr(Context, OutResult);
 }
 
 FGameplayTagColumn::FGameplayTagColumn()
diff --git a/Engine/Plugins/Experimental/Chooser/Source/Chooser/Private/ObjectColumn.cpp b/Engine/Plugins/Experimental/Chooser/Source/Chooser/Private/ObjectColumn.cpp
--- a/Engine/Plugins/Experimental/Chooser/Source/Chooser/Private/ObjectColumn.cpp
+++ b/Engine/Plugins/Experimental/Chooser/Source/Chooser/Private/ObjectColumn.cpp
@@ -11,6 +11,7 @@ bool FObjectContextProperty::GetValue(FChooserEvaluationContext& Context, FSoftO
 	const UStruct* StructType = nullptr;
 	const void* Container = nullptr;
 
+	// todo: update this to use the compiled property chain
 	if (UE::Chooser::ResolvePropertyChain(Context, Binding, Container, StructType))
 	{
 		if (const FObjectPropertyBase* ObjectProperty = FindFProperty<FObjectPropertyBase>(StructType, Binding.PropertyBindingChain.Last()))
@@ -33,24 +34,6 @@ bool FObjectContextProperty::GetValue(FChooserEvaluationContext& Context, FSoftO
 	return false;
 }
 
-#if WITH_EDITOR
-
-void FObjectContextProperty::SetBinding(const TArray<FBindingChainElement>& InBindingChain)
-{
-	const UClass* PreviousClass = Binding.AllowedClass;
-	Binding.AllowedClass = nullptr;
-
-	UE::Chooser::CopyPropertyChain(InBindingChain, Binding);
-
-	const FField* Field = InBindingChain.Last().Field.ToField();
-	if (const FObjectPropertyBase* ObjectProperty = CastField<FObjectPropertyBase>(Field))
-	{
-		Binding.AllowedClass = ObjectProperty->PropertyClass;
-	}
-}
-
-#endif // WITH_EDITOR
-
 FObjectColumn::FObjectColumn()
 {
 #if WITH_EDITOR
diff --git a/Engine/Plugins/Experimental/Chooser/Source/Chooser/Private/OutputStructColumn.cpp b/Engine/Plugins/Experimental/Chooser/Source/Chooser/Private/OutputStructColumn.cpp
--- a/Engine/Plugins/Experimental/Chooser/Source/Chooser/Private/OutputStructColumn.cpp
+++ b/Engine/Plugins/Experimental/Chooser/Source/Chooser/Private/OutputStructColumn.cpp
@@ -9,84 +9,17 @@
 
 bool FStructContextProperty::SetValue(FChooserEvaluationContext& Context, const FInstancedStruct& InValue) const
 {
-	const UStruct* StructType = nullptr;
-	const void* Container = nullptr;
-
-	if (Binding.PropertyBindingChain.IsEmpty())
+	void* TargetData;
+	if (Binding.GetValuePtr(Context, TargetData))
 	{
-		if(Context.Params.IsValidIndex((Binding.ContextIndex)))
-		{
-			// directly bound to context struct
-			if (Context.Params[Binding.ContextIndex].GetScriptStruct() == InValue.GetScriptStruct())
-			{
-				void* TargetData = Context.Params[Binding.ContextIndex].GetMutableMemory();
-				InValue.GetScriptStruct()->CopyScriptStruct(TargetData, InValue.GetMemory());
-			}
-		}
+		InValue.GetScriptStruct()->CopyScriptStruct(TargetData, InValue.GetMemory());
+		return true;
 	}
-	else if (UE::Chooser::ResolvePropertyChain(Context, Binding, Container, StructType))
-	{
-		if (FStructProperty* Property = FindFProperty<FStructProperty>(StructType, Binding.PropertyBindingChain.Last()))
-		{
-			// const cast is here just because ResolvePropertyChain expects a const void*&
-			void* TargetData = Property->ContainerPtrToValuePtr<void>(const_cast<void*>(Container));
-			
-			if (Property->Struct == InValue.GetScriptStruct())
-			{
-				Property->Struct->CopyScriptStruct(TargetData, InValue.GetMemory());
-			}
-
-			return true;
-		}
-	}
-
 	return false;
 }
 
 #if WITH_EDITOR
 
-void FStructContextProperty::SetBinding(const UObject* OuterObject, const TArray<FBindingChainElement>& InBindingChain)
-{
-	Binding.StructType = nullptr;
-
-	UE::Chooser::CopyPropertyChain(InBindingChain, Binding);
-
-	if (Binding.PropertyBindingChain.Num() == 0)
-	{
-		// binding directly to context struct, get struct type from there
-	    if (const IHasContextClass* HasContextClass = Cast<IHasContextClass>(OuterObject))
-	    {
-	    	TConstArrayView<FInstancedStruct> ContextData = HasContextClass->GetContextData();
-	    	if (ContextData.IsValidIndex(Binding.ContextIndex))
-		    {
-	    		if (const FContextObjectTypeStruct* StructContextData = ContextData[Binding.ContextIndex].GetPtr<FContextObjectTypeStruct>())
-	    		{
-					Binding.StructType = StructContextData->Struct;
-	    			Binding.DisplayName = StructContextData->Struct->GetAuthoredName();
-	    		}
-		    }
-	    }
-	}
-	else
-	{
-		const FField* Field = InBindingChain.Last().Field.ToField();
-		if (const FStructProperty* StructProperty = CastField<FStructProperty>(Field))
-		{
-			Binding.StructType = StructProperty->Struct;
-			Binding.DisplayName = StructProperty->GetAuthoredName();
-			if (InBindingChain.Num() > 2)
-			{
-				// add the parent struct name to the display name if there is one
-				
-				if (const FField* NextToLastField = InBindingChain[InBindingChain.Num()-2].Field.ToField())
-				{
-					Binding.DisplayName = NextToLastField->GetAuthoredName() + "." + Binding.DisplayName;
-				}
-			}
-		}
-	}
-}
-
 void FOutputStructColumn::StructTypeChanged()
 {
 	if (InputValue.IsValid())
diff --git a/Engine/Plugins/Experimental/Chooser/Source/Chooser/Private/RandomizeColumn.cpp b/Engine/Plugins/Experimental/Chooser/Source/Chooser/Private/RandomizeColumn.cpp
--- a/Engine/Plugins/Experimental/Chooser/Source/Chooser/Private/RandomizeColumn.cpp
+++ b/Engine/Plugins/Experimental/Chooser/Source/Chooser/Private/RandomizeColumn.cpp
@@ -4,20 +4,7 @@
 
 bool FRandomizeContextProperty::GetValue(FChooserEvaluationContext& Context, const FChooserRandomizationContext*& OutResult) const
 {
-	
-	const UStruct* StructType = nullptr;
-	const void* Container = nullptr;
-	
-	if (UE::Chooser::ResolvePropertyChain(Context, Binding, Container, StructType))
-	{
-		if (FStructProperty* Property = FindFProperty<FStructProperty>(StructType, Binding.PropertyBindingChain.Last()))
-		{
-			OutResult = Property->ContainerPtrToValuePtr<FChooserRandomizationContext>(Container);
-			return true;
-		}
-	}
-
-	return false;
+	return Binding.GetValuePtr(Context,OutResult);
 }
 
 FRandomizeColumn::FRandomizeColumn()
diff --git a/Engine/Plugins/Experimental/Chooser/Source/Chooser/Public/ChooserPropertyAccess.h b/Engine/Plugins/Experimental/Chooser/Source/Chooser/Public/ChooserPropertyAccess.h
--- a/Engine/Plugins/Experimental/Chooser/Source/Chooser/Public/ChooserPropertyAccess.h
+++ b/Engine/Plugins/Experimental/Chooser/Source/Chooser/Public/ChooserPropertyAccess.h
@@ -3,42 +3,69 @@
 #pragma once
 
 #include "CoreMinimal.h"
-#include "UObject/Interface.h"
+#include "IHasContext.h"
 #include "InstancedStruct.h"
 #include "IObjectChooser.h"
-
 #include "ChooserPropertyAccess.generated.h"
 
 #if WITH_EDITOR
 struct FBindingChainElement;
 #endif
 
-UINTERFACE(MinimalAPI)
-class UHasContextClass : public UInterface
-{
-	GENERATED_BODY()
-};
 
-DECLARE_MULTICAST_DELEGATE(FContextClassChanged)
+namespace UE::Chooser
+{
+	struct FCompiledBindingElement
+	{
+		FCompiledBindingElement() : bIsFunction(false), Offset(0)
+		{
+		}
+		
+		explicit FCompiledBindingElement(int InOffset)
+		{
+			bIsFunction = false;
+			Offset = InOffset;
+		}
+		
+		explicit FCompiledBindingElement(UFunction* InFunction)
+		{
+			bIsFunction = true;
+			Function = InFunction;
+		}
+		
+		bool bIsFunction;
+		union 
+		{
+			int Offset;
+			UFunction* Function;
+		};
+	};
 
+	// property type, for numerical conversions
+	enum class EPropertyNumericalType
+	{
+		NONE,
+		INT32,
+		FLOAT,
+		DOUBLE,
+	};
 
-UENUM()
-enum class EObjectChooserResultType
-{
-	ObjectResult,
-	ClassResult,
-};
+	struct FCompiledBinding
+	{
+		int ContextIndex = 0;
+		EPropertyNumericalType PropertyType = EPropertyNumericalType::NONE;
+		// type info for number and enum conversions
+		TArray<FCompiledBindingElement> CompiledChain;
+		const UStruct* TargetType = nullptr;
+#if WITH_EDITORONLY_DATA
+		int SerialNumber = 0;
+#endif
+	};
+}
 
-class IHasContextClass
-{
-	GENERATED_BODY()
-public:
-	FContextClassChanged OnContextClassChanged;
-	virtual TConstArrayView<FInstancedStruct> GetContextData() const { return TConstArrayView<FInstancedStruct>(); }
-};
 
 USTRUCT()
-struct FChooserPropertyBinding 
+struct CHOOSER_API FChooserPropertyBinding 
 {
 	GENERATED_BODY()
 	
@@ -51,7 +78,25 @@ struct FChooserPropertyBinding
 #if WITH_EDITORONLY_DATA
 	UPROPERTY()
 	FString DisplayName;
+
+	FText CompileMessage;
 #endif
+
+	TSharedPtr<UE::Chooser::FCompiledBinding> CompiledBinding;
+
+	FName GetUniqueId() const;
+
+	void Compile(IHasContextClass* HasContext, bool bForce = false);
+
+	
+	template <typename T>
+	bool GetValuePtr(FChooserEvaluationContext& Context, T*& Value) const;
+	
+	template <typename T>
+	bool GetValue(FChooserEvaluationContext& Context, T& Value) const;
+	
+	template <typename T>
+	bool SetValue(FChooserEvaluationContext& Context, const T& Value) const;
 };
 
 USTRUCT()
@@ -87,7 +132,6 @@ struct FChooserStructPropertyBinding : public FChooserPropertyBinding
 #endif
 };
 
-DECLARE_MULTICAST_DELEGATE_OneParam(FChooserOutputObjectTypeChanged, const UClass* OutputObjectType);
 UENUM()
 enum class EContextObjectDirection
 {
@@ -123,13 +167,190 @@ struct FContextObjectTypeStruct : public FContextObjectTypeBase
 	TObjectPtr<UScriptStruct> Struct;
 };
 
-
 namespace UE::Chooser
 {
+	CHOOSER_API uint8* ResolveCompiledPropertyChain(FChooserEvaluationContext& Context, const FCompiledBinding& CompiledBinding);
 	CHOOSER_API bool ResolvePropertyChain(const void*& Container, const UStruct*& StructType, const TArray<FName>& PropertyBindingChain);
 	CHOOSER_API bool ResolvePropertyChain(FChooserEvaluationContext& Context, const FChooserPropertyBinding& Binding, const void*& OutContainer, const UStruct*& OutStructType);
 
 #if WITH_EDITOR
 	CHOOSER_API void CopyPropertyChain(const TArray<FBindingChainElement>& InBindingChain, FChooserPropertyBinding& OutPropertyBinding);
 #endif
-}
\ No newline at end of file
+}
+
+template <typename T>
+bool FChooserPropertyBinding::GetValuePtr(FChooserEvaluationContext& Context, T*& OutResult) const
+{
+	using namespace  UE::Chooser;
+	if (!CompiledBinding.IsValid())
+	{
+		return false;
+	}
+
+	const FCompiledBinding& Binding = *CompiledBinding.Get();
+		
+	if (uint8* Result = ResolveCompiledPropertyChain(Context, Binding))
+	{
+		const FCompiledBindingElement& Element = Binding.CompiledChain.Last();
+		if (!Element.bIsFunction)
+		{
+			OutResult = reinterpret_cast<T*>(Result + Binding.CompiledChain.Last().Offset);
+			return true;
+		}
+	}
+	return false;
+}
+
+template <typename T>
+bool FChooserPropertyBinding::GetValue(FChooserEvaluationContext& Context, T& OutResult) const
+{
+	using namespace  UE::Chooser;
+	if (!CompiledBinding.IsValid())
+	{
+		return false;
+	}
+
+	const FCompiledBinding& Binding = *CompiledBinding.Get();
+		
+	if (uint8* Result = ResolveCompiledPropertyChain(Context, Binding))
+	{
+		const FCompiledBindingElement& Element = Binding.CompiledChain.Last();
+		if (!Element.bIsFunction)
+		{
+			switch (Binding.PropertyType)
+			{
+			case EPropertyNumericalType::FLOAT:
+				{
+					float FloatResult = *reinterpret_cast<float*>(Result + Binding.CompiledChain.Last().Offset);
+					OutResult = static_cast<T>(FloatResult);
+					break;
+				}
+			case EPropertyNumericalType::DOUBLE:
+				{
+					double DoubleResult = *reinterpret_cast<double*>(Result + Binding.CompiledChain.Last().Offset);
+					OutResult = static_cast<T>(DoubleResult);
+					break;
+				}
+			case EPropertyNumericalType::INT32:
+				{
+					int32 IntResult = *reinterpret_cast<int*>(Result + Binding.CompiledChain.Last().Offset);
+					OutResult = static_cast<T>(IntResult);
+					break;
+				}
+			default:
+				OutResult = *reinterpret_cast<T*>(Result + Binding.CompiledChain.Last().Offset);
+				break;
+			}
+		}
+		else
+		{
+			UObject* Object = reinterpret_cast<UObject*>(Result);
+			if (Element.Function->IsNative())
+			{
+				FFrame Stack(Object, Element.Function, nullptr, nullptr, Element.Function->ChildProperties);
+				switch (Binding.PropertyType)
+				{
+				case EPropertyNumericalType::FLOAT:
+					{
+						float FloatResult = 0;
+						Element.Function->Invoke(Object, Stack, &FloatResult);
+						OutResult = static_cast<T>(FloatResult);
+						break;
+					}
+				case EPropertyNumericalType::DOUBLE:
+					{
+						double DoubleResult = 0;
+						Element.Function->Invoke(Object, Stack, &DoubleResult);
+						OutResult = static_cast<T>(DoubleResult);
+						break;
+					}
+				case EPropertyNumericalType::INT32:
+					{
+						int32 IntResult = 0;
+						Element.Function->Invoke(Object, Stack, &IntResult);
+						OutResult = static_cast<T>(IntResult);
+						break;
+					}
+				default:
+					Element.Function->Invoke(Object, Stack, &OutResult);
+					break;
+				}
+			}
+			else
+			{
+				switch (Binding.PropertyType)
+				{
+				case EPropertyNumericalType::FLOAT:
+					{
+						float FloatResult = 0;
+						Object->ProcessEvent(Element.Function, &FloatResult);
+						OutResult = static_cast<T>(FloatResult);
+						break;
+					}
+				case EPropertyNumericalType::DOUBLE:
+					{
+						double DoubleResult = 0;
+						Object->ProcessEvent(Element.Function, &DoubleResult);
+						OutResult = static_cast<T>(DoubleResult);
+						break;
+					}
+				case EPropertyNumericalType::INT32:
+					{
+						int32 IntResult = 0;
+						Object->ProcessEvent(Element.Function, &IntResult);
+						OutResult = static_cast<T>(IntResult);
+						break;
+					}
+				default:
+					Object->ProcessEvent(Element.Function, &OutResult);
+					break;
+				}
+			}
+		}
+		return true;
+	}
+	return false;
+}
+
+template <typename T>
+bool FChooserPropertyBinding::SetValue(FChooserEvaluationContext& Context, const T& InValue) const
+{
+	using namespace  UE::Chooser;
+	if (!CompiledBinding.IsValid())
+	{
+		return false;
+	}
+
+	const FCompiledBinding& Binding = *CompiledBinding.Get();
+		
+	if (uint8* Result = ResolveCompiledPropertyChain(Context, Binding))
+	{
+		const FCompiledBindingElement& Element = Binding.CompiledChain.Last();
+		if (!Element.bIsFunction)
+		{
+			switch (Binding.PropertyType)
+			{
+			case EPropertyNumericalType::FLOAT:
+				{
+					*reinterpret_cast<float*>(Result + Binding.CompiledChain.Last().Offset) = static_cast<float>(InValue);
+					break;
+				}
+			case EPropertyNumericalType::DOUBLE:
+				{
+					*reinterpret_cast<double*>(Result + Binding.CompiledChain.Last().Offset) = static_cast<double>(InValue);
+					break;
+				}
+			case EPropertyNumericalType::INT32:
+				{
+					*reinterpret_cast<int*>(Result + Binding.CompiledChain.Last().Offset) = static_cast<int>(InValue);
+					break;
+				}
+			default:
+				*reinterpret_cast<T*>(Result + Binding.CompiledChain.Last().Offset) = InValue;
+				break;
+			}
+			return true;
+		}
+	}
+	return false;
+}
diff --git a/Engine/Plugins/Experimental/Chooser/Source/Chooser/Public/IChooserColumn.h b/Engine/Plugins/Experimental/Chooser/Source/Chooser/Public/IChooserColumn.h
--- a/Engine/Plugins/Experimental/Chooser/Source/Chooser/Public/IChooserColumn.h
+++ b/Engine/Plugins/Experimental/Chooser/Source/Chooser/Public/IChooserColumn.h
@@ -40,6 +40,15 @@ public:
 	
 	virtual void SetOutputs(FChooserEvaluationContext& Context, int RowIndex) const { }
 
+	virtual void Compile(IHasContextClass* Owner, bool bForce)
+	{
+		if (FChooserParameterBase* Input = GetInputValue())
+		{
+			Input->Compile(Owner, bForce);
+		}
+	};
+
+	virtual FChooserParameterBase* GetInputValue() { return nullptr; };
 #if WITH_EDITOR
 	virtual FName RowValuesPropertyName() { return FName(); }
 	virtual void SetNumRows(int32 NumRows) {}
@@ -49,7 +58,6 @@ public:
 	virtual UScriptStruct* GetInputBaseType() const { return nullptr; };
 	virtual const UScriptStruct* GetInputType() const { return nullptr; };
 	virtual void SetInputType(const UScriptStruct* Type) { };
-	virtual FChooserParameterBase* GetInputValue() { return nullptr; };
 
 	// random columns must go last, and get a special icon
 	// using a virtual fucntion to identify them (rather than hard coding a specific type) to potentially support multiple varieties of randomization column.
@@ -90,14 +98,13 @@ public:
 	}\
 	virtual UScriptStruct* GetInputBaseType() const override { return ParameterType::StaticStruct(); };\
 	virtual const UScriptStruct* GetInputType() const override { return InputValue.IsValid() ? InputValue.GetScriptStruct() : nullptr; };\
-	virtual FChooserParameterBase* GetInputValue() override { return InputValue.IsValid() ? &InputValue.GetMutable<FChooserParameterBase>() : nullptr; };\
 	virtual void SetInputType(const UScriptStruct* Type) override { InputValue.InitializeAs(Type); };
 
-#define CHOOSER_COLUMN_BOILERPLATE(ParameterType) CHOOSER_COLUMN_BOILERPLATE2(ParameterType, RowValues)
-
 #else
 
-#define CHOOSER_COLUMN_BOILERPLATE2(ParameterType, RowValuesProperty)
-#define CHOOSER_COLUMN_BOILERPLATE(ParameterType)
+#define CHOOSER_COLUMN_BOILERPLATE2(ParameterType, RowValuesProperty)\
+	virtual FChooserParameterBase* GetInputValue() override { return InputValue.GetMutablePtr<FChooserParameterBase>(); };\
 
 #endif
+
+#define CHOOSER_COLUMN_BOILERPLATE(ParameterType) CHOOSER_COLUMN_BOILERPLATE2(ParameterType, RowValues)
diff --git a/Engine/Plugins/Experimental/Chooser/Source/Chooser/Public/IChooserParameterBase.h b/Engine/Plugins/Experimental/Chooser/Source/Chooser/Public/IChooserParameterBase.h
--- a/Engine/Plugins/Experimental/Chooser/Source/Chooser/Public/IChooserParameterBase.h
+++ b/Engine/Plugins/Experimental/Chooser/Source/Chooser/Public/IChooserParameterBase.h
@@ -3,7 +3,7 @@
 #pragma once
 
 #include "CoreMinimal.h"
-#include "IObjectChooser.h"
+#include "IHasContext.h"
 #include "IChooserParameterBase.generated.h"
 
 USTRUCT()
@@ -14,6 +14,7 @@ struct FChooserParameterBase
 	virtual void GetDisplayName(FText& OutName) const { }
 
 	virtual void PostLoad() {};
+	virtual void Compile(IHasContextClass* Owner, bool bForce) {};
 
 	virtual ~FChooserParameterBase() {}
 };
\ No newline at end of file
diff --git a/Engine/Plugins/Experimental/Chooser/Source/Chooser/Public/IHasContext.h b/Engine/Plugins/Experimental/Chooser/Source/Chooser/Public/IHasContext.h
new file mode 100644
--- /dev/null
+++ b/Engine/Plugins/Experimental/Chooser/Source/Chooser/Public/IHasContext.h
@@ -0,0 +1,42 @@
+// Copyright Epic Games, Inc. All Rights Reserved.
+
+#pragma once
+
+#include "CoreMinimal.h"
+#include "UObject/Interface.h"
+#include "InstancedStruct.h"
+
+#include "IHasContext.generated.h"
+
+CHOOSER_API DECLARE_LOG_CATEGORY_EXTERN(LogChooser, Log, All);
+
+UINTERFACE(MinimalAPI)
+class UHasContextClass : public UInterface
+{
+	GENERATED_BODY()
+};
+
+DECLARE_MULTICAST_DELEGATE(FContextClassChanged)
+
+class IHasContextClass
+{
+	GENERATED_BODY()
+public:
+	FContextClassChanged OnContextClassChanged;
+	virtual TConstArrayView<FInstancedStruct> GetContextData() const { return TConstArrayView<FInstancedStruct>(); }
+
+	virtual void Compile(bool bForce = false) {}
+#if WITH_EDITOR
+	virtual void AddCompileDependency(const UStruct* Struct) { }
+#endif
+
+};
+
+DECLARE_MULTICAST_DELEGATE_OneParam(FChooserOutputObjectTypeChanged, const UClass* OutputObjectType);
+
+UENUM()
+enum class EObjectChooserResultType
+{
+	ObjectResult,
+	ClassResult,
+};
diff --git a/Engine/Plugins/Experimental/Chooser/Source/Chooser/Public/IObjectChooser.h b/Engine/Plugins/Experimental/Chooser/Source/Chooser/Public/IObjectChooser.h
--- a/Engine/Plugins/Experimental/Chooser/Source/Chooser/Public/IObjectChooser.h
+++ b/Engine/Plugins/Experimental/Chooser/Source/Chooser/Public/IObjectChooser.h
@@ -59,6 +59,8 @@ public:
 
 	DECLARE_DELEGATE_RetVal_OneParam( EIteratorStatus, FObjectChooserIteratorCallback, UObject*);
 
+	virtual void Compile(class IHasContextClass* HasContext, bool bForce) {};
+
 	virtual UObject* ChooseObject(FChooserEvaluationContext& Context) const { return nullptr; };
 	virtual EIteratorStatus ChooseMulti(FChooserEvaluationContext& ContextData, FObjectChooserIteratorCallback Callback) const
 	{
diff --git a/Engine/Plugins/Experimental/Chooser/Source/ChooserEditor/Private/BoolColumnEditor.cpp b/Engine/Plugins/Experimental/Chooser/Source/ChooserEditor/Private/BoolColumnEditor.cpp
--- a/Engine/Plugins/Experimental/Chooser/Source/ChooserEditor/Private/BoolColumnEditor.cpp
+++ b/Engine/Plugins/Experimental/Chooser/Source/ChooserEditor/Private/BoolColumnEditor.cpp
@@ -10,6 +10,7 @@
 #include "Widgets/Input/SCheckBox.h"
 #include "Widgets/Images/SImage.h"
 #include "GraphEditorSettings.h"
+#include "ScopedTransaction.h"
 
 #define LOCTEXT_NAMESPACE "BoolColumnEditor"
 
@@ -191,14 +192,7 @@ TSharedRef<SWidget> CreateBoolPropertyWidget(bool bReadOnly, UObject* Transactio
 
 	return SNew(SPropertyAccessChainWidget).ContextClassOwner(HasContextClass).AllowFunctions(true).BindingColor("BooleanPinTypeColor").TypeFilter("bool")
 	.PropertyBindingValue(&ContextProperty->Binding)
-	.OnAddBinding_Lambda(
-		[ContextProperty, TransactionObject, ValueChanged](FName InPropertyName, const TArray<FBindingChainElement>& InBindingChain)
-		{
-			const FScopedTransaction Transaction(NSLOCTEXT("ContextPropertyWidget", "Change Property Binding", "Change Property Binding"));
-			TransactionObject->Modify(true);
-			ContextProperty->SetBinding(InBindingChain);
-			ValueChanged.ExecuteIfBound();
-		});
+	.OnValueChanged(ValueChanged);
 }
 	
 void RegisterBoolWidgets()
diff --git a/Engine/Plugins/Experimental/Chooser/Source/ChooserEditor/Private/EnumColumnEditor.cpp b/Engine/Plugins/Experimental/Chooser/Source/ChooserEditor/Private/EnumColumnEditor.cpp
--- a/Engine/Plugins/Experimental/Chooser/Source/ChooserEditor/Private/EnumColumnEditor.cpp
+++ b/Engine/Plugins/Experimental/Chooser/Source/ChooserEditor/Private/EnumColumnEditor.cpp
@@ -10,6 +10,7 @@
 #include "SEnumCombo.h"
 #include "TransactionCommon.h"
 #include "Widgets/Input/SButton.h"
+#include "ScopedTransaction.h"
 
 #define LOCTEXT_NAMESPACE "EnumColumnEditor"
 
@@ -278,14 +279,7 @@ TSharedRef<SWidget> CreateEnumPropertyWidget(bool bReadOnly, UObject* Transactio
 
 	return SNew(SPropertyAccessChainWidget).ContextClassOwner(HasContextClass).AllowFunctions(false).BindingColor("BytePinTypeColor").TypeFilter("enum")
 	.PropertyBindingValue(&ContextProperty->Binding)
-	.OnAddBinding_Lambda(
-		[ContextProperty, TransactionObject, ValueChanged](FName InPropertyName, const TArray<FBindingChainElement>& InBindingChain)
-		{
-			const FScopedTransaction Transaction(NSLOCTEXT("ContextPropertyWidget", "Change Property Binding", "Change Property Binding"));
-			TransactionObject->Modify(true);
-			ContextProperty->SetBinding(InBindingChain);
-			ValueChanged.ExecuteIfBound();
-		});
+	.OnValueChanged(ValueChanged);
 }
 	
 void RegisterEnumWidgets()
diff --git a/Engine/Plugins/Experimental/Chooser/Source/ChooserEditor/Private/FloatRangeColumnEditor.cpp b/Engine/Plugins/Experimental/Chooser/Source/ChooserEditor/Private/FloatRangeColumnEditor.cpp
--- a/Engine/Plugins/Experimental/Chooser/Source/ChooserEditor/Private/FloatRangeColumnEditor.cpp
+++ b/Engine/Plugins/Experimental/Chooser/Source/ChooserEditor/Private/FloatRangeColumnEditor.cpp
@@ -8,11 +8,13 @@
 #include "Widgets/Images/SImage.h"
 #include "GraphEditorSettings.h"
 #include "SPropertyAccessChainWidget.h"
+#include "ScopedTransaction.h"
 
 #define LOCTEXT_NAMESPACE "FloatRangeColumnEditor"
 
 namespace UE::ChooserEditor
 {
+
 TSharedRef<SWidget> CreateFloatPropertyWidget(bool bReadOnly, UObject* TransactionObject, void* Value, UClass* ResultBaseClass, FChooserWidgetValueChanged ValueChanged)
 {
 	IHasContextClass* HasContextClass = Cast<IHasContextClass>(TransactionObject);
@@ -21,14 +23,7 @@ TSharedRef<SWidget> CreateFloatPropertyWidget(bool bReadOnly, UObject* Transacti
 
 	return SNew(SPropertyAccessChainWidget).ContextClassOwner(HasContextClass).AllowFunctions(true).BindingColor("FloatPinTypeColor").TypeFilter("double")
 	.PropertyBindingValue(&ContextProperty->Binding)
-	.OnAddBinding_Lambda(
-		[ContextProperty, TransactionObject, ValueChanged](FName InPropertyName, const TArray<FBindingChainElement>& InBindingChain)
-		{
-			const FScopedTransaction Transaction(NSLOCTEXT("ContextPropertyWidget", "Change Property Binding", "Change Property Binding"));
-			TransactionObject->Modify(true);
-			ContextProperty->SetBinding(InBindingChain);
-			ValueChanged.ExecuteIfBound();
-		});
+	.OnValueChanged(ValueChanged);
 }
 	
 TSharedRef<SWidget> CreateFloatRangeColumnWidget(UChooserTable* Chooser, FChooserColumnBase* Column, int Row)
diff --git a/Engine/Plugins/Experimental/Chooser/Source/ChooserEditor/Private/GameplayTagColumnEditor.cpp b/Engine/Plugins/Experimental/Chooser/Source/ChooserEditor/Private/GameplayTagColumnEditor.cpp
--- a/Engine/Plugins/Experimental/Chooser/Source/ChooserEditor/Private/GameplayTagColumnEditor.cpp
+++ b/Engine/Plugins/Experimental/Chooser/Source/ChooserEditor/Private/GameplayTagColumnEditor.cpp
@@ -119,14 +119,7 @@ TSharedRef<SWidget> CreateGameplayTagPropertyWidget(bool bReadOnly, UObject* Tra
 
 	return SNew(SPropertyAccessChainWidget).ContextClassOwner(HasContextClass).AllowFunctions(false).BindingColor("StructPinTypeColor").TypeFilter("FGameplayTagContainer")
 	.PropertyBindingValue(&ContextProperty->Binding)
-	.OnAddBinding_Lambda(
-		[ContextProperty, TransactionObject, ValueChanged](FName InPropertyName, const TArray<FBindingChainElement>& InBindingChain)
-		{
-			const FScopedTransaction Transaction(NSLOCTEXT("ContextPropertyWidget", "Change Property Binding", "Change Property Binding"));
-			TransactionObject->Modify(true);
-			ContextProperty->SetBinding(InBindingChain);
-			ValueChanged.ExecuteIfBound();
-		});
+	.OnValueChanged(ValueChanged);
 }
 	
 void RegisterGameplayTagWidgets()
diff --git a/Engine/Plugins/Experimental/Chooser/Source/ChooserEditor/Private/ObjectColumnEditor.cpp b/Engine/Plugins/Experimental/Chooser/Source/ChooserEditor/Private/ObjectColumnEditor.cpp
--- a/Engine/Plugins/Experimental/Chooser/Source/ChooserEditor/Private/ObjectColumnEditor.cpp
+++ b/Engine/Plugins/Experimental/Chooser/Source/ChooserEditor/Private/ObjectColumnEditor.cpp
@@ -10,6 +10,7 @@
 #include "Widgets/Images/SImage.h"
 #include "Widgets/Input/SButton.h"
 #include "Widgets/Layout/SBox.h"
+#include "ScopedTransaction.h"
 
 #define LOCTEXT_NAMESPACE "ObjectColumnEditor"
 
@@ -223,14 +224,7 @@ namespace UE::ChooserEditor
 
 		return SNew(SPropertyAccessChainWidget).ContextClassOwner(HasContextClass).AllowFunctions(false).BindingColor("ObjectPinTypeColor").TypeFilter("object")
 		.PropertyBindingValue(&ContextProperty->Binding)
-		.OnAddBinding_Lambda(
-			[ContextProperty, TransactionObject, ValueChanged](FName InPropertyName, const TArray<FBindingChainElement>& InBindingChain)
-			{
-				const FScopedTransaction Transaction(NSLOCTEXT("ContextPropertyWidget", "Change Property Binding", "Change Property Binding"));
-				TransactionObject->Modify(true);
-				ContextProperty->SetBinding(InBindingChain);
-				ValueChanged.ExecuteIfBound();
-			});
+		.OnValueChanged(ValueChanged);
 	}
 
 	void RegisterObjectWidgets()
diff --git a/Engine/Plugins/Experimental/Chooser/Source/ChooserEditor/Private/OutputStructColumnEditor.cpp b/Engine/Plugins/Experimental/Chooser/Source/ChooserEditor/Private/OutputStructColumnEditor.cpp
--- a/Engine/Plugins/Experimental/Chooser/Source/ChooserEditor/Private/OutputStructColumnEditor.cpp
+++ b/Engine/Plugins/Experimental/Chooser/Source/ChooserEditor/Private/OutputStructColumnEditor.cpp
@@ -62,14 +62,7 @@ TSharedRef<SWidget> CreateStructPropertyWidget(bool bReadOnly, UObject* Transact
 
 	return SNew(SPropertyAccessChainWidget).ContextClassOwner(HasContextClass).BindingColor("StructPinTypeColor").TypeFilter("struct")
 		.PropertyBindingValue(&ContextProperty->Binding)
-		.OnAddBinding_Lambda(
-    		[ContextProperty, TransactionObject, ValueChanged](FName InPropertyName, const TArray<FBindingChainElement>& InBindingChain)
-    		{
-    			const FScopedTransaction Transaction(NSLOCTEXT("ContextPropertyWidget", "Change Property Binding", "Change Property Binding"));
-    			TransactionObject->Modify(true);
-    			ContextProperty->SetBinding(TransactionObject, InBindingChain);
-    			ValueChanged.ExecuteIfBound();
-    		});
+		.OnValueChanged(ValueChanged);
 }
 	
 void RegisterStructWidgets()
diff --git a/Engine/Plugins/Experimental/Chooser/Source/ChooserEditor/Private/PropertyAccessChainCustomization.cpp b/Engine/Plugins/Experimental/Chooser/Source/ChooserEditor/Private/PropertyAccessChainCustomization.cpp
--- a/Engine/Plugins/Experimental/Chooser/Source/ChooserEditor/Private/PropertyAccessChainCustomization.cpp
+++ b/Engine/Plugins/Experimental/Chooser/Source/ChooserEditor/Private/PropertyAccessChainCustomization.cpp
@@ -60,66 +60,12 @@ void FPropertyAccessChainCustomization::CustomizeHeader(TSharedRef<IPropertyHand
 		{
 			void* data;
 			PropertyHandle->GetValueData(data);
-			return reinterpret_cast<const FChooserPropertyBinding*>(data);
+			return reinterpret_cast<FChooserPropertyBinding*>(data);
+		})
+		.OnValueChanged_Lambda([PropertyHandle]()
+		{
+			PropertyHandle->NotifyPostChange(EPropertyChangeType::ValueSet);
 		})
-		.OnAddBinding_Lambda([PropertyHandle, TypeFilter](FName InPropertyName, const TArray<FBindingChainElement>& InBindingChain)
-         	 	{
-         	 		TArray<UObject*> OuterObjects;
-         	 		PropertyHandle->GetOuterObjects(OuterObjects);
-         	 		const FScopedTransaction Transaction(NSLOCTEXT("ChooserPropertyBinding", "Change Property Binding", "Change Property Binding"));
-         	 	
-         			for (uint32 i=0; i<PropertyHandle->GetNumOuterObjects(); i++)
-         	 		{
-         	 			void* ValuePtr = nullptr;
-         	 			PropertyHandle->GetValueData(ValuePtr);// todo get per object value data?
-         	 			FChooserPropertyBinding* PropertyValue = reinterpret_cast<FChooserPropertyBinding*>(ValuePtr);
-         	 		
-         	 			if (PropertyValue != nullptr)
-         	 			{
-         	 				PropertyHandle->NotifyPreChange();
-         	 				
-         	 				OuterObjects[i]->Modify(true);
-         	 				Chooser::CopyPropertyChain(InBindingChain, *PropertyValue);
-         	
-         	 				if (TypeFilter == "enum")
-         	 				{
-         	 					FField* Property = InBindingChain.Last().Field.ToField();
-         	 					FChooserEnumPropertyBinding* EnumPropertyValue = static_cast<FChooserEnumPropertyBinding*>(PropertyValue);
-         	 					
-         	 					if (const FEnumProperty* EnumProperty = CastField<const FEnumProperty>(Property))
-         	 					{
-         	 						EnumPropertyValue->Enum = EnumProperty->GetEnum();
-         	 					}
-         	 					else if (const FByteProperty* ByteProperty = CastField<const FByteProperty>(Property))
-         	 					{
-         	 						EnumPropertyValue->Enum = ByteProperty->Enum;
-         	 					}
-         	 				}
-							if (TypeFilter == "object")
-							{
-								FField* Property = InBindingChain.Last().Field.ToField();
-								FChooserObjectPropertyBinding* ObjectPropertyBinding = static_cast<FChooserObjectPropertyBinding*>(PropertyValue);
-								
-								if (const FObjectPropertyBase* ObjectProperty = CastField<const FObjectPropertyBase>(Property))
-								{
-									ObjectPropertyBinding->AllowedClass = ObjectProperty->PropertyClass;
-								}
-							}
-         	 				if (TypeFilter == "struct")
-							{
-								FField* Property = InBindingChain.Last().Field.ToField();
-								FChooserStructPropertyBinding* StructPropertyBinding = static_cast<FChooserStructPropertyBinding*>(PropertyValue);
-								
-								if (const FStructProperty* StructProperty = CastField<const FStructProperty>(Property))
-								{
-									StructPropertyBinding->StructType = StructProperty->Struct;
-								}
-							}
-         	
-         	 				PropertyHandle->NotifyPostChange(EPropertyChangeType::ValueSet);
-         	 			}
-         	 		}
-         	 	})
 	];
 }
 	
diff --git a/Engine/Plugins/Experimental/Chooser/Source/ChooserEditor/Private/RandomizeColumnEditor.cpp b/Engine/Plugins/Experimental/Chooser/Source/ChooserEditor/Private/RandomizeColumnEditor.cpp
--- a/Engine/Plugins/Experimental/Chooser/Source/ChooserEditor/Private/RandomizeColumnEditor.cpp
+++ b/Engine/Plugins/Experimental/Chooser/Source/ChooserEditor/Private/RandomizeColumnEditor.cpp
@@ -9,6 +9,7 @@
 #include "Widgets/Layout/SSpacer.h"
 #include "Widgets/Input/SNumericEntryBox.h"
 #include "GraphEditorSettings.h"
+#include "ScopedTransaction.h"
 
 #define LOCTEXT_NAMESPACE "RandomizeColumnEditor"
 
@@ -75,27 +76,8 @@ TSharedRef<SWidget> CreateRandomizeColumnWidget(UChooserTable* Chooser, FChooser
 		+ SHorizontalBox::Slot().FillWidth(1);
 }
 	
-TSharedRef<SWidget> CreateRandomizePropertyWidget(bool bReadOnly, UObject* TransactionObject, void* Value, UClass* ResultBaseClass, FChooserWidgetValueChanged ValueChanged)
-{
-	IHasContextClass* HasContextClass = Cast<IHasContextClass>(TransactionObject);
-
-	FRandomizeContextProperty* ContextProperty = reinterpret_cast<FRandomizeContextProperty*>(Value);
-
-	return SNew(SPropertyAccessChainWidget).ContextClassOwner(HasContextClass).AllowFunctions(true).BindingColor("StructPinTypeColor").TypeFilter("FChooserRandomizationContext")
-	.PropertyBindingValue(&ContextProperty->Binding)
-	.OnAddBinding_Lambda(
-		[ContextProperty, TransactionObject, ValueChanged](FName InPropertyName, const TArray<FBindingChainElement>& InBindingChain)
-		{
-			const FScopedTransaction Transaction(NSLOCTEXT("ContextPropertyWidget", "Change Property Binding", "Change Property Binding"));
-			TransactionObject->Modify(true);
-			ContextProperty->SetBinding(InBindingChain);
-			ValueChanged.ExecuteIfBound();
-		});
-}
-	
 void RegisterRandomizeWidgets()
 {
-	FObjectChooserWidgetFactories::RegisterWidgetCreator(FRandomizeContextProperty::StaticStruct(), CreateRandomizePropertyWidget);
 	FObjectChooserWidgetFactories::RegisterColumnWidgetCreator(FRandomizeColumn::StaticStruct(), CreateRandomizeColumnWidget);
 }
 	
diff --git a/Engine/Plugins/Experimental/Chooser/Source/ChooserEditor/Private/SPropertyAccessChainWidget.cpp b/Engine/Plugins/Experimental/Chooser/Source/ChooserEditor/Private/SPropertyAccessChainWidget.cpp
--- a/Engine/Plugins/Experimental/Chooser/Source/ChooserEditor/Private/SPropertyAccessChainWidget.cpp
+++ b/Engine/Plugins/Experimental/Chooser/Source/ChooserEditor/Private/SPropertyAccessChainWidget.cpp
@@ -4,11 +4,9 @@
 
 #include "Chooser.h"
 #include "ChooserPropertyAccess.h"
-#include "PropertyHandle.h"
 #include "DetailWidgetRow.h"
 #include "GraphEditorSettings.h"
 #include "SClassViewer.h"
-#include "IDetailChildrenBuilder.h"
 #include "IPropertyAccessEditor.h"
 #include "ScopedTransaction.h"
 #include "Features/IModularFeatures.h"
@@ -51,7 +49,7 @@ TSharedRef<SWidget> SPropertyAccessChainWidget::CreatePropertyAccessWidget()
 		if (TypeFilter == "double")
 		{
 			// special case for doubles to bind to either floats or doubles
-			return Property->GetCPPType() == "float" || Property->GetCPPType() == "double";
+			return Property->GetCPPType() == "float" || Property->GetCPPType() == "double" || Property->GetCPPType() == "int32";
 		}
 		else if (TypeFilter == "enum")
 		{
@@ -74,15 +72,19 @@ TSharedRef<SWidget> SPropertyAccessChainWidget::CreatePropertyAccessWidget()
 	// allow struct bindings to bind context structs directly
 	Args.OnCanBindToContextStruct = FOnCanBindToContextStruct::CreateLambda([this](UStruct* StructType)
 	{
-		if (TypeFilter == "struct" && !StructType->IsChildOf(UObject::StaticClass()))
+		if (StructType)
 		{
-			// struct bindings can bind any type of struct
-			return true;
-		}
-		else
-		{
-			return StructType->GetName()  == TypeFilter;
+			if (TypeFilter == "struct" && !StructType->IsChildOf(UObject::StaticClass()))
+			{
+				// struct bindings can bind any type of struct
+				return true;
+			}
+			else
+			{
+				return StructType->GetName()  == TypeFilter;
+			}
 		}
+		return false;
 	});
 
 	Args.OnCanBindProperty = FOnCanBindProperty::CreateLambda(CanBindProperty);
@@ -135,31 +137,108 @@ TSharedRef<SWidget> SPropertyAccessChainWidget::CreatePropertyAccessWidget()
 			return InProperty->HasAnyPropertyFlags(CPF_BlueprintVisible);
 		});	
 
-	Args.OnAddBinding = OnAddBinding;
+	if (OnAddBinding.IsBound())
+	{
+		Args.OnAddBinding = OnAddBinding; 
+	}
+	else
+	{
+		Args.OnAddBinding = FOnAddBinding::CreateLambda( 
+			[this](FName InPropertyName, const TArray<FBindingChainElement>& InBindingChain)
+			{
+				if (PropertyBindingValue.IsSet())
+				{
+					UObject* TransactionObject = Cast<UObject>(ContextClassOwner);
+					FChooserPropertyBinding* ContextProperty = PropertyBindingValue.Get();
+					const FScopedTransaction Transaction(NSLOCTEXT("ContextPropertyWidget", "Change Property Binding", "Change Property Binding"));
+					TransactionObject->Modify(true);
 
+					// todo: move these to a virtual function on FChooserPropertyBinding
+					Chooser::CopyPropertyChain(InBindingChain, *PropertyBindingValue.Get());
+					
+					FField* Property = InBindingChain.Last().Field.ToField();
+					
+					ContextProperty->DisplayName = "";
+					if (Property)
+					{
+						ContextProperty->DisplayName = Property->GetDisplayNameText().ToString();
+						static const int ShortNameLength = 5;
+						if (ContextProperty->DisplayName.Len() < ShortNameLength && InBindingChain.Num() > 2)
+						{
+							FField* ParentProperty = InBindingChain[InBindingChain.Num() - 2].Field.ToField();
+							ContextProperty->DisplayName = ParentProperty->GetDisplayNameText().ToString() + "." + ContextProperty->DisplayName;
+						}
+					}
+					
+					if (TypeFilter == "enum")
+					{
+						FChooserEnumPropertyBinding* EnumPropertyValue = static_cast<FChooserEnumPropertyBinding*>(ContextProperty);
+						
+						if (const FEnumProperty* EnumProperty = CastField<const FEnumProperty>(Property))
+						{
+							EnumPropertyValue->Enum = EnumProperty->GetEnum();
+						}
+						else if (const FByteProperty* ByteProperty = CastField<const FByteProperty>(Property))
+						{
+							EnumPropertyValue->Enum = ByteProperty->Enum;
+						}
+					}
+					if (TypeFilter == "object")
+					{
+						FChooserObjectPropertyBinding* ObjectPropertyBinding = static_cast<FChooserObjectPropertyBinding*>(ContextProperty);
+						
+						if (const FObjectPropertyBase* ObjectProperty = CastField<const FObjectPropertyBase>(Property))
+						{
+							ObjectPropertyBinding->AllowedClass = ObjectProperty->PropertyClass;
+						}
+					}
+					if (TypeFilter == "struct")
+					{
+						FChooserStructPropertyBinding* StructPropertyBinding = static_cast<FChooserStructPropertyBinding*>(ContextProperty);
+						
+						if (const FStructProperty* StructProperty = CastField<const FStructProperty>(Property))
+						{
+							StructPropertyBinding->StructType = StructProperty->Struct;
+						}
+					}
+
+					ContextProperty->Compile(ContextClassOwner);
+					OnValueChanged.ExecuteIfBound();
+				}
+			});
+	}
 
 	Args.CurrentBindingToolTipText = MakeAttributeLambda([this]()
+	{
+		const FText Bind = NSLOCTEXT("ContextPropertyWidget", "Bind", "Bind");
+		FText CurrentValue = Bind;
+		
+		const FChooserPropertyBinding* PropertyValue = PropertyBindingValue.Get();
+		if (PropertyValue != nullptr)
+		{
+			if (!PropertyValue->CompileMessage.IsEmpty())
 			{
-				const FText Bind = NSLOCTEXT("ContextPropertyWidget", "Bind", "Bind");
-				FText CurrentValue = Bind;
-				
-				const FChooserPropertyBinding* PropertyValue = PropertyBindingValue.Get();
-
-				if (PropertyValue != nullptr && PropertyValue->PropertyBindingChain.Num()>0)
+				CurrentValue = PropertyValue->CompileMessage;
+			}
+			else
+			{
+				if (PropertyValue->PropertyBindingChain.Num()>0)
 				{
 					TArray<FText> BindingChainText;
 					BindingChainText.Reserve(PropertyValue->PropertyBindingChain.Num());
-				 
+			 
 					for (const FName& Name : PropertyValue->PropertyBindingChain)
 					{
 						BindingChainText.Add(FText::FromName(Name));
 					}
-					
+				
 					CurrentValue = FText::Join(NSLOCTEXT("ContextPropertyWidget", "PropertyPathSeparator","."), BindingChainText);
 				}
-	
-				return CurrentValue;	
-			});
+			}
+		}
+
+		return CurrentValue;	
+	});
 	
 	Args.CurrentBindingText = MakeAttributeLambda([this]()
 			{
@@ -181,7 +260,10 @@ TSharedRef<SWidget> SPropertyAccessChainWidget::CreatePropertyAccessWidget()
 							{
 								if (const FContextObjectTypeStruct* StructType = ContextData[PropertyValue->ContextIndex].GetPtr<FContextObjectTypeStruct>())
 								{
-									CurrentValue = FText::FromString(StructType->Struct->GetAuthoredName());
+									if (StructType->Struct)
+									{
+										CurrentValue = FText::FromString(StructType->Struct->GetAuthoredName());
+									}
 								}
 							}
 						}
@@ -212,10 +294,19 @@ TSharedRef<SWidget> SPropertyAccessChainWidget::CreatePropertyAccessWidget()
 				return CurrentValue;
 			});
 
-	Args.CurrentBindingImage = MakeAttributeLambda([]() -> const FSlateBrush*
+	Args.CurrentBindingImage = MakeAttributeLambda([this]() -> const FSlateBrush*
 		{
 			static FName PropertyIcon(TEXT("Kismet.Tabs.Variables"));
-			return FAppStyle::GetBrush(PropertyIcon);
+			static FName WarningIcon(TEXT("Icons.WarningWithColor"));
+			bool bHasWarning = false;
+		
+			const FChooserPropertyBinding* PropertyValue = PropertyBindingValue.Get();
+			if (PropertyValue != nullptr)
+			{
+				bHasWarning = !PropertyValue->CompileMessage.IsEmpty();
+			}
+		
+			return FAppStyle::GetBrush(bHasWarning ? WarningIcon : PropertyIcon);
 		});
 	
 	TArray<FBindingContextStruct> ContextStructs;
@@ -253,7 +344,7 @@ void SPropertyAccessChainWidget::Construct( const FArguments& InArgs)
 	BindingColor = InArgs._BindingColor;
 	ContextClassOwner = InArgs._ContextClassOwner;
 	bAllowFunctions = InArgs._AllowFunctions;
-	OnAddBinding = InArgs._OnAddBinding;
+	OnValueChanged = InArgs._OnValueChanged;
 	PropertyBindingValue = InArgs._PropertyBindingValue;
 	UpdateWidget();
 
diff --git a/Engine/Plugins/Experimental/Chooser/Source/ChooserEditor/Public/SPropertyAccessChainWidget.h b/Engine/Plugins/Experimental/Chooser/Source/ChooserEditor/Public/SPropertyAccessChainWidget.h
--- a/Engine/Plugins/Experimental/Chooser/Source/ChooserEditor/Public/SPropertyAccessChainWidget.h
+++ b/Engine/Plugins/Experimental/Chooser/Source/ChooserEditor/Public/SPropertyAccessChainWidget.h
@@ -1,16 +1,15 @@
 // Copyright Epic Games, Inc. All Rights Reserved.
 #pragma once
 
-#include "ScopedTransaction.h"
 #include "IPropertyAccessEditor.h"
-#include "Features/IModularFeatures.h"
 #include "Kismet2/BlueprintEditorUtils.h"
-#include "Styling/AppStyle.h"
 #include "Widgets/SCompoundWidget.h"
+#include "IHasContext.h"
 #include "ChooserPropertyAccess.h"
 
 namespace UE::ChooserEditor
 {
+DECLARE_DELEGATE(FPropertyAccessChainChanged)
 		
 // Wrapper widget for Property access widget, which can update when the target Class changes
 class CHOOSEREDITOR_API SPropertyAccessChainWidget : public SCompoundWidget
@@ -24,7 +23,8 @@ public:
 	SLATE_ARGUMENT(FString, BindingColor);
 	SLATE_ARGUMENT(bool, AllowFunctions);
 	SLATE_EVENT(FOnAddBinding, OnAddBinding);
-	SLATE_ATTRIBUTE(const FChooserPropertyBinding*, PropertyBindingValue);
+	SLATE_EVENT(FPropertyAccessChainChanged, OnValueChanged);
+	SLATE_ATTRIBUTE(FChooserPropertyBinding*, PropertyBindingValue);
             
 	SLATE_END_ARGS()
 
@@ -41,7 +41,8 @@ private:
 	FDelegateHandle ContextClassChangedHandle();
 	IHasContextClass* ContextClassOwner = nullptr;
 	FOnAddBinding OnAddBinding;
-	TAttribute<const FChooserPropertyBinding*> PropertyBindingValue;
+	FPropertyAccessChainChanged OnValueChanged;
+	TAttribute<FChooserPropertyBinding*> PropertyBindingValue;
 	bool bAllowFunctions;
 };
 	
diff --git a/Engine/Plugins/Experimental/Chooser/Source/ChooserUncooked/Private/EvaluateChooserNode.cpp b/Engine/Plugins/Experimental/Chooser/Source/ChooserUncooked/Private/EvaluateChooserNode.cpp
--- a/Engine/Plugins/Experimental/Chooser/Source/ChooserUncooked/Private/EvaluateChooserNode.cpp
+++ b/Engine/Plugins/Experimental/Chooser/Source/ChooserUncooked/Private/EvaluateChooserNode.cpp
@@ -6,6 +6,7 @@
 #include "BlueprintActionDatabaseRegistrar.h"
 #include "BlueprintNodeSpawner.h"
 #include "ChooserFunctionLibrary.h"
+#include "ChooserPropertyAccess.h"
 #include "Containers/UnrealString.h"
 #include "EdGraph/EdGraph.h"
 #include "EdGraph/EdGraphSchema.h"
diff --git a/Engine/Plugins/Experimental/Chooser/Source/ProxyTable/Internal/LookupProxy.h b/Engine/Plugins/Experimental/Chooser/Source/ProxyTable/Internal/LookupProxy.h
--- a/Engine/Plugins/Experimental/Chooser/Source/ProxyTable/Internal/LookupProxy.h
+++ b/Engine/Plugins/Experimental/Chooser/Source/ProxyTable/Internal/LookupProxy.h
@@ -21,18 +21,11 @@ public:
 	FChooserPropertyBinding Binding;
 
 	virtual bool GetValue(FChooserEvaluationContext& Context, const UProxyTable*& OutResult) const override;
-
-#if WITH_EDITOR
-	static bool CanBind(const FProperty& Property)
-	{
-		return Property.GetCPPType() == "UProxyTable*";
-	}
-
-	void SetBinding(const TArray<FBindingChainElement>& InBindingChain)
+	
+	virtual void Compile(IHasContextClass* Owner, bool bForce) override
 	{
-		UE::Chooser::CopyPropertyChain(InBindingChain, Binding);
-	}
-#endif
+		Binding.Compile(Owner, bForce);
+	};
 };
 
 USTRUCT()
@@ -44,6 +37,8 @@ struct PROXYTABLE_API FLookupProxy : public FObjectChooserBase
 
 	FLookupProxy();
 	
+	virtual void Compile(IHasContextClass* HasContext, bool bForce) override;
+	
 	public:
 	
 	UPROPERTY(EditAnywhere, Category="Parameters")
diff --git a/Engine/Plugins/Experimental/Chooser/Source/ProxyTable/Private/LookupProxy.cpp b/Engine/Plugins/Experimental/Chooser/Source/ProxyTable/Private/LookupProxy.cpp
--- a/Engine/Plugins/Experimental/Chooser/Source/ProxyTable/Private/LookupProxy.cpp
+++ b/Engine/Plugins/Experimental/Chooser/Source/ProxyTable/Private/LookupProxy.cpp
@@ -71,17 +71,41 @@ FObjectChooserBase::EIteratorStatus FLookupProxyWithOverrideTable::ChooseMulti(F
 
 bool FProxyTableContextProperty::GetValue(FChooserEvaluationContext& Context, const UProxyTable*& OutResult) const
 {
-	const UStruct* StructType = nullptr;
-	const void* Container = nullptr;
-	
-	if (UE::Chooser::ResolvePropertyChain(Context, Binding,Container, StructType))
+	if (Binding.CompiledBinding)
+	{
+		UProxyTable** ProxyTableReference;
+		return Binding.GetValuePtr(Context, ProxyTableReference);
+		OutResult = *ProxyTableReference;
+	}
+	else
 	{
-		if (const FObjectProperty* Property = FindFProperty<FObjectProperty>(StructType, Binding.PropertyBindingChain.Last()))
+		// for temporary backwards compatibility: ProxyTableContextProperties on UProxyAsset are being phased out,
+		// but they don't get compiled so we need to keep this code-path temporarily
+		const UStruct* StructType = nullptr;
+		const void* Container = nullptr;
+	
+		if (UE::Chooser::ResolvePropertyChain(Context, Binding,Container, StructType))
 		{
-			OutResult = *Property->ContainerPtrToValuePtr<UProxyTable*>(Container);
-			return true;
+			if (const FObjectProperty* Property = FindFProperty<FObjectProperty>(StructType, Binding.PropertyBindingChain.Last()))
+			{
+				OutResult = *Property->ContainerPtrToValuePtr<UProxyTable*>(Container);
+				return true;
+			}
 		}
 	}
+	
 
 	return false;
 }
+
+void FLookupProxy::Compile(IHasContextClass* HasContext, bool bForce)
+{
+	if (Proxy)
+	{
+		if (FChooserParameterBase* ProxyTableParam = ProxyTable.GetMutablePtr<FChooserParameterBase>())
+		{
+			// todo: should also validate here that the ProxyAsset context is compatible with the passed in HasContext
+			ProxyTableParam->Compile(Proxy, bForce);
+		}
+	}
+}
diff --git a/Engine/Plugins/Experimental/Chooser/Source/ProxyTable/Private/ProxyTable.cpp b/Engine/Plugins/Experimental/Chooser/Source/ProxyTable/Private/ProxyTable.cpp
--- a/Engine/Plugins/Experimental/Chooser/Source/ProxyTable/Private/ProxyTable.cpp
+++ b/Engine/Plugins/Experimental/Chooser/Source/ProxyTable/Private/ProxyTable.cpp
@@ -167,6 +167,13 @@ static void OutputStructData(const FRuntimeProxyValue& EntryValueData, FChooserE
 {
 	for (const FProxyStructOutput& StructOutput : EntryValueData.OutputStructData)
 	{
+		// todo: should be updated to the following, but ProxyTable needs to trigger compiling of it's output struct bindings first
+		// void* TargetData;
+		// if (StructOutput.Binding.GetValuePtr(Context, TargetData))
+		// {
+		// 	StructOutput.Value.GetScriptStruct()->CopyScriptStruct(TargetData, StructOutput.Value.GetMemory());
+		// }	
+		
 		const void* Container = nullptr;
 		const UStruct* StructType;
 
diff --git a/Engine/Plugins/Experimental/Chooser/Source/ProxyTable/Public/ProxyAsset.h b/Engine/Plugins/Experimental/Chooser/Source/ProxyTable/Public/ProxyAsset.h
--- a/Engine/Plugins/Experimental/Chooser/Source/ProxyTable/Public/ProxyAsset.h
+++ b/Engine/Plugins/Experimental/Chooser/Source/ProxyTable/Public/ProxyAsset.h
@@ -3,10 +3,8 @@
 
 #include "CoreMinimal.h"
 #include "UObject/Object.h"
-#include "ChooserPropertyAccess.h"
-#include "IChooserParameterProxyTable.h"
+#include "IHasContext.h"
 #include "InstancedStruct.h"
-#include "InstancedStructContainer.h"
 #include "Misc/Guid.h"
 #include "IObjectChooser.h"
 #include "ProxyAsset.generated.h"
diff --git a/Engine/Plugins/Experimental/Chooser/Source/ProxyTable/Public/ProxyTable.h b/Engine/Plugins/Experimental/Chooser/Source/ProxyTable/Public/ProxyTable.h
--- a/Engine/Plugins/Experimental/Chooser/Source/ProxyTable/Public/ProxyTable.h
+++ b/Engine/Plugins/Experimental/Chooser/Source/ProxyTable/Public/ProxyTable.h
@@ -3,9 +3,8 @@
 
 #include "CoreMinimal.h"
 #include "UObject/Object.h"
-#include "IChooserParameterProxyTable.h"
+#include "ChooserPropertyAccess.h"
 #include "InstancedStruct.h"
-#include "InstancedStructContainer.h"
 #include "Misc/Guid.h"
 #include "ProxyAsset.h"
 #include "ProxyTable.generated.h"
diff --git a/Engine/Plugins/Experimental/Chooser/Source/ProxyTableEditor/Private/ProxyTableEditor.cpp b/Engine/Plugins/Experimental/Chooser/Source/ProxyTableEditor/Private/ProxyTableEditor.cpp
--- a/Engine/Plugins/Experimental/Chooser/Source/ProxyTableEditor/Private/ProxyTableEditor.cpp
+++ b/Engine/Plugins/Experimental/Chooser/Source/ProxyTableEditor/Private/ProxyTableEditor.cpp
@@ -1033,14 +1033,7 @@ TSharedRef<SWidget> CreateProxyTablePropertyWidget(bool bReadOnly, UObject* Tran
 
 	return SNew(UE::ChooserEditor::SPropertyAccessChainWidget).ContextClassOwner(HasContextClass).AllowFunctions(false).BindingColor("ClassPinTypeColor").TypeFilter("UProxyTable*")
 	.PropertyBindingValue(&ContextProperty->Binding)
-	.OnAddBinding_Lambda(
-		[ContextProperty, TransactionObject, ValueChanged](FName InPropertyName, const TArray<FBindingChainElement>& InBindingChain)
-		{
-			const FScopedTransaction Transaction(NSLOCTEXT("ContextPropertyWidget", "Change Property Binding", "Change Property Binding"));
-			TransactionObject->Modify(true);
-			ContextProperty->SetBinding(InBindingChain);
-			ValueChanged.ExecuteIfBound();
-		});
+	.OnValueChanged(ValueChanged);
 }
 
 TSharedRef<SWidget> CreateLookupProxyWidget(bool bReadOnly, UObject* TransactionObject, void* Value, UClass* ResultBaseClass, ChooserEditor::FChooserWidgetValueChanged ValueChanged)
diff --git a/Engine/Plugins/Experimental/Chooser/Source/ProxyTableEditor/Private/StructOutputDataCustomization.cpp b/Engine/Plugins/Experimental/Chooser/Source/ProxyTableEditor/Private/StructOutputDataCustomization.cpp
--- a/Engine/Plugins/Experimental/Chooser/Source/ProxyTableEditor/Private/StructOutputDataCustomization.cpp
+++ b/Engine/Plugins/Experimental/Chooser/Source/ProxyTableEditor/Private/StructOutputDataCustomization.cpp
@@ -75,7 +75,7 @@ void FStructOutputDataCustomization::CustomizeHeader(TSharedRef<IPropertyHandle>
 		{
 			void* data;
 			BindingHandle->GetValueData(data);
-			return reinterpret_cast<const FChooserPropertyBinding*>(data);
+			return reinterpret_cast<FChooserPropertyBinding*>(data);
 		})
 		.OnAddBinding_Lambda([HasContext, ValueHandle, BindingHandle](FName InPropertyName, const TArray<FBindingChainElement>& InBindingChain)
          	 	{
@@ -96,6 +96,7 @@ void FStructOutputDataCustomization::CustomizeHeader(TSharedRef<IPropertyHandle>
          	 				
          	 				OuterObjects[i]->Modify(true);
          	 				Chooser::CopyPropertyChain(InBindingChain, *PropertyValue);
+         	 				PropertyValue->Compile(HasContext);
          	
 							FChooserStructPropertyBinding* StructPropertyBinding = static_cast<FChooserStructPropertyBinding*>(PropertyValue);
 
