one more non-unity compile fix

diff --git a/Engine/Plugins/Experimental/Chooser/Source/ChooserEditor/Private/OutputFloatColumnEditor.cpp b/Engine/Plugins/Experimental/Chooser/Source/ChooserEditor/Private/OutputFloatColumnEditor.cpp
--- a/Engine/Plugins/Experimental/Chooser/Source/ChooserEditor/Private/OutputFloatColumnEditor.cpp
+++ b/Engine/Plugins/Experimental/Chooser/Source/ChooserEditor/Private/OutputFloatColumnEditor.cpp
@@ -8,6 +8,7 @@
 #include "Widgets/Images/SImage.h"
 #include "GraphEditorSettings.h"
 #include "SPropertyAccessChainWidget.h"
+#include "ScopedTransaction.h"
 
 #define LOCTEXT_NAMESPACE "OutputBoolColumnEditor"
 
