Some bug fixes for chooser property binding optimizations

diff --git a/Engine/Plugins/Experimental/Chooser/Source/Chooser/Private/ChooserPropertyAccess.cpp b/Engine/Plugins/Experimental/Chooser/Source/Chooser/Private/ChooserPropertyAccess.cpp
--- a/Engine/Plugins/Experimental/Chooser/Source/Chooser/Private/ChooserPropertyAccess.cpp
+++ b/Engine/Plugins/Experimental/Chooser/Source/Chooser/Private/ChooserPropertyAccess.cpp
@@ -99,9 +99,7 @@ void FChooserPropertyBinding::Compile(IHasContextClass* Owner, bool bForce)
 		CompiledBinding = nullptr;
 		return;
 	}
-	
-	CompiledBinding->TargetType = StructType;
-	
+		
 	FCompiledBindingCacheId Id;
 	Id.Type = StructType;
 	Id.BindingPath = GetUniqueId();
@@ -139,7 +137,8 @@ void FChooserPropertyBinding::Compile(IHasContextClass* Owner, bool bForce)
 
 	TSharedPtr<UE::Chooser::FCompiledBinding> NewCompiledBinding = MakeShared<UE::Chooser::FCompiledBinding>();
 	UE::Chooser::FCompiledBinding& OutCompiledBinding = *NewCompiledBinding.Get();
-	
+
+	OutCompiledBinding.TargetType = StructType;
 	OutCompiledBinding.CompiledChain.SetNum(0);
 	OutCompiledBinding.ContextIndex = ContextIndex;
 	
diff --git a/Engine/Plugins/Experimental/Chooser/Source/Chooser/Public/IChooserColumn.h b/Engine/Plugins/Experimental/Chooser/Source/Chooser/Public/IChooserColumn.h
--- a/Engine/Plugins/Experimental/Chooser/Source/Chooser/Public/IChooserColumn.h
+++ b/Engine/Plugins/Experimental/Chooser/Source/Chooser/Public/IChooserColumn.h
@@ -98,6 +98,7 @@ public:
 	}\
 	virtual UScriptStruct* GetInputBaseType() const override { return ParameterType::StaticStruct(); };\
 	virtual const UScriptStruct* GetInputType() const override { return InputValue.IsValid() ? InputValue.GetScriptStruct() : nullptr; };\
+	virtual FChooserParameterBase* GetInputValue() override { return InputValue.GetMutablePtr<FChooserParameterBase>(); };\
 	virtual void SetInputType(const UScriptStruct* Type) override { InputValue.InitializeAs(Type); };
 
 #else
