Fix for static analysis unchecked nullptr issue

diff --git a/Engine/Plugins/Experimental/Chooser/Source/Chooser/Private/ChooserPropertyAccess.cpp b/Engine/Plugins/Experimental/Chooser/Source/Chooser/Private/ChooserPropertyAccess.cpp
--- a/Engine/Plugins/Experimental/Chooser/Source/Chooser/Private/ChooserPropertyAccess.cpp
+++ b/Engine/Plugins/Experimental/Chooser/Source/Chooser/Private/ChooserPropertyAccess.cpp
@@ -291,6 +291,11 @@ namespace UE::Chooser
 			InputType = Context.Params[CompiledBinding.ContextIndex].GetScriptStruct();
 		}
 
+		if (Result == nullptr || InputType == nullptr)
+		{
+			return nullptr;
+		}
+
 		if (!InputType->IsChildOf(CompiledBinding.TargetType))
 		{
 			UE_LOG(LogChooser, Warning, TEXT("Property Binding compiled for type: {%s} is being evaluated on incompatible type: {%s}."), ToCStr(CompiledBinding.TargetType->GetName()), ToCStr(InputType->GetName()));
