Fix chooser compile dependencies not being updated in the case when the compiled property chains were all found in the cache

diff --git a/Engine/Plugins/Experimental/Chooser/Source/Chooser/Private/Chooser.cpp b/Engine/Plugins/Experimental/Chooser/Source/Chooser/Private/Chooser.cpp
--- a/Engine/Plugins/Experimental/Chooser/Source/Chooser/Private/Chooser.cpp
+++ b/Engine/Plugins/Experimental/Chooser/Source/Chooser/Private/Chooser.cpp
@@ -79,6 +79,7 @@ void UChooserTable::AddCompileDependency(const UStruct* InStructType)
 		if (UUserDefinedStruct* UserDefinedStruct = Cast<UUserDefinedStruct>(StructType))
 		{
 			UserDefinedStruct->ChangedEvent.AddUObject(this, &UChooserTable::OnDependentStructChanged);
+			CompileDependencies.Add(StructType);
 		}
 		else if (UClass* Class = Cast<UClass>(StructType))
 		{
diff --git a/Engine/Plugins/Experimental/Chooser/Source/Chooser/Private/ChooserPropertyAccess.cpp b/Engine/Plugins/Experimental/Chooser/Source/Chooser/Private/ChooserPropertyAccess.cpp
--- a/Engine/Plugins/Experimental/Chooser/Source/Chooser/Private/ChooserPropertyAccess.cpp
+++ b/Engine/Plugins/Experimental/Chooser/Source/Chooser/Private/ChooserPropertyAccess.cpp
@@ -129,6 +129,13 @@ void FChooserPropertyBinding::Compile(IHasContextClass* Owner, bool bForce)
 				if (bUseCachedBinding)
 				{
 					CompiledBinding = Binding->Pin();
+
+#if WITH_EDITOR
+					for (const UStruct* Dependency : CompiledBinding->Dependencies)
+					{
+						Owner->AddCompileDependency(Dependency);
+					}
+#endif
 					return;
 				}
 			}
@@ -148,6 +155,7 @@ void FChooserPropertyBinding::Compile(IHasContextClass* Owner, bool bForce)
 	{
 #if WITH_EDITOR
 		Owner->AddCompileDependency(StructType);
+		OutCompiledBinding.Dependencies.AddUnique(StructType);
 #endif
 
 		bool bFound = false;
@@ -192,6 +200,7 @@ void FChooserPropertyBinding::Compile(IHasContextClass* Owner, bool bForce)
 	
 	#if WITH_EDITOR
    	Owner->AddCompileDependency(StructType);
+	OutCompiledBinding.Dependencies.AddUnique(StructType);
     #endif
 
 	bool bFound = false;
diff --git a/Engine/Plugins/Experimental/Chooser/Source/Chooser/Public/ChooserPropertyAccess.h b/Engine/Plugins/Experimental/Chooser/Source/Chooser/Public/ChooserPropertyAccess.h
--- a/Engine/Plugins/Experimental/Chooser/Source/Chooser/Public/ChooserPropertyAccess.h
+++ b/Engine/Plugins/Experimental/Chooser/Source/Chooser/Public/ChooserPropertyAccess.h
@@ -59,6 +59,7 @@ namespace UE::Chooser
 		const UStruct* TargetType = nullptr;
 #if WITH_EDITORONLY_DATA
 		int SerialNumber = 0;
+		TArray<const UStruct*> Dependencies;
 #endif
 	};
 }
