Fix for choosers and proxy assets that return classes (such as for selecting an Anim Class)

diff --git a/Engine/Plugins/Experimental/Chooser/Source/ChooserUncooked/Private/EvaluateChooserNode.cpp b/Engine/Plugins/Experimental/Chooser/Source/ChooserUncooked/Private/EvaluateChooserNode.cpp
--- a/Engine/Plugins/Experimental/Chooser/Source/ChooserUncooked/Private/EvaluateChooserNode.cpp
+++ b/Engine/Plugins/Experimental/Chooser/Source/ChooserUncooked/Private/EvaluateChooserNode.cpp
@@ -19,6 +19,7 @@
 #include "K2Node_MakeStruct.h"
 #include "Kismet2/BlueprintEditorUtils.h"
 #include "Kismet2/CompilerResultsLog.h"
+#include "Kismet/KismetSystemLibrary.h"
 #include "KismetCompiler.h"
 #include "Misc/AssertionMacros.h"
 #include "Misc/CString.h"
@@ -763,11 +764,33 @@ void UK2Node_EvaluateChooser2::ExpandNode(class FKismetCompilerContext& Compiler
 
 			if (UEdGraphPin* ResultIsClassPin = CallFunction->FindPin(TEXT("bResultIsClass")))
 			{
+				// this ensures that the function does the right kind of type validation
 				CallFunction->GetSchema()->TrySetDefaultValue(*ResultIsClassPin, Chooser->ResultType == EObjectChooserResultType::ClassResult ? TEXT("true") : TEXT("false"));
+
+				if (Chooser->ResultType == EObjectChooserResultType::ClassResult)
+				{
+					// if it's a class we need to add a CastToClass function to cast it from an Object pointer to a Class poointer
+					UK2Node_CallFunction* CastToClass = CompilerContext.SpawnIntermediateNode<UK2Node_CallFunction>(this, SourceGraph);
+					CompilerContext.MessageLog.NotifyIntermediateObjectCreation(CastToClass, this);
+
+					CastToClass->SetFromFunction(UKismetSystemLibrary::StaticClass()->FindFunctionByName(GET_MEMBER_NAME_CHECKED(UKismetSystemLibrary, Conv_ObjectToClass)));
+					CastToClass->AllocateDefaultPins();
+
+					if (UEdGraphPin* ClassPin = CastToClass->FindPin(TEXT("Class")))
+					{
+						CastToClass->GetSchema()->TrySetDefaultObject(*ClassPin, Chooser->OutputObjectType);
+					}
+
+					if (UEdGraphPin* ObjectPin = CastToClass->FindPin(FName("Object")))
+					{
+						CallFunction->GetReturnValuePin()->MakeLinkTo(ObjectPin);
+					}
+
+					OutputPin = CastToClass->GetReturnValuePin();
+				}
 			}
 		}
-
-		
+				
 
 		CompilerContext.MovePinLinksToIntermediate(*ResultPin, *OutputPin);
 	}
diff --git a/Engine/Plugins/Experimental/Chooser/Source/ProxyTableUncooked/Private/EvaluateProxyNode.cpp b/Engine/Plugins/Experimental/Chooser/Source/ProxyTableUncooked/Private/EvaluateProxyNode.cpp
--- a/Engine/Plugins/Experimental/Chooser/Source/ProxyTableUncooked/Private/EvaluateProxyNode.cpp
+++ b/Engine/Plugins/Experimental/Chooser/Source/ProxyTableUncooked/Private/EvaluateProxyNode.cpp
@@ -19,6 +19,7 @@
 #include "K2Node_MakeStruct.h"
 #include "Kismet2/BlueprintEditorUtils.h"
 #include "Kismet2/CompilerResultsLog.h"
+#include "Kismet/KismetSystemLibrary.h"
 #include "KismetCompiler.h"
 #include "Misc/AssertionMacros.h"
 #include "Misc/CString.h"
@@ -765,7 +766,30 @@ void UK2Node_EvaluateProxy2::ExpandNode(class FKismetCompilerContext& CompilerCo
 
 			if (UEdGraphPin* ResultIsClassPin = CallFunction->FindPin(TEXT("bResultIsClass")))
 			{
+				// this ensures that the function does the right kind of type validation
 				CallFunction->GetSchema()->TrySetDefaultValue(*ResultIsClassPin, Proxy->ResultType == EObjectChooserResultType::ClassResult ? TEXT("true") : TEXT("false"));
+
+				if (Proxy->ResultType == EObjectChooserResultType::ClassResult)
+				{
+					// if it's a class we need to add a CastToClass function to cast it from an Object pointer to a Class poointer
+					UK2Node_CallFunction* CastToClass = CompilerContext.SpawnIntermediateNode<UK2Node_CallFunction>(this, SourceGraph);
+					CompilerContext.MessageLog.NotifyIntermediateObjectCreation(CastToClass, this);
+
+					CastToClass->SetFromFunction(UKismetSystemLibrary::StaticClass()->FindFunctionByName(GET_MEMBER_NAME_CHECKED(UKismetSystemLibrary, Conv_ObjectToClass)));
+					CastToClass->AllocateDefaultPins();
+
+					if (UEdGraphPin* ClassPin = CastToClass->FindPin(TEXT("Class")))
+					{
+						CastToClass->GetSchema()->TrySetDefaultObject(*ClassPin, Proxy->Type);
+					}
+
+					if (UEdGraphPin* ObjectPin = CastToClass->FindPin(FName("Object")))
+					{
+						CallFunction->GetReturnValuePin()->MakeLinkTo(ObjectPin);
+					}
+
+					OutputPin = CastToClass->GetReturnValuePin();
+				}
 			}
 		}
 
