Crash fix in ChooserTableEditor when deleting rows

diff --git a/Engine/Plugins/Experimental/Chooser/Source/ChooserEditor/Private/ChooserTableEditor.cpp b/Engine/Plugins/Experimental/Chooser/Source/ChooserEditor/Private/ChooserTableEditor.cpp
--- a/Engine/Plugins/Experimental/Chooser/Source/ChooserEditor/Private/ChooserTableEditor.cpp
+++ b/Engine/Plugins/Experimental/Chooser/Source/ChooserEditor/Private/ChooserTableEditor.cpp
@@ -1044,6 +1044,10 @@ TSharedRef<SDockTab> FChooserTableEditor::SpawnTableTab( const FSpawnTabArgs& Ar
 						{
 							RowsToDelete.Add(SelectedRow->Row);
 						}
+
+						SelectedRows.SetNum(0);
+						SelectRootProperties();
+
 						// sort indices in reverse
 						RowsToDelete.Sort([](int32 A, int32 B){ return A>B; });
 						for(uint32 RowIndex : RowsToDelete)
