Fix LookupProxy for proxy assets referenced from inside choosers

diff --git a/Engine/Plugins/Experimental/Chooser/Source/ProxyTable/Private/LookupProxy.cpp b/Engine/Plugins/Experimental/Chooser/Source/ProxyTable/Private/LookupProxy.cpp
--- a/Engine/Plugins/Experimental/Chooser/Source/ProxyTable/Private/LookupProxy.cpp
+++ b/Engine/Plugins/Experimental/Chooser/Source/ProxyTable/Private/LookupProxy.cpp
@@ -74,8 +74,15 @@ bool FProxyTableContextProperty::GetValue(FChooserEvaluationContext& Context, co
 	if (Binding.CompiledBinding)
 	{
 		UProxyTable** ProxyTableReference;
-		return Binding.GetValuePtr(Context, ProxyTableReference);
-		OutResult = *ProxyTableReference;
+		if (Binding.GetValuePtr(Context, ProxyTableReference))
+		{
+			OutResult = *ProxyTableReference;
+			return true;
+		}
+		else
+		{
+			return false;
+		}
 	}
 	else
 	{
@@ -105,7 +112,7 @@ void FLookupProxy::Compile(IHasContextClass* HasContext, bool bForce)
 		if (FChooserParameterBase* ProxyTableParam = ProxyTable.GetMutablePtr<FChooserParameterBase>())
 		{
 			// todo: should also validate here that the ProxyAsset context is compatible with the passed in HasContext
-			ProxyTableParam->Compile(Proxy, bForce);
+			ProxyTableParam->Compile(HasContext, bForce);
 		}
 	}
 }
