Fix native proxy table bindings not working

diff --git a/Engine/Plugins/Experimental/Chooser/Source/ChooserEditor/Private/SPropertyAccessChainWidget.cpp b/Engine/Plugins/Experimental/Chooser/Source/ChooserEditor/Private/SPropertyAccessChainWidget.cpp
--- a/Engine/Plugins/Experimental/Chooser/Source/ChooserEditor/Private/SPropertyAccessChainWidget.cpp
+++ b/Engine/Plugins/Experimental/Chooser/Source/ChooserEditor/Private/SPropertyAccessChainWidget.cpp
@@ -44,7 +44,7 @@ TSharedRef<SWidget> SPropertyAccessChainWidget::CreatePropertyAccessWidget()
 		if (TypeFilter == "object")
 		{
 			// special case for objects references of any type
-			return CastField<FObjectProperty>(Property) != nullptr;
+			return CastField<FObjectPropertyBase>(Property) != nullptr;
 		}
 		if (TypeFilter == "double")
 		{
@@ -65,8 +65,10 @@ TSharedRef<SWidget> SPropertyAccessChainWidget::CreatePropertyAccessWidget()
 			}
 			return false;
 		}
-		
-		return Property->GetCPPType() == TypeFilter;
+
+		const FString CPPType = Property->GetCPPType();
+
+		return CPPType == TypeFilter || CPPType == AlternateTypeFilter;
 	};
 
 	// allow struct bindings to bind context structs directly
@@ -348,6 +350,14 @@ void SPropertyAccessChainWidget::Construct( const FArguments& InArgs)
 	PropertyBindingValue = InArgs._PropertyBindingValue;
 	UpdateWidget();
 
+
+	if (TypeFilter[TypeFilter.Len() - 1] == '*')
+	{
+		FString Trimmed = TypeFilter.TrimChar('*');
+		AlternateTypeFilter = "TObjectPtr<" + Trimmed + ">";
+	}
+
+
 	if (ContextClassOwner)
 	{
 		ContextClassOwner->OnContextClassChanged.AddSP(this, &SPropertyAccessChainWidget::ContextClassChanged);
diff --git a/Engine/Plugins/Experimental/Chooser/Source/ChooserEditor/Public/SPropertyAccessChainWidget.h b/Engine/Plugins/Experimental/Chooser/Source/ChooserEditor/Public/SPropertyAccessChainWidget.h
--- a/Engine/Plugins/Experimental/Chooser/Source/ChooserEditor/Public/SPropertyAccessChainWidget.h
+++ b/Engine/Plugins/Experimental/Chooser/Source/ChooserEditor/Public/SPropertyAccessChainWidget.h
@@ -37,6 +37,7 @@ private:
 	void ContextClassChanged();
 		
 	FString TypeFilter;
+	FString AlternateTypeFilter;
 	FString BindingColor;
 	FDelegateHandle ContextClassChangedHandle();
 	IHasContextClass* ContextClassOwner = nullptr;
