Fixes for proxy table Struct Outputs

diff --git a/Engine/Plugins/Experimental/Chooser/Source/Chooser/Private/ChooserPropertyAccess.cpp b/Engine/Plugins/Experimental/Chooser/Source/Chooser/Private/ChooserPropertyAccess.cpp
--- a/Engine/Plugins/Experimental/Chooser/Source/Chooser/Private/ChooserPropertyAccess.cpp
+++ b/Engine/Plugins/Experimental/Chooser/Source/Chooser/Private/ChooserPropertyAccess.cpp
@@ -69,7 +69,7 @@ void FChooserPropertyBinding::Compile(IHasContextClass* Owner, bool bForce)
 
 	int CompiledBindingSerialNumber = 0;
 	
-	if (PropertyBindingChain.IsEmpty() || !ContextData.IsValidIndex(ContextIndex))
+	if (!ContextData.IsValidIndex(ContextIndex))
 	{
 #if WITH_EDITORONLY_DATA
 		CompileMessage = LOCTEXT("No Property Bound", "No Property Bound");
@@ -79,16 +79,13 @@ void FChooserPropertyBinding::Compile(IHasContextClass* Owner, bool bForce)
 	}
 
 	const UStruct* StructType = nullptr;
-	if (ContextData.IsValidIndex(ContextIndex))
+	if (const FContextObjectTypeClass* ClassContext = ContextData[ContextIndex].GetPtr<FContextObjectTypeClass>())
 	{
-		if (const FContextObjectTypeClass* ClassContext = ContextData[ContextIndex].GetPtr<FContextObjectTypeClass>())
-		{
-			StructType = ClassContext->Class;
-		}
-		else if (const FContextObjectTypeStruct* StructContext = ContextData[ContextIndex].GetPtr<FContextObjectTypeStruct>())
-		{
-			StructType = StructContext->Struct;
-		}
+		StructType = ClassContext->Class;
+	}
+	else if (const FContextObjectTypeStruct* StructContext = ContextData[ContextIndex].GetPtr<FContextObjectTypeStruct>())
+	{
+		StructType = StructContext->Struct;
 	}
 
 	if(StructType == nullptr)
@@ -197,59 +194,68 @@ void FChooserPropertyBinding::Compile(IHasContextClass* Owner, bool bForce)
 			return;
 		}
 	}
-	
-	#if WITH_EDITOR
-   	Owner->AddCompileDependency(StructType);
-	OutCompiledBinding.Dependencies.AddUnique(StructType);
-    #endif
 
 	bool bFound = false;
-	if (const FProperty* BaseProperty = FindFProperty<FProperty>(StructType, PropertyBindingChain.Last()))
+	if (PropertyBindingChain.IsEmpty())
 	{
+		// handle binding directly to a context struct
+		OutCompiledBinding.CompiledChain.Add(UE::Chooser::FCompiledBindingElement(0));
 		bFound = true;
-		
-		// last element should be the actual property - add it's offset to whatever was accumulated from struct offsets
-		CurrentOffset += BaseProperty->GetOffset_ForInternal();
-		OutCompiledBinding.CompiledChain.Add(UE::Chooser::FCompiledBindingElement(CurrentOffset));
-
-		if (BaseProperty->IsA<FFloatProperty>())
-		{
-			OutCompiledBinding.PropertyType = UE::Chooser::EPropertyNumericalType::FLOAT;
-		}
-		else if (BaseProperty->IsA<FDoubleProperty>())
-		{
-			OutCompiledBinding.PropertyType = UE::Chooser::EPropertyNumericalType::DOUBLE;
-		}
-		else if (BaseProperty->IsA<FIntProperty>())
-		{
-			OutCompiledBinding.PropertyType = UE::Chooser::EPropertyNumericalType::INT32;
-		}
 	}
 	else
 	{
-		// handle function calls 
-		if (const UClass* ClassType = Cast<const UClass>(StructType))
+#if WITH_EDITOR
+		Owner->AddCompileDependency(StructType);
+		OutCompiledBinding.Dependencies.AddUnique(StructType);
+#endif
+
+		if (const FProperty* BaseProperty = FindFProperty<FProperty>(StructType, PropertyBindingChain.Last()))
+		{
+			bFound = true;
+			
+			// last element should be the actual property - add it's offset to whatever was accumulated from struct offsets
+			CurrentOffset += BaseProperty->GetOffset_ForInternal();
+			OutCompiledBinding.CompiledChain.Add(UE::Chooser::FCompiledBindingElement(CurrentOffset));
+
+			if (BaseProperty->IsA<FFloatProperty>())
+			{
+				OutCompiledBinding.PropertyType = UE::Chooser::EPropertyNumericalType::FLOAT;
+			}
+			else if (BaseProperty->IsA<FDoubleProperty>())
+			{
+				OutCompiledBinding.PropertyType = UE::Chooser::EPropertyNumericalType::DOUBLE;
+			}
+			else if (BaseProperty->IsA<FIntProperty>())
+			{
+				OutCompiledBinding.PropertyType = UE::Chooser::EPropertyNumericalType::INT32;
+			}
+		}
+		else
 		{
-			if (UFunction* Function = ClassType->FindFunctionByName(PropertyBindingChain.Last()))
+			// handle function calls 
+			if (const UClass* ClassType = Cast<const UClass>(StructType))
 			{
-				bFound = true;
-				
-				const FProperty* ReturnProperty = Function->GetReturnProperty();
-				OutCompiledBinding.CompiledChain.Add(UE::Chooser::FCompiledBindingElement(Function));
-				if (ReturnProperty->IsA<FFloatProperty>())
+				if (UFunction* Function = ClassType->FindFunctionByName(PropertyBindingChain.Last()))
 				{
-					OutCompiledBinding.PropertyType = UE::Chooser::EPropertyNumericalType::FLOAT;
-				}
-				else if (ReturnProperty->IsA<FDoubleProperty>())
-				{
-					OutCompiledBinding.PropertyType = UE::Chooser::EPropertyNumericalType::DOUBLE;
-				}
-				else if (ReturnProperty->IsA<FIntProperty>())
-				{
-					OutCompiledBinding.PropertyType = UE::Chooser::EPropertyNumericalType::INT32;
+					bFound = true;
+					
+					const FProperty* ReturnProperty = Function->GetReturnProperty();
+					OutCompiledBinding.CompiledChain.Add(UE::Chooser::FCompiledBindingElement(Function));
+					if (ReturnProperty->IsA<FFloatProperty>())
+					{
+						OutCompiledBinding.PropertyType = UE::Chooser::EPropertyNumericalType::FLOAT;
+					}
+					else if (ReturnProperty->IsA<FDoubleProperty>())
+					{
+						OutCompiledBinding.PropertyType = UE::Chooser::EPropertyNumericalType::DOUBLE;
+					}
+					else if (ReturnProperty->IsA<FIntProperty>())
+					{
+						OutCompiledBinding.PropertyType = UE::Chooser::EPropertyNumericalType::INT32;
+					}
 				}
 			}
-		 }
+		}
 	}
 
 	if (bFound)
diff --git a/Engine/Plugins/Experimental/Chooser/Source/Chooser/Public/ChooserPropertyAccess.h b/Engine/Plugins/Experimental/Chooser/Source/Chooser/Public/ChooserPropertyAccess.h
--- a/Engine/Plugins/Experimental/Chooser/Source/Chooser/Public/ChooserPropertyAccess.h
+++ b/Engine/Plugins/Experimental/Chooser/Source/Chooser/Public/ChooserPropertyAccess.h
@@ -74,7 +74,7 @@ struct CHOOSER_API FChooserPropertyBinding
 	TArray<FName> PropertyBindingChain;
 	
 	UPROPERTY()
-	int ContextIndex = 0;
+	int ContextIndex = -1;
 
 #if WITH_EDITORONLY_DATA
 	UPROPERTY()
diff --git a/Engine/Plugins/Experimental/Chooser/Source/ChooserEditor/Private/SPropertyAccessChainWidget.cpp b/Engine/Plugins/Experimental/Chooser/Source/ChooserEditor/Private/SPropertyAccessChainWidget.cpp
--- a/Engine/Plugins/Experimental/Chooser/Source/ChooserEditor/Private/SPropertyAccessChainWidget.cpp
+++ b/Engine/Plugins/Experimental/Chooser/Source/ChooserEditor/Private/SPropertyAccessChainWidget.cpp
@@ -348,6 +348,7 @@ void SPropertyAccessChainWidget::Construct( const FArguments& InArgs)
 	bAllowFunctions = InArgs._AllowFunctions;
 	OnValueChanged = InArgs._OnValueChanged;
 	PropertyBindingValue = InArgs._PropertyBindingValue;
+	OnAddBinding = InArgs._OnAddBinding;
 	UpdateWidget();
 
 
