Make chooser table editor select rows when you click the drag handle

diff --git a/Engine/Plugins/Experimental/Chooser/Source/ChooserEditor/Private/ChooserTableEditor.cpp b/Engine/Plugins/Experimental/Chooser/Source/ChooserEditor/Private/ChooserTableEditor.cpp
--- a/Engine/Plugins/Experimental/Chooser/Source/ChooserEditor/Private/ChooserTableEditor.cpp
+++ b/Engine/Plugins/Experimental/Chooser/Source/ChooserEditor/Private/ChooserTableEditor.cpp
@@ -529,6 +529,7 @@ public:
 
 	FReply OnMouseButtonDown(const FGeometry& MyGeometry, const FPointerEvent& MouseEvent) override
 	{
+		ChooserEditor->SelectRow(RowIndex);
 		return FReply::Handled().DetectDrag(SharedThis(this), EKeys::LeftMouseButton);
 	};
 
@@ -759,6 +760,18 @@ void FChooserTableEditor::MoveRow(int SourceRowIndex, int TargetRowIndex)
 
 	UpdateTableRows();
 }
+	
+void FChooserTableEditor::SelectRow(int32 RowIndex) const
+{
+	if (TableRows.IsValidIndex(RowIndex))
+	{
+		if (!TableView->IsItemSelected(TableRows[RowIndex]))
+		{
+			TableView->ClearSelection();
+			TableView->SetItemSelection(TableRows[RowIndex], true, ESelectInfo::OnMouseClick);
+		}
+	}
+}
 
 void FChooserTableEditor::UpdateTableColumns()
 {
diff --git a/Engine/Plugins/Experimental/Chooser/Source/ChooserEditor/Private/ChooserTableEditor.h b/Engine/Plugins/Experimental/Chooser/Source/ChooserEditor/Private/ChooserTableEditor.h
--- a/Engine/Plugins/Experimental/Chooser/Source/ChooserEditor/Private/ChooserTableEditor.h
+++ b/Engine/Plugins/Experimental/Chooser/Source/ChooserEditor/Private/ChooserTableEditor.h
@@ -146,6 +146,8 @@ namespace UE::ChooserEditor
 		TSharedPtr<SHeaderRow> HeaderRow;
 		TSharedPtr<SListView<TSharedPtr<FChooserTableRow>>> TableView;
 	public:
+		void SelectRow(int32 RowIndex) const;
+		
 		TSharedPtr<SComboButton>& GetCreateRowComboButton() { return CreateRowComboButton; };
 
 		/** The name given to all instances of this type of editor */
