Change default value for ChooserPropertyAccess Context Index back to 0.
Fixes issues with values needing to be re-bound in editor

diff --git a/Engine/Plugins/Experimental/Chooser/Source/Chooser/Public/ChooserPropertyAccess.h b/Engine/Plugins/Experimental/Chooser/Source/Chooser/Public/ChooserPropertyAccess.h
--- a/Engine/Plugins/Experimental/Chooser/Source/Chooser/Public/ChooserPropertyAccess.h
+++ b/Engine/Plugins/Experimental/Chooser/Source/Chooser/Public/ChooserPropertyAccess.h
@@ -74,7 +74,7 @@ struct CHOOSER_API FChooserPropertyBinding
 	TArray<FName> PropertyBindingChain;
 	
 	UPROPERTY()
-	int ContextIndex = -1;
+	int ContextIndex = 0;
 	
 	UPROPERTY()
 	bool IsBoundToRoot = false;
