Fix for old proxy assets sometimes not working correctly with the new Evaluate Proxy node

diff --git a/Engine/Plugins/Experimental/Chooser/Source/ProxyTableUncooked/Private/EvaluateProxyNode.cpp b/Engine/Plugins/Experimental/Chooser/Source/ProxyTableUncooked/Private/EvaluateProxyNode.cpp
--- a/Engine/Plugins/Experimental/Chooser/Source/ProxyTableUncooked/Private/EvaluateProxyNode.cpp
+++ b/Engine/Plugins/Experimental/Chooser/Source/ProxyTableUncooked/Private/EvaluateProxyNode.cpp
@@ -394,6 +394,13 @@ void UK2Node_EvaluateProxy2::AllocateDefaultPins()
 
 	if (Proxy)
 	{
+		if (Proxy->ContextClass_DEPRECATED)
+		{
+			// post load fix up code doesn't run until it's too late when loading a blueprint that references a proxy
+			// need to manually trigger it here so we generate the right pins
+			Proxy->PostLoad();
+		}
+
 		for(FInstancedStruct& ContextDataEntry : Proxy->ContextData)
 		{
 			if (ContextDataEntry.IsValid())
