Fix for cooker crash with User Defined Structs in Instanced Struct data in Choosers/Proxy Tables

diff --git a/Engine/Plugins/Experimental/Chooser/Source/Chooser/Internal/Chooser.h b/Engine/Plugins/Experimental/Chooser/Source/Chooser/Internal/Chooser.h
--- a/Engine/Plugins/Experimental/Chooser/Source/Chooser/Internal/Chooser.h
+++ b/Engine/Plugins/Experimental/Chooser/Source/Chooser/Internal/Chooser.h
@@ -19,6 +19,7 @@ class CHOOSER_API UChooserTable : public UObject, public IHasContextClass
 	GENERATED_UCLASS_BODY()
 public:
 	UChooserTable() {}
+	virtual void BeginDestroy() override;
 
 	virtual void PostLoad() override;
 	virtual void Compile(bool bForce = false) override;
diff --git a/Engine/Plugins/Experimental/Chooser/Source/Chooser/Private/Chooser.cpp b/Engine/Plugins/Experimental/Chooser/Source/Chooser/Private/Chooser.cpp
--- a/Engine/Plugins/Experimental/Chooser/Source/Chooser/Private/Chooser.cpp
+++ b/Engine/Plugins/Experimental/Chooser/Source/Chooser/Private/Chooser.cpp
@@ -69,6 +69,13 @@ void UChooserTable::PostLoad()
 	Compile();
 }
 
+void UChooserTable::BeginDestroy()
+{
+	ColumnsStructs.Empty();
+	ResultsStructs.Empty();
+	Super::BeginDestroy();
+}
+
 #if WITH_EDITOR
 
 void UChooserTable::AddCompileDependency(const UStruct* InStructType)
diff --git a/Engine/Plugins/Experimental/Chooser/Source/ProxyTable/Private/ProxyTable.cpp b/Engine/Plugins/Experimental/Chooser/Source/ProxyTable/Private/ProxyTable.cpp
--- a/Engine/Plugins/Experimental/Chooser/Source/ProxyTable/Private/ProxyTable.cpp
+++ b/Engine/Plugins/Experimental/Chooser/Source/ProxyTable/Private/ProxyTable.cpp
@@ -163,6 +163,15 @@ void UProxyTable::PostTransacted(const FTransactionObjectEvent& TransactionEvent
 
 #endif
 
+void UProxyTable::BeginDestroy()
+{
+	RuntimeValues.Empty();
+#if WITH_EDITORONLY_DATA
+	Entries.Empty();
+#endif
+	Super::BeginDestroy();
+}
+
 static void OutputStructData(const FRuntimeProxyValue& EntryValueData, FChooserEvaluationContext& Context)
 {
 	for (const FProxyStructOutput& StructOutput : EntryValueData.OutputStructData)
diff --git a/Engine/Plugins/Experimental/Chooser/Source/ProxyTable/Public/ProxyTable.h b/Engine/Plugins/Experimental/Chooser/Source/ProxyTable/Public/ProxyTable.h
--- a/Engine/Plugins/Experimental/Chooser/Source/ProxyTable/Public/ProxyTable.h
+++ b/Engine/Plugins/Experimental/Chooser/Source/ProxyTable/Public/ProxyTable.h
@@ -69,6 +69,7 @@ class UProxyTable : public UObject
 	GENERATED_UCLASS_BODY()
 public:
 	UProxyTable() {}
+	virtual void BeginDestroy() override;
 
 	UPROPERTY()
 	TArray<FGuid> Keys;
