Fix output struct column

diff --git a/Engine/Plugins/Experimental/Chooser/Source/Chooser/Internal/OutputStructColumn.h b/Engine/Plugins/Experimental/Chooser/Source/Chooser/Internal/OutputStructColumn.h
--- a/Engine/Plugins/Experimental/Chooser/Source/Chooser/Internal/OutputStructColumn.h
+++ b/Engine/Plugins/Experimental/Chooser/Source/Chooser/Internal/OutputStructColumn.h
@@ -20,6 +20,11 @@ struct CHOOSER_API FStructContextProperty : public FChooserParameterStructBase
 
 	virtual bool SetValue(FChooserEvaluationContext& Context, const FInstancedStruct &Value) const override;
 
+	virtual void Compile(IHasContextClass* Owner, bool bForce) override
+ 	{
+ 		Binding.Compile(Owner, bForce);
+ 	};
+
 #if WITH_EDITOR
 	virtual void GetDisplayName(FText& OutName) const override
 	{
