Add support for bitfield bool properties in chooser property access

diff --git a/Engine/Plugins/Experimental/Chooser/Source/Chooser/Private/ChooserPropertyAccess.cpp b/Engine/Plugins/Experimental/Chooser/Source/Chooser/Private/ChooserPropertyAccess.cpp
--- a/Engine/Plugins/Experimental/Chooser/Source/Chooser/Private/ChooserPropertyAccess.cpp
+++ b/Engine/Plugins/Experimental/Chooser/Source/Chooser/Private/ChooserPropertyAccess.cpp
@@ -216,7 +216,12 @@ void FChooserPropertyBinding::Compile(IHasContextClass* Owner, bool bForce)
 			// last element should be the actual property - add it's offset to whatever was accumulated from struct offsets
 			CurrentOffset += BaseProperty->GetOffset_ForInternal();
 			OutCompiledBinding.CompiledChain.Add(UE::Chooser::FCompiledBindingElement(CurrentOffset));
-
+			
+			if (const FBoolProperty* BoolProperty = CastField<FBoolProperty>(BaseProperty))
+			{
+				OutCompiledBinding.CompiledChain.Last().Mask = BoolProperty->GetFieldMask();
+				OutCompiledBinding.PropertyType = UE::Chooser::EPropertyNumericalType::BOOL;
+			}
 			if (BaseProperty->IsA<FFloatProperty>())
 			{
 				OutCompiledBinding.PropertyType = UE::Chooser::EPropertyNumericalType::FLOAT;
diff --git a/Engine/Plugins/Experimental/Chooser/Source/Chooser/Public/ChooserPropertyAccess.h b/Engine/Plugins/Experimental/Chooser/Source/Chooser/Public/ChooserPropertyAccess.h
--- a/Engine/Plugins/Experimental/Chooser/Source/Chooser/Public/ChooserPropertyAccess.h
+++ b/Engine/Plugins/Experimental/Chooser/Source/Chooser/Public/ChooserPropertyAccess.h
@@ -25,6 +25,7 @@ namespace UE::Chooser
 		{
 			bIsFunction = false;
 			Offset = InOffset;
+			Mask = 0;
 		}
 		
 		explicit FCompiledBindingElement(UFunction* InFunction)
@@ -36,7 +37,11 @@ namespace UE::Chooser
 		bool bIsFunction;
 		union 
 		{
-			int Offset;
+			struct
+			{
+				int Offset;
+				uint8 Mask; // mask, for bitset bools
+			};
 			UFunction* Function;
 		};
 	};
@@ -45,6 +50,7 @@ namespace UE::Chooser
 	enum class EPropertyNumericalType
 	{
 		NONE,
+		BOOL,
 		INT32,
 		FLOAT,
 		DOUBLE,
@@ -241,6 +247,13 @@ bool FChooserPropertyBinding::GetValue(FChooserEvaluationContext& Context, T& Ou
 					OutResult = static_cast<T>(IntResult);
 					break;
 				}
+			case EPropertyNumericalType::BOOL:
+				{
+					const FCompiledBindingElement& Last = Binding.CompiledChain.Last();
+					uint8* ByteValue = Result + Last.Offset;
+					OutResult = !!(*ByteValue & Last.Mask);
+					break;
+				}
 			default:
 				OutResult = *reinterpret_cast<T*>(Result + Binding.CompiledChain.Last().Offset);
 				break;
@@ -349,6 +362,20 @@ bool FChooserPropertyBinding::SetValue(FChooserEvaluationContext& Context, const
 					*reinterpret_cast<int*>(Result + Binding.CompiledChain.Last().Offset) = static_cast<int>(InValue);
 					break;
 				}
+			case EPropertyNumericalType::BOOL:
+			{
+				const FCompiledBindingElement& Last = Binding.CompiledChain.Last();
+				uint8* ByteValue = Result + Last.Offset;
+				if (Last.Mask == 255) // regular bool
+				{
+					*reinterpret_cast<bool*>(Result + Binding.CompiledChain.Last().Offset) = static_cast<bool>(InValue);
+				}
+				else // bitset bool
+				{
+					*ByteValue = ((*ByteValue) & ~Last.Mask) | (InValue ? Last.Mask : 0);
+				}
+				break;
+			}
 			default:
 				*reinterpret_cast<T*>(Result + Binding.CompiledChain.Last().Offset) = InValue;
 				break;
diff --git a/Engine/Plugins/Experimental/Chooser/Source/ChooserEditor/Private/SPropertyAccessChainWidget.cpp b/Engine/Plugins/Experimental/Chooser/Source/ChooserEditor/Private/SPropertyAccessChainWidget.cpp
--- a/Engine/Plugins/Experimental/Chooser/Source/ChooserEditor/Private/SPropertyAccessChainWidget.cpp
+++ b/Engine/Plugins/Experimental/Chooser/Source/ChooserEditor/Private/SPropertyAccessChainWidget.cpp
@@ -65,6 +65,16 @@ TSharedRef<SWidget> SPropertyAccessChainWidget::CreatePropertyAccessWidget()
 			}
 			return false;
 		}
+		else if (TypeFilter == "bool")
+		{
+			// special case for bools, because CPPType == "bool" doesn't catch: uint8 bBool : 1
+	
+			if (const FBoolProperty* BoolProperty = CastField<const FBoolProperty>(Property))
+			{
+				return true;
+			}
+			return false;
+		}
 
 		const FString CPPType = Property->GetCPPType();
 
