ChooserTableEditor fix for binding root context struct in StructOutputColumn

diff --git a/Engine/Plugins/Experimental/Chooser/Source/ChooserEditor/Private/SPropertyAccessChainWidget.cpp b/Engine/Plugins/Experimental/Chooser/Source/ChooserEditor/Private/SPropertyAccessChainWidget.cpp
--- a/Engine/Plugins/Experimental/Chooser/Source/ChooserEditor/Private/SPropertyAccessChainWidget.cpp
+++ b/Engine/Plugins/Experimental/Chooser/Source/ChooserEditor/Private/SPropertyAccessChainWidget.cpp
@@ -207,11 +207,28 @@ TSharedRef<SWidget> SPropertyAccessChainWidget::CreatePropertyAccessWidget()
 					if (TypeFilter == "struct")
 					{
 						FChooserStructPropertyBinding* StructPropertyBinding = static_cast<FChooserStructPropertyBinding*>(ContextProperty);
-						
+
+						StructPropertyBinding->StructType = nullptr;
+
 						if (const FStructProperty* StructProperty = CastField<const FStructProperty>(Property))
 						{
 							StructPropertyBinding->StructType = StructProperty->Struct;
 						}
+						else if (InBindingChain.Num() == 1)
+						{
+							// direct binding to a context struct
+							if (ContextClassOwner)
+							{
+								TConstArrayView<FInstancedStruct> ContextData = ContextClassOwner->GetContextData();
+								if (ContextData.IsValidIndex(InBindingChain[0].ArrayIndex))
+								{
+									if (const FContextObjectTypeStruct* StructContext = ContextData[InBindingChain[0].ArrayIndex].GetPtr<FContextObjectTypeStruct>())
+									{
+										StructPropertyBinding->StructType = StructContext->Struct;
+									}
+								}
+							}
+						}
 					}
 
 					ContextProperty->Compile(ContextClassOwner);
