Chooser Table Editor Crash fix

diff --git a/Engine/Plugins/Experimental/Chooser/Source/ChooserEditor/Private/ChooserTableEditor.cpp b/Engine/Plugins/Experimental/Chooser/Source/ChooserEditor/Private/ChooserTableEditor.cpp
--- a/Engine/Plugins/Experimental/Chooser/Source/ChooserEditor/Private/ChooserTableEditor.cpp
+++ b/Engine/Plugins/Experimental/Chooser/Source/ChooserEditor/Private/ChooserTableEditor.cpp
@@ -308,12 +308,16 @@ void FChooserTableEditor::PostUndo(bool bSuccess)
 {
 	UpdateTableColumns();
 	UpdateTableRows();
+	ClearSelectedColumn();
+	SelectRootProperties();
 }
 
 void FChooserTableEditor::PostRedo(bool bSuccess)
 {
 	UpdateTableColumns();
 	UpdateTableRows();
+	ClearSelectedColumn();
+	SelectRootProperties();
 }
 
 
@@ -848,15 +852,11 @@ void FChooserTableEditor::UpdateTableColumns()
 				}
 
 
-				MenuBuilder.AddMenuEntry(LOCTEXT("Delete Column","Delete"),LOCTEXT("Delete Column ToolTip", "Remove this column and all its data from the table"),FSlateIcon(),
+				MenuBuilder.AddMenuEntry(LOCTEXT("Delete Column", "Delete"), LOCTEXT("Delete Column ToolTip", "Remove this column and all its data from the table"), FSlateIcon(),
 					FUIAction(
-						FExecuteAction::CreateLambda([this,Chooser,ColumnIndex, &Column]()
+						FExecuteAction::CreateLambda([this, Chooser, ColumnIndex, &Column]()
 						{
-							const FScopedTransaction Transaction(LOCTEXT("Delete Column Transaction", "Delete Column"));
-							Chooser->Modify(true);
-							Chooser->ColumnsStructs.RemoveAt(ColumnIndex);
-							ClearSelectedColumn();
-							UpdateTableColumns();
+							DeleteColumn(ColumnIndex);
 						})
 						));
 			
@@ -1237,9 +1237,14 @@ void FChooserTableEditor::ClearSelectedColumn()
 	
 void FChooserTableEditor::DeleteColumn(int Index)
 {
+	const FScopedTransaction Transaction(LOCTEXT("Delete Column Transaction", "Delete Column"));
+	ClearSelectedColumn();
+	SelectRootProperties();
 	UChooserTable* Chooser = GetChooser();
+
 	if (Index < Chooser->ColumnsStructs.Num())
 	{
+		Chooser->Modify(true);
 		Chooser->ColumnsStructs.RemoveAt(Index);
 		UpdateTableColumns();
 	}
