Call ConditionalPostLoad instead of PostLoad on dependent assets in EvaluateChooser and EvaluateProxy and fix MakeStruct warning

diff --git a/Engine/Plugins/Experimental/Chooser/Source/ChooserUncooked/Private/EvaluateChooserNode.cpp b/Engine/Plugins/Experimental/Chooser/Source/ChooserUncooked/Private/EvaluateChooserNode.cpp
--- a/Engine/Plugins/Experimental/Chooser/Source/ChooserUncooked/Private/EvaluateChooserNode.cpp
+++ b/Engine/Plugins/Experimental/Chooser/Source/ChooserUncooked/Private/EvaluateChooserNode.cpp
@@ -405,12 +405,8 @@ void UK2Node_EvaluateChooser2::AllocateDefaultPins()
 
 	if (Chooser)
 	{
-		if (Chooser->ContextObjectType_DEPRECATED)
-		{
-			// post load fix up code doesn't run until it's too late when loading a blueprint that references an chooser
-			// need to manually trigger it here so we generate the right pins
-			Chooser->PostLoad();
-		}
+		// ensure any data upgrades have been applied to Chooser before generating pins
+		Chooser->ConditionalPostLoad();
 		
 		for(FInstancedStruct& ContextDataEntry : Chooser->ContextData)
 		{
@@ -694,6 +690,7 @@ void UK2Node_EvaluateChooser2::ExpandNode(class FKismetCompilerContext& Compiler
 								// create a struct to hold the output (or for input structs that were not connected to anything)
 								UK2Node_MakeStruct* OutputStructNode = CompilerContext.SpawnIntermediateNode<UK2Node_MakeStruct>(this, SourceGraph);
 								CompilerContext.MessageLog.NotifyIntermediateObjectCreation(OutputStructNode, this);
+								OutputStructNode->PostPlacedNewNode();
 								OutputStructNode->StructType = static_cast<UScriptStruct*>(StructContext.Struct);
 								OutputStructNode->AllocateDefaultPins();
 								OutputStructNode->FindPin(StructContext.Struct->GetFName(), EGPD_Output)->MakeLinkTo(AddStructPin);
diff --git a/Engine/Plugins/Experimental/Chooser/Source/ProxyTableUncooked/Private/EvaluateProxyNode.cpp b/Engine/Plugins/Experimental/Chooser/Source/ProxyTableUncooked/Private/EvaluateProxyNode.cpp
--- a/Engine/Plugins/Experimental/Chooser/Source/ProxyTableUncooked/Private/EvaluateProxyNode.cpp
+++ b/Engine/Plugins/Experimental/Chooser/Source/ProxyTableUncooked/Private/EvaluateProxyNode.cpp
@@ -394,12 +394,8 @@ void UK2Node_EvaluateProxy2::AllocateDefaultPins()
 
 	if (Proxy)
 	{
-		if (Proxy->ContextClass_DEPRECATED)
-		{
-			// post load fix up code doesn't run until it's too late when loading a blueprint that references a proxy
-			// need to manually trigger it here so we generate the right pins
-			Proxy->PostLoad();
-		}
+		// ensure any data upgrades have been applied to Proxy before generating pins
+		Proxy->ConditionalPostLoad();
 
 		for(FInstancedStruct& ContextDataEntry : Proxy->ContextData)
 		{
@@ -688,6 +684,7 @@ void UK2Node_EvaluateProxy2::ExpandNode(class FKismetCompilerContext& CompilerCo
 								// create a struct to hold the output (or for input structs that were not connected to anything)
 								UK2Node_MakeStruct* OutputStructNode = CompilerContext.SpawnIntermediateNode<UK2Node_MakeStruct>(this, SourceGraph);
 								CompilerContext.MessageLog.NotifyIntermediateObjectCreation(OutputStructNode, this);
+								OutputStructNode->PostPlacedNewNode();
 								OutputStructNode->StructType = static_cast<UScriptStruct*>(StructContext.Struct);
 								OutputStructNode->AllocateDefaultPins();
 								OutputStructNode->FindPin(StructContext.Struct->GetFName(), EGPD_Output)->MakeLinkTo(AddStructPin);
