Fix crash in proxy table editor when selecting the row after the last row
#b Samuele.Rigamonti

diff --git a/Engine/Plugins/Experimental/Chooser/Source/ProxyTableEditor/Private/ProxyTableEditor.cpp b/Engine/Plugins/Experimental/Chooser/Source/ProxyTableEditor/Private/ProxyTableEditor.cpp
--- a/Engine/Plugins/Experimental/Chooser/Source/ProxyTableEditor/Private/ProxyTableEditor.cpp
+++ b/Engine/Plugins/Experimental/Chooser/Source/ProxyTableEditor/Private/ProxyTableEditor.cpp
@@ -71,7 +71,7 @@ void FProxyRowDetails::CustomizeDetails(IDetailLayoutBuilder& DetailBuilder)
 	UProxyRowDetails* Row = Cast<UProxyRowDetails>(Objects[0]);
 	UProxyTable* ProxyTable = Row->ProxyTable;
 	
-	if (ProxyTable->Entries.Num() > Row->Row)
+	if (ProxyTable->Entries.IsValidIndex(Row->Row))
 	{
 		IDetailCategoryBuilder& PropertiesCategory = DetailBuilder.EditCategory("Row Properties");
 
