ChooserPlayer anim graph node

diff --git a/Engine/Plugins/Experimental/Animation/BlendStack/Source/Editor/Public/AnimGraphNode_BlendStack.h b/Engine/Plugins/Experimental/Animation/BlendStack/Source/Editor/Public/AnimGraphNode_BlendStack.h
--- a/Engine/Plugins/Experimental/Animation/BlendStack/Source/Editor/Public/AnimGraphNode_BlendStack.h
+++ b/Engine/Plugins/Experimental/Animation/BlendStack/Source/Editor/Public/AnimGraphNode_BlendStack.h
@@ -13,6 +13,7 @@ class BLENDSTACKEDITOR_API UAnimGraphNode_BlendStack_Base : public UAnimGraphNod
 {
 	GENERATED_BODY()
 
+	public:
 	virtual void GetOutputLinkAttributes(FNodeAttributeArray& OutAttributes) const override;
 	virtual void OnProcessDuringCompilation(IAnimBlueprintCompilationContext& InCompilationContext, IAnimBlueprintGeneratedClassCompiledData& OutCompiledData) override;
 
@@ -28,6 +29,9 @@ class BLENDSTACKEDITOR_API UAnimGraphNode_BlendStack_Base : public UAnimGraphNod
 	virtual void PostPasteNode() override;
 	virtual void Serialize(FArchive& Ar) override;
 
+	virtual TArray<UEdGraph*> GetSubGraphs() const override;
+
+protected:
 	// Helper function for compilation
 	void ExpandGraphAndProcessNodes(
 		UEdGraph* SourceGraph, 
@@ -35,9 +39,6 @@ class BLENDSTACKEDITOR_API UAnimGraphNode_BlendStack_Base : public UAnimGraphNod
 		IAnimBlueprintCompilationContext& InCompilationContext, IAnimBlueprintGeneratedClassCompiledData& OutCompiledData,
 		UAnimGraphNode_Base*& OutRootNode, TArrayView<UAnimGraphNode_BlendStackInput*> OutInputNodes);
 
-	virtual TArray<UEdGraph*> GetSubGraphs() const override;
-
-protected:
 	virtual FAnimNode_BlendStack_Standalone* GetBlendStackNode() const PURE_VIRTUAL(UAnimGraphNode_BlendStack_Base::GetBlendStackNode, return nullptr;);
 	int32 GetMaxActiveBlends() const;
 
diff --git a/Engine/Plugins/Experimental/Animation/BlendStack/Source/Runtime/Private/AnimNode_BlendStack.cpp b/Engine/Plugins/Experimental/Animation/BlendStack/Source/Runtime/Private/AnimNode_BlendStack.cpp
--- a/Engine/Plugins/Experimental/Animation/BlendStack/Source/Runtime/Private/AnimNode_BlendStack.cpp
+++ b/Engine/Plugins/Experimental/Animation/BlendStack/Source/Runtime/Private/AnimNode_BlendStack.cpp
@@ -24,7 +24,8 @@ static constexpr FBoneIndexType RootBoneIndexType = 0;
 
 /////////////////////////////////////////////////////
 // FBlendStackAnimPlayer
-void FBlendStackAnimPlayer::Initialize(const FAnimationInitializeContext& Context, UAnimationAsset* AnimationAsset, float AccumulatedTime, bool bLoop, bool bMirrored, UMirrorDataTable* MirrorDataTable, float BlendTime, float RootBoneBlendTime, const UBlendProfile* BlendProfile, EAlphaBlendOption InBlendOption, FVector BlendParameters, float PlayRate, int32 InPoseLinkIdx)
+void FBlendStackAnimPlayer::Initialize(const FAnimationInitializeContext& Context, UAnimationAsset* AnimationAsset, float AccumulatedTime, bool bLoop, bool bMirrored, UMirrorDataTable* MirrorDataTable, float BlendTime, float RootBoneBlendTime, const UBlendProfile* BlendProfile, EAlphaBlendOption InBlendOption, FVector BlendParameters, float PlayRate, int32 InPoseLinkIdx,
+										FName GroupName, EAnimGroupRole::Type GroupRole, EAnimSyncMethod GroupMethod)
 {
 	if (bMirrored && !MirrorDataTable)
 	{
@@ -94,6 +95,9 @@ void FBlendStackAnimPlayer::Initialize(const FAnimationInitializeContext& Contex
 		SequencePlayerNode.SetSequence(SequenceBase);
 		SequencePlayerNode.SetLoopAnimation(bLoop);
 		SequencePlayerNode.SetPlayRate(PlayRate);
+		SequencePlayerNode.SetGroupMethod(GroupMethod);
+		SequencePlayerNode.SetGroupName(GroupName);
+		SequencePlayerNode.SetGroupRole(GroupRole);
 	}
 	else if (UBlendSpace* BlendSpace = Cast<UBlendSpace>(AnimationAsset))
 	{
@@ -106,6 +110,9 @@ void FBlendStackAnimPlayer::Initialize(const FAnimationInitializeContext& Contex
 		BlendSpacePlayerNode.SetLoop(bLoop);
 		BlendSpacePlayerNode.SetPlayRate(PlayRate);
 		BlendSpacePlayerNode.SetPosition(BlendParameters);
+		BlendSpacePlayerNode.SetGroupMethod(GroupMethod);
+		BlendSpacePlayerNode.SetGroupName(GroupName);
+		BlendSpacePlayerNode.SetGroupRole(GroupRole);
 	}
 	else if (AnimationAsset)
 	{
@@ -657,7 +664,8 @@ static void RequestInertialBlend(const FAnimationUpdateContext& Context, float B
 	}
 }
 
-void FAnimNode_BlendStack_Standalone::BlendTo(const FAnimationUpdateContext& Context, UAnimationAsset* AnimationAsset, float AccumulatedTime, bool bLoop, bool bMirrored, UMirrorDataTable* MirrorDataTable, float BlendTime, float RootBoneBlendTime, const UBlendProfile* BlendProfile, EAlphaBlendOption BlendOption, bool bUseInertialBlend, FVector BlendParameters, float PlayRate)
+void FAnimNode_BlendStack_Standalone::BlendTo(const FAnimationUpdateContext& Context, UAnimationAsset* AnimationAsset, float AccumulatedTime, bool bLoop, bool bMirrored, UMirrorDataTable* MirrorDataTable, float BlendTime, float RootBoneBlendTime, const UBlendProfile* BlendProfile, EAlphaBlendOption BlendOption, bool bUseInertialBlend, FVector BlendParameters, float PlayRate,
+											FName GroupName, EAnimGroupRole::Type GroupRole, EAnimSyncMethod GroupMethod)
 {
 	if (bUseInertialBlend)
 	{
@@ -669,7 +677,7 @@ void FAnimNode_BlendStack_Standalone::BlendTo(const FAnimationUpdateContext& Con
 	FBlendStackAnimPlayer& AnimPlayer = AnimPlayers.First();
 
 	FAnimationInitializeContext InitContext(Context.AnimInstanceProxy, Context.SharedContext);
-	AnimPlayer.Initialize(InitContext, AnimationAsset, AccumulatedTime, bLoop, bMirrored, MirrorDataTable, BlendTime, RootBoneBlendTime, BlendProfile, BlendOption, BlendParameters, PlayRate, GetNextPoseLinkIndex());
+	AnimPlayer.Initialize(InitContext, AnimationAsset, AccumulatedTime, bLoop, bMirrored, MirrorDataTable, BlendTime, RootBoneBlendTime, BlendProfile, BlendOption, BlendParameters, PlayRate, GetNextPoseLinkIndex(), GroupName, GroupRole, GroupMethod);
 	InitializeSample(InitContext, AnimPlayer);
 }
 
diff --git a/Engine/Plugins/Experimental/Animation/BlendStack/Source/Runtime/Public/BlendStack/AnimNode_BlendStack.h b/Engine/Plugins/Experimental/Animation/BlendStack/Source/Runtime/Public/BlendStack/AnimNode_BlendStack.h
--- a/Engine/Plugins/Experimental/Animation/BlendStack/Source/Runtime/Public/BlendStack/AnimNode_BlendStack.h
+++ b/Engine/Plugins/Experimental/Animation/BlendStack/Source/Runtime/Public/BlendStack/AnimNode_BlendStack.h
@@ -14,7 +14,9 @@ struct BLENDSTACK_API FBlendStackAnimPlayer
 {
 	GENERATED_BODY()
 	
-	void Initialize(const FAnimationInitializeContext& Context, UAnimationAsset* AnimationAsset, float AccumulatedTime, bool bLoop, bool bMirrored, UMirrorDataTable* MirrorDataTable, float BlendTime, float RootBoneBlendTime, const UBlendProfile* BlendProfile, EAlphaBlendOption InBlendOption, FVector BlendParameters, float PlayRate, int32 InPoseLinkIdx);
+	void Initialize(const FAnimationInitializeContext& Context, UAnimationAsset* AnimationAsset, float AccumulatedTime, bool bLoop, bool bMirrored, UMirrorDataTable* MirrorDataTable, float BlendTime, float RootBoneBlendTime, const UBlendProfile* BlendProfile, EAlphaBlendOption InBlendOption, FVector BlendParameters, float PlayRate, int32 InPoseLinkIdx,
+					FName GroupName, EAnimGroupRole::Type GroupRole, EAnimSyncMethod Method);
+	
 	void UpdatePlayRate(float PlayRate);
 	void Evaluate_AnyThread(FPoseContext& Output);
 	void Update_AnyThread(const FAnimationUpdateContext& Context);
@@ -102,7 +104,8 @@ struct BLENDSTACK_API FAnimNode_BlendStack_Standalone : public FAnimNode_AssetPl
 	virtual void GatherDebugData(FNodeDebugData& DebugData) override;
 	// End of FAnimNode_Base interface
 
-	void BlendTo(const FAnimationUpdateContext& Context, UAnimationAsset* AnimationAsset, float AccumulatedTime = 0.f, bool bLoop = false, bool bMirrored = false, UMirrorDataTable* MirrorDataTable = nullptr, float BlendTime = 0.2f, float RootBoneBlendTime = -1.f, const UBlendProfile* BlendProfile = nullptr, EAlphaBlendOption BlendOption = EAlphaBlendOption::Linear, bool bUseInertialBlend = false, FVector BlendParameters = FVector::Zero(), float PlayRate = 1.f);
+	void BlendTo(const FAnimationUpdateContext& Context, UAnimationAsset* AnimationAsset, float AccumulatedTime = 0.f, bool bLoop = false, bool bMirrored = false, UMirrorDataTable* MirrorDataTable = nullptr, float BlendTime = 0.2f, float RootBoneBlendTime = -1.f, const UBlendProfile* BlendProfile = nullptr, EAlphaBlendOption BlendOption = EAlphaBlendOption::Linear, bool bUseInertialBlend = false, FVector BlendParameters = FVector::Zero(), float PlayRate = 1.f,
+					FName GroupName = NAME_None, EAnimGroupRole::Type GroupRole = EAnimGroupRole::CanBeLeader, EAnimSyncMethod Method = EAnimSyncMethod::DoNotSync);
 	void UpdatePlayRate(float PlayRate);
 	void Reset();
 
diff --git a/Engine/Plugins/Experimental/Chooser/Chooser.uplugin b/Engine/Plugins/Experimental/Chooser/Chooser.uplugin
--- a/Engine/Plugins/Experimental/Chooser/Chooser.uplugin
+++ b/Engine/Plugins/Experimental/Chooser/Chooser.uplugin
@@ -23,6 +23,10 @@
 		{
 			"Name": "StructUtils",
 			"Enabled": true
+		},
+		{
+			"Name": "BlendStack",
+			"Enabled": true
 		}
 	],
 	"Modules" :
diff --git a/Engine/Plugins/Experimental/Chooser/Source/Chooser/Chooser.Build.cs b/Engine/Plugins/Experimental/Chooser/Source/Chooser/Chooser.Build.cs
--- a/Engine/Plugins/Experimental/Chooser/Source/Chooser/Chooser.Build.cs
+++ b/Engine/Plugins/Experimental/Chooser/Source/Chooser/Chooser.Build.cs
@@ -21,6 +21,7 @@ namespace UnrealBuildTool.Rules
 					"CoreUObject",
 					"Engine",
 					"GameplayTags",
+					"BlendStack"
 					// ... add private dependencies that you statically link with here ...
 				}
 			);
diff --git a/Engine/Plugins/Experimental/Chooser/Source/Chooser/Internal/AnimNode_ChooserPlayer.h b/Engine/Plugins/Experimental/Chooser/Source/Chooser/Internal/AnimNode_ChooserPlayer.h
new file mode 100644
--- /dev/null
+++ b/Engine/Plugins/Experimental/Chooser/Source/Chooser/Internal/AnimNode_ChooserPlayer.h
@@ -0,0 +1,215 @@
+// Copyright Epic Games, Inc. All Rights Reserved.
+
+#pragma once
+
+#include "AlphaBlend.h"
+#include "Animation/AnimNode_RelevantAssetPlayerBase.h"
+#include "Animation/AnimationAsset.h"
+#include "CoreMinimal.h"
+#include "InstancedStruct.h"
+#include "IObjectChooser.h"
+#include "BlendStack/AnimNode_BlendStack.h"
+#include "Styling/SlateBrush.h"
+
+#include "AnimNode_ChooserPlayer.generated.h"
+
+UENUM()
+enum class EChooserEvaluationFrequency
+{
+	OnInitialUpdate,
+	OnBecomeRelevant,
+	OnLoop,
+	OnUpdate
+};
+
+/** The random player node holds a list of sequences and parameter ranges which will be played continuously
+  * In a random order. If shuffle mode is enabled then each entry will be played once before repeating any
+  */
+USTRUCT(BlueprintInternalUseOnly)
+struct FChooserPlayerSequenceEntry
+{
+	GENERATED_BODY()
+
+	FChooserPlayerSequenceEntry()
+	    : Sequence(nullptr)
+	    , ChanceToPlay(1.0f)
+	    , MinLoopCount(0)
+	    , MaxLoopCount(0)
+	    , MinPlayRate(1.0f)
+	    , MaxPlayRate(1.0f)
+	{
+	}
+
+	/** Sequence to play when this entry is picked */
+	UPROPERTY(EditAnywhere, BlueprintReadWrite, Category = "Settings", meta = (DisallowedClasses = "/Script/Engine.AnimMontage"))
+	TObjectPtr<UAnimSequenceBase> Sequence;
+
+	/** When not in shuffle mode, this is the chance this entry will play (normalized against all other sample chances) */
+	UPROPERTY(EditAnywhere, BlueprintReadWrite, Category = "Settings", meta = (UIMin = "0", ClampMin = "0"))
+	float ChanceToPlay;
+
+	/** Minimum number of times this entry will loop before ending */
+	UPROPERTY(EditAnywhere, BlueprintReadWrite, Category = "Settings", meta = (UIMin = "0", ClampMin = "0"))
+	int32 MinLoopCount;
+
+	/** Maximum number of times this entry will loop before ending */
+	UPROPERTY(EditAnywhere, BlueprintReadWrite, Category = "Settings", meta = (UIMin = "0", ClampMin = "0"))
+	int32 MaxLoopCount;
+
+	/** Minimum playrate for this entry */
+	UPROPERTY(EditAnywhere, BlueprintReadWrite, Category = "Settings", meta = (UIMin = "0", ClampMin = "0"))
+	float MinPlayRate;
+
+	/** Maximum playrate for this entry */
+	UPROPERTY(EditAnywhere, BlueprintReadWrite, Category = "Settings", meta = (UIMin = "0", ClampMin = "0"))
+	float MaxPlayRate;
+
+	/** Blending properties used when this entry is blending in ontop of another entry */
+	UPROPERTY(EditAnywhere, Category = "Settings")
+	FAlphaBlend BlendIn;
+};
+
+USTRUCT(BlueprintType)
+struct FAnimCurveOverride
+{
+	GENERATED_BODY()
+	// Name of curve to override
+	UPROPERTY(EditAnywhere, BlueprintReadWrite, Category = "CurveValue")
+	FName CurveName;
+	// Value to set to the curve
+	UPROPERTY(EditAnywhere, BlueprintReadWrite, Category = "CurveValue")
+	float CurveValue;
+};
+
+USTRUCT(BlueprintType)
+struct FAnimCurveOverrideList
+{
+	GENERATED_BODY()
+	UPROPERTY(EditAnywhere, BlueprintReadWrite, Category = "Values")
+	TArray<FAnimCurveOverride> Values;
+};
+
+USTRUCT(BlueprintType)
+struct FChooserPlayerSettings
+{
+	GENERATED_BODY()
+
+	// Set this value to mirror animations - the MirrorDataTable must also be set on the AnimNode
+	UPROPERTY(EditAnywhere, BlueprintReadWrite, Category = "Settings")
+	bool bMirror = false;
+	
+	// Start offset when starting the Animation Asset
+	UPROPERTY(EditAnywhere, BlueprintReadWrite, Category = "Settings")
+    float StartTime = 0;
+
+	// Loop the animation asset, even if the asset is not set as looping
+	UPROPERTY(EditAnywhere, BlueprintReadWrite, Category = "Settings")
+    bool bForceLooping = false;
+
+	// playback rate modifier
+	UPROPERTY(EditAnywhere, BlueprintReadWrite ,Category = "Settings")
+    float PlaybackRate = 1.0;
+
+	// List of curve values to set 
+	UPROPERTY(EditAnywhere, BlueprintReadWrite, Category = "Settings")
+	FAnimCurveOverrideList CurveOverrides;
+	
+	UPROPERTY(EditAnywhere, BlueprintReadWrite, Category = "Blending")
+	float BlendTime = 0.2;
+	
+	// Set Blend Profiles (editable in the skeleton) to determine how the blending is distributed among your character's bones. It could be used to differentiate between upper body and lower body to blend timing.
+	UPROPERTY(EditAnywhere, BlueprintReadWrite, Category = Blending, meta = (UseAsBlendProfile = true))
+	TObjectPtr<UBlendProfile> BlendProfile;
+ 
+	// How the blend is applied over time to the bones. Common selections are linear, ease in, ease out, and ease in and out.
+	UPROPERTY(EditAnywhere, BlueprintReadWrite, Category = Blending)
+	EAlphaBlendOption BlendOption = EAlphaBlendOption::Linear;
+
+	UPROPERTY(EditAnywhere, BlueprintReadWrite, Category = Blending)
+	bool bUseInertialBlend = false;
+};
+
+USTRUCT(BlueprintInternalUseOnly)
+struct FAnimNode_ChooserPlayer : public FAnimNode_BlendStack_Standalone
+{
+	GENERATED_BODY()
+
+	CHOOSER_API FAnimNode_ChooserPlayer();
+
+public:
+	// How often the chooser should be evaluated
+	UPROPERTY(EditAnywhere, Category = "Chooser")
+	EChooserEvaluationFrequency EvaluationFrequency = EChooserEvaluationFrequency::OnBecomeRelevant;
+	
+	// Type of chooser logic to use: Use "Evaluate Chooser" for chooser table evaluation, and "Lookup Proxy" for proxy table lookups
+	UPROPERTY(EditAnywhere, Meta = (ExcludeBaseStruct, BaseStruct = "/Script/Chooser.ObjectChooserBase"), Category = "Chooser")
+	FInstancedStruct Chooser;
+
+	// if set and bMirrored MirrorDataTable will be used for mirroring the aniamtion
+	UPROPERTY(EditAnywhere, Category = "Settings")
+	TObjectPtr<UMirrorDataTable> MirrorDataTable;
+
+	// requested blend space blend X parameter (if AnimationAsset is a blend space)
+	UPROPERTY(EditAnywhere, BlueprintReadWrite, Category = Settings, meta = (PinHiddenByDefault))
+	float BlendSpaceX = 0;
+	
+	// requested blend space blend Y parameter (if AnimationAsset is a blend space)
+	UPROPERTY(EditAnywhere, BlueprintReadWrite, Category = Settings, meta = (PinHiddenByDefault))
+	float BlendSpaceY = 0;
+
+	// Should BlendSpaceParameters be applied to all blend spaces, including ones that are in the process of blending out.
+	bool bUpdateAllActiveBlendSpaces = true;
+	
+	// Settings when starting an animation - these can be overridden per animation asset by the chooser itself
+	UPROPERTY(EditAnywhere, Category = Settings)
+	FChooserPlayerSettings DefaultSettings;
+	
+	UPROPERTY()
+	TArray<FInstancedStruct> ChooserContextDefinition;
+
+	// FAnimNode_Base interface
+	virtual void Initialize_AnyThread(const FAnimationInitializeContext& Context) override;
+	virtual void Evaluate_AnyThread(FPoseContext& Output) override;
+	// End of FAnimNode_Base interface
+	
+	// FAnimNode_AssetPlayerBase interface
+	virtual void UpdateAssetPlayer(const FAnimationUpdateContext& Context) override;
+	virtual FName GetGroupName() const override;
+	virtual EAnimGroupRole::Type GetGroupRole() const override;
+	virtual EAnimSyncMethod GetGroupMethod() const override;
+	virtual bool GetIgnoreForRelevancyTest() const override;
+	virtual bool IsLooping() const override;
+	virtual bool SetGroupName(FName InGroupName) override;
+	virtual bool SetGroupRole(EAnimGroupRole::Type InRole) override;
+	virtual bool SetGroupMethod(EAnimSyncMethod InMethod) override;
+	virtual bool SetIgnoreForRelevancyTest(bool bInIgnoreForRelevancyTest) override;
+	// End of FAnimNode_AssetPlayerBase interface
+	
+private:
+	UAnimationAsset* ChooseAsset(const FAnimationUpdateContext& Context);
+	FChooserEvaluationContext ChooserContext;
+	UAnimationAsset* CurrentAsset = nullptr;
+	uint32 CurveOverridesIndex = 0;
+	TBaseBlendedCurve<FDefaultAllocator, UE::Anim::FCurveElement> OverrideCurves[2]; 
+	
+	// Update Counter for detecting being relevant
+	FGraphTraversalCounter UpdateCounter;
+	
+#if WITH_EDITORONLY_DATA
+	// The group name (NAME_None if it is not part of any group)
+	UPROPERTY(EditAnywhere, Category = Sync, meta = (FoldProperty))
+	FName GroupName = NAME_None;
+
+	// The role this player can assume within the group (ignored if GroupIndex is INDEX_NONE)
+	UPROPERTY(EditAnywhere, Category = Sync, meta = (FoldProperty))
+	TEnumAsByte<EAnimGroupRole::Type> GroupRole = EAnimGroupRole::CanBeLeader;
+
+	// How synchronization is determined
+	UPROPERTY(EditAnywhere, Category = Sync, meta = (FoldProperty))
+	EAnimSyncMethod Method = EAnimSyncMethod::DoNotSync;
+
+	// If true, "Relevant anim" nodes that look for the highest weighted animation in a state will ignore this node
+	UPROPERTY(EditAnywhere, Category = Relevancy, meta = (FoldProperty, PinHiddenByDefault))
+	bool bIgnoreForRelevancyTest = false;
+#endif // WITH_EDITORONLY_DATA
+};
diff --git a/Engine/Plugins/Experimental/Chooser/Source/Chooser/Private/AnimNode_ChooserPlayer.cpp b/Engine/Plugins/Experimental/Chooser/Source/Chooser/Private/AnimNode_ChooserPlayer.cpp
new file mode 100644
--- /dev/null
+++ b/Engine/Plugins/Experimental/Chooser/Source/Chooser/Private/AnimNode_ChooserPlayer.cpp
@@ -0,0 +1,271 @@
+// Copyright Epic Games, Inc. All Rights Reserved.
+
+#include "AnimNode_ChooserPlayer.h"
+
+#include "IObjectChooser.h"
+#include "../../../../Animation/BlendStack/Source/Runtime/Public/BlendStack/AnimNode_BlendStack.h"
+#include "Animation/AnimInstanceProxy.h"
+#include "Animation/AnimSequence.h"
+#include "Animation/AnimTrace.h"
+#include "Animation/AnimStats.h"
+#include "Animation/AnimSyncScope.h"
+#include "Animation/BlendSpace.h"
+
+#include UE_INLINE_GENERATED_CPP_BY_NAME(AnimNode_ChooserPlayer)
+
+FAnimNode_ChooserPlayer::FAnimNode_ChooserPlayer()
+{
+}
+
+UAnimationAsset* FAnimNode_ChooserPlayer::ChooseAsset(const FAnimationUpdateContext& Context)
+{
+	if (Chooser.IsValid())
+	{
+		// reset settings to default
+		FChooserPlayerSettings& Settings = ChooserContext.Params[1].GetMutable<FChooserPlayerSettings>();
+		Settings = DefaultSettings;
+		UAnimationAsset* Result = Cast<UAnimationAsset>(Chooser.Get<FObjectChooserBase>().ChooseObject(ChooserContext));
+	
+		return Result;
+	}
+	return nullptr;
+}
+
+void FAnimNode_ChooserPlayer::Initialize_AnyThread(const FAnimationInitializeContext& Context)
+{
+	DECLARE_SCOPE_HIERARCHICAL_COUNTER_ANIMNODE(Initialize_AnyThread)
+	FAnimNode_BlendStack_Standalone::Initialize_AnyThread(Context);
+
+	FInstancedStruct This;
+	This.InitializeAs(FChooserEvaluationInputObject::StaticStruct());
+	This.GetMutable<FChooserEvaluationInputObject>().Object = Context.AnimInstanceProxy->GetAnimInstanceObject();
+	ChooserContext.Params.Add(This);
+	ChooserContext.Params.Add(FInstancedStruct(FChooserPlayerSettings::StaticStruct()));
+}
+
+
+void FAnimNode_ChooserPlayer::UpdateAssetPlayer(const FAnimationUpdateContext& Context)
+{
+	DECLARE_SCOPE_HIERARCHICAL_COUNTER_ANIMNODE(UpdateAassetPlayer)
+	GetEvaluateGraphExposedInputs().Execute(Context);
+
+	if (!Chooser.IsValid())
+	{
+		// We don't have any entries, play data will be invalid - early out
+		return;
+	}
+	
+	const bool bJustBecameRelevant = !UpdateCounter.WasSynchronizedCounter(Context.AnimInstanceProxy->GetUpdateCounter());
+	UpdateCounter.SynchronizeWith(Context.AnimInstanceProxy->GetUpdateCounter());
+
+	UAnimationAsset* NewAsset = CurrentAsset;
+
+	if (EvaluationFrequency == EChooserEvaluationFrequency::OnUpdate || CurrentAsset == nullptr)
+	{
+		NewAsset = ChooseAsset(Context);
+	}
+	else if (EvaluationFrequency == EChooserEvaluationFrequency::OnLoop)
+	{
+		if (!AnimPlayers.IsEmpty())
+		{
+			FBlendStackAnimPlayer& Player = AnimPlayers.First();
+			if (Player.GetPlayRate() * Context.GetDeltaTime() + Player.GetAccumulatedTime() > CurrentAsset->GetPlayLength())
+			{
+				NewAsset = ChooseAsset(Context);
+			}
+		}
+	}
+	else if (EvaluationFrequency == EChooserEvaluationFrequency::OnBecomeRelevant && bJustBecameRelevant)
+	{
+		NewAsset = ChooseAsset(Context);
+	}
+
+	if (bJustBecameRelevant || NewAsset != CurrentAsset)
+	{
+		const FChooserPlayerSettings& Settings = ChooserContext.Params[1].Get<FChooserPlayerSettings>();
+			
+		CurveOverridesIndex = (CurveOverridesIndex + 1) % 2;
+		OverrideCurves[CurveOverridesIndex].Empty();
+		if (!Settings.CurveOverrides.Values.IsEmpty())
+		{
+			OverrideCurves[CurveOverridesIndex].Reserve(Settings.CurveOverrides.Values.Num());
+			for (const FAnimCurveOverride& CurveOverride : Settings.CurveOverrides.Values)
+			{
+				OverrideCurves[CurveOverridesIndex].Add(CurveOverride.CurveName, CurveOverride.CurveValue);
+			}
+		}
+
+		float BlendTime = bJustBecameRelevant ? 0 : Settings.BlendTime;
+
+		bool bLoop = Settings.bForceLooping;
+		if (!bLoop)
+		{
+			if (const UAnimSequenceBase* Sequence = Cast<UAnimSequenceBase>(NewAsset))
+			{
+				bLoop = Sequence->bLoop;				
+			}
+			else if (const UBlendSpace* BlendSpace = Cast<UBlendSpace>(NewAsset))
+			{
+				bLoop = BlendSpace->bLoop;
+			}
+		}
+		
+		FAnimNode_BlendStack_Standalone::BlendTo(Context, NewAsset, Settings.StartTime,
+			bLoop, Settings.bMirror, MirrorDataTable,
+			Settings.BlendTime, Settings.BlendTime, Settings.BlendProfile, Settings.BlendOption, Settings.bUseInertialBlend, FVector::Zero(), Settings.PlaybackRate,
+			GetGroupName(), GetGroupRole(), GetGroupMethod());
+		
+		CurrentAsset = NewAsset;
+	}
+
+	// Update blend space parameters
+	if (bUpdateAllActiveBlendSpaces)
+	{
+		// apply blend space parameters to all blendspaces that are playing, including ones that are blending out
+		for(FBlendStackAnimPlayer& Player : AnimPlayers)
+		{
+			if (Cast<UBlendSpace>(Player.GetAnimationAsset()))
+			{
+				Player.BlendSpacePlayerNode.SetPosition(FVector(BlendSpaceX,BlendSpaceY,0));
+			}
+		}
+	}
+	else
+	{
+		// apply blend space parameters only to the blendspace that is playing/blending in
+		if (!AnimPlayers.IsEmpty())
+		{
+			FBlendStackAnimPlayer& Player = AnimPlayers.First();
+			if (Cast<UBlendSpace>(Player.GetAnimationAsset()))
+			{
+				Player.BlendSpacePlayerNode.SetPosition(FVector(BlendSpaceX,BlendSpaceY,0));
+			}
+		}
+	}
+
+	FAnimNode_BlendStack_Standalone::UpdateAssetPlayer(Context);
+}
+
+void FAnimNode_ChooserPlayer::Evaluate_AnyThread(FPoseContext& Output)
+{
+	DECLARE_SCOPE_HIERARCHICAL_COUNTER_ANIMNODE(Evaluate_AnyThread);
+	ANIM_MT_SCOPE_CYCLE_COUNTER_VERBOSE(MotionMatching, !IsInGameThread());
+
+	FAnimNode_BlendStack_Standalone::Evaluate_AnyThread(Output);
+
+	const FChooserPlayerSettings& Settings = ChooserContext.Params[1].Get<FChooserPlayerSettings>();
+
+
+	
+	if (AnimPlayers.Num() == 1)
+	{
+		UE::Anim::FNamedValueArrayUtils::Union(Output.Curve, OverrideCurves[CurveOverridesIndex]);	
+	}
+	if (AnimPlayers.Num() > 1)
+	{
+		const float Weight = FAlphaBlend::AlphaToBlendOption(AnimPlayers[0].GetBlendInPercentage(), AnimPlayers[0].GetBlendOption());
+		const uint32 PrevCurveOverridesIndex = (CurveOverridesIndex + 1) % 2;
+		UE::Anim::FNamedValueArrayUtils::Union(Output.Curve, OverrideCurves[PrevCurveOverridesIndex]);
+		Output.Curve.LerpTo(OverrideCurves[CurveOverridesIndex], Weight);
+	}
+
+	
+}
+
+FName FAnimNode_ChooserPlayer::GetGroupName() const
+{
+	return GET_ANIM_NODE_DATA(FName, GroupName);
+}
+
+EAnimGroupRole::Type FAnimNode_ChooserPlayer::GetGroupRole() const
+{
+	return GET_ANIM_NODE_DATA(TEnumAsByte<EAnimGroupRole::Type>, GroupRole);
+}
+
+EAnimSyncMethod FAnimNode_ChooserPlayer::GetGroupMethod() const
+{
+	return GET_ANIM_NODE_DATA(EAnimSyncMethod, Method);
+}
+
+bool FAnimNode_ChooserPlayer::GetIgnoreForRelevancyTest() const
+{
+	return GET_ANIM_NODE_DATA(bool, bIgnoreForRelevancyTest);
+}
+
+bool FAnimNode_ChooserPlayer::SetGroupName(FName InGroupName)
+{
+#if WITH_EDITORONLY_DATA
+	GroupName = InGroupName;
+#endif
+
+	if(FName* GroupNamePtr = GET_INSTANCE_ANIM_NODE_DATA_PTR(FName, GroupName))
+	{
+		*GroupNamePtr = InGroupName;
+		return true;
+	}
+
+	return false;
+}
+
+bool FAnimNode_ChooserPlayer::SetGroupRole(EAnimGroupRole::Type InRole)
+{
+#if WITH_EDITORONLY_DATA
+	GroupRole = InRole;
+#endif
+	
+	if(TEnumAsByte<EAnimGroupRole::Type>* GroupRolePtr = GET_INSTANCE_ANIM_NODE_DATA_PTR(TEnumAsByte<EAnimGroupRole::Type>, GroupRole))
+	{
+		*GroupRolePtr = InRole;
+		return true;
+	}
+
+	return false;
+}
+
+bool FAnimNode_ChooserPlayer::SetGroupMethod(EAnimSyncMethod InMethod)
+{
+#if WITH_EDITORONLY_DATA
+	Method = InMethod;
+#endif
+
+	if(EAnimSyncMethod* MethodPtr = GET_INSTANCE_ANIM_NODE_DATA_PTR(EAnimSyncMethod, Method))
+	{
+		*MethodPtr = InMethod;
+		return true;
+	}
+
+	return false;
+}
+
+bool FAnimNode_ChooserPlayer::IsLooping() const
+{
+	if (!AnimPlayers.IsEmpty())
+	{
+		const FBlendStackAnimPlayer& AnimPlayer = AnimPlayers.First();
+		if (Cast<UBlendSpace>(AnimPlayer.GetAnimationAsset()))
+		{
+			return AnimPlayer.BlendSpacePlayerNode.IsLooping();
+		}
+		else
+		{
+			return AnimPlayer.SequencePlayerNode.IsLooping();
+		}
+	}
+	return false;
+}
+
+
+bool FAnimNode_ChooserPlayer::SetIgnoreForRelevancyTest(bool bInIgnoreForRelevancyTest)
+{
+#if WITH_EDITORONLY_DATA
+	bIgnoreForRelevancyTest = bInIgnoreForRelevancyTest;
+#endif
+
+	if (bool* bIgnoreForRelevancyTestPtr = GET_INSTANCE_ANIM_NODE_DATA_PTR(bool, bIgnoreForRelevancyTest))
+	{
+		*bIgnoreForRelevancyTestPtr = bInIgnoreForRelevancyTest;
+		return true;
+	}
+
+	return false;
+}
\ No newline at end of file
diff --git a/Engine/Plugins/Experimental/Chooser/Source/ChooserEditor/ChooserEditor.Build.cs b/Engine/Plugins/Experimental/Chooser/Source/ChooserEditor/ChooserEditor.Build.cs
--- a/Engine/Plugins/Experimental/Chooser/Source/ChooserEditor/ChooserEditor.Build.cs
+++ b/Engine/Plugins/Experimental/Chooser/Source/ChooserEditor/ChooserEditor.Build.cs
@@ -30,7 +30,7 @@ namespace UnrealBuildTool.Rules
 					"GameplayTags",
 					"GameplayTagsEditor",
 					"StructUtils",
-					"KismetCompiler"
+					"KismetCompiler",
 					// ... add private dependencies that you statically link with here ...
 				}
 			);
diff --git a/Engine/Plugins/Experimental/Chooser/Source/ChooserUncooked/ChooserUncooked.Build.cs b/Engine/Plugins/Experimental/Chooser/Source/ChooserUncooked/ChooserUncooked.Build.cs
--- a/Engine/Plugins/Experimental/Chooser/Source/ChooserUncooked/ChooserUncooked.Build.cs
+++ b/Engine/Plugins/Experimental/Chooser/Source/ChooserUncooked/ChooserUncooked.Build.cs
@@ -19,7 +19,10 @@ namespace UnrealBuildTool.Rules
 					"GameplayTags",
 					"StructUtils",
 					"BlueprintGraph",
-					"KismetCompiler"
+					"KismetCompiler",
+					"AnimGraph",
+					"BlendStackEditor",
+					"BlendStack"
 					// ... add private dependencies that you statically link with here ...
 				}
 			);
diff --git a/Engine/Plugins/Experimental/Chooser/Source/ChooserUncooked/Private/AnimGraphNode_ChooserPlayer.cpp b/Engine/Plugins/Experimental/Chooser/Source/ChooserUncooked/Private/AnimGraphNode_ChooserPlayer.cpp
new file mode 100644
--- /dev/null
+++ b/Engine/Plugins/Experimental/Chooser/Source/ChooserUncooked/Private/AnimGraphNode_ChooserPlayer.cpp
@@ -0,0 +1,68 @@
+// Copyright Epic Games, Inc. All Rights Reserved.
+
+#include "AnimGraphNode_ChooserPlayer.h"
+
+#include <chrono>
+
+#include "EditorCategoryUtils.h"
+#include "Animation/AnimAttributes.h"
+#include "Animation/AnimRootMotionProvider.h"
+#include "ChooserPropertyAccess.h"
+#include "Elements/Interfaces/TypedElementDataStorageCompatibilityInterface.h"
+
+#define LOCTEXT_NAMESPACE "AnimGraphNode_ChooserPlayer"
+
+void UAnimGraphNode_ChooserPlayer::UpdateContextData()
+{
+	if (Node.ChooserContextDefinition.IsEmpty())
+	{
+		Node.ChooserContextDefinition.SetNum(2);
+		Node.ChooserContextDefinition[0].InitializeAs(FContextObjectTypeClass::StaticStruct());
+		FContextObjectTypeClass& ClassEntry = Node.ChooserContextDefinition[0].GetMutable<FContextObjectTypeClass>();
+		ClassEntry.Class = GetBlueprint()->GeneratedClass;
+		ClassEntry.Direction = EContextObjectDirection::ReadWrite;
+		Node.ChooserContextDefinition[1].InitializeAs(FContextObjectTypeStruct::StaticStruct());
+		FContextObjectTypeStruct& StructEntry = Node.ChooserContextDefinition[1].GetMutable<FContextObjectTypeStruct>();
+		StructEntry.Struct = FChooserPlayerSettings::StaticStruct();
+		ClassEntry.Direction = EContextObjectDirection::ReadWrite;
+	}
+}
+
+void UAnimGraphNode_ChooserPlayer::PostPlacedNewNode()
+{
+	Super::PostPlacedNewNode();
+	UpdateContextData();
+}
+
+FLinearColor UAnimGraphNode_ChooserPlayer::GetNodeTitleColor() const
+{
+	return FLinearColor(0.10f, 0.60f, 0.12f);
+}
+
+FText UAnimGraphNode_ChooserPlayer::GetTooltipText() const
+{
+	return LOCTEXT("NodeToolTip", "Selects Animation Assets using a Chooser, and plays them with an underlying BlendStack.");
+}
+
+FText UAnimGraphNode_ChooserPlayer::GetNodeTitle(ENodeTitleType::Type TitleType) const
+{
+	return LOCTEXT("NodeTitle", "Chooser Player");
+}
+
+FText UAnimGraphNode_ChooserPlayer::GetMenuCategory() const
+{
+	return LOCTEXT("NodeCategory", "Animation|Sequences");
+}
+
+void UAnimGraphNode_ChooserPlayer::GetOutputLinkAttributes(FNodeAttributeArray& OutAttributes) const
+{
+	OutAttributes.Add(UE::Anim::FAttributes::Curves);
+	OutAttributes.Add(UE::Anim::FAttributes::Attributes);
+
+	if (UE::Anim::IAnimRootMotionProvider::Get())
+	{
+		OutAttributes.Add(UE::Anim::IAnimRootMotionProvider::AttributeName);
+	}
+}
+
+#undef LOCTEXT_NAMESPACE
diff --git a/Engine/Plugins/Experimental/Chooser/Source/ChooserUncooked/Public/AnimGraphNode_ChooserPlayer.h b/Engine/Plugins/Experimental/Chooser/Source/ChooserUncooked/Public/AnimGraphNode_ChooserPlayer.h
new file mode 100644
--- /dev/null
+++ b/Engine/Plugins/Experimental/Chooser/Source/ChooserUncooked/Public/AnimGraphNode_ChooserPlayer.h
@@ -0,0 +1,44 @@
+// Copyright Epic Games, Inc. All Rights Reserved.
+
+#pragma once
+
+#include "CoreMinimal.h"
+#include "UObject/ObjectMacros.h"
+#include "AnimGraphNode_Base.h"
+#include "AnimNode_ChooserPlayer.h"
+#include "AnimGraphNode_BlendStack.h"
+#include "IHasContext.h"
+#include "AnimGraphNode_ChooserPlayer.generated.h"
+
+UCLASS(MinimalAPI)
+class UAnimGraphNode_ChooserPlayer : public UAnimGraphNode_BlendStack_Base, public IHasContextClass
+{
+	GENERATED_BODY()
+
+public:
+
+	UPROPERTY(EditAnywhere, Category=Settings)
+	FAnimNode_ChooserPlayer Node;
+
+	// IHasContextClass interface
+	virtual TConstArrayView<FInstancedStruct> GetContextData() const { return Node.ChooserContextDefinition; }
+	// end of IHasContextClass interface
+	
+	// UEdGraphNode interface
+	virtual void PostPlacedNewNode() override;
+	virtual FLinearColor GetNodeTitleColor() const override;
+	virtual FText GetTooltipText() const override;
+	virtual FText GetNodeTitle(ENodeTitleType::Type TitleType) const override;
+	virtual FText GetMenuCategory() const override;
+	// End of UEdGraphNode interface
+
+	virtual FAnimNode_BlendStack_Standalone* GetBlendStackNode() const override { return (FAnimNode_BlendStack_Standalone*)(&Node); }
+
+	// UAnimGraphNode_Base interface
+	virtual void GetOutputLinkAttributes(FNodeAttributeArray& OutAttributes) const override;
+	// End of UAnimGraphNode_Base interface
+
+
+private:
+	void UpdateContextData();
+};
