Fix linux compile issue

diff --git a/Engine/Plugins/Experimental/Chooser/Source/Chooser/Chooser.Build.cs b/Engine/Plugins/Experimental/Chooser/Source/Chooser/Chooser.Build.cs
--- a/Engine/Plugins/Experimental/Chooser/Source/Chooser/Chooser.Build.cs
+++ b/Engine/Plugins/Experimental/Chooser/Source/Chooser/Chooser.Build.cs
@@ -21,6 +21,8 @@ namespace UnrealBuildTool.Rules
 					"CoreUObject",
 					"Engine",
 					"GameplayTags",
+					"AnimationCore",
+					"AnimGraphRuntime",
 					"BlendStack"
 					// ... add private dependencies that you statically link with here ...
 				}
