fix linux build failure some more

diff --git a/Engine/Plugins/Experimental/Chooser/Source/ChooserUncooked/ChooserUncooked.Build.cs b/Engine/Plugins/Experimental/Chooser/Source/ChooserUncooked/ChooserUncooked.Build.cs
--- a/Engine/Plugins/Experimental/Chooser/Source/ChooserUncooked/ChooserUncooked.Build.cs
+++ b/Engine/Plugins/Experimental/Chooser/Source/ChooserUncooked/ChooserUncooked.Build.cs
@@ -21,6 +21,7 @@ namespace UnrealBuildTool.Rules
 					"BlueprintGraph",
 					"KismetCompiler",
 					"AnimGraph",
+					"AnimGraphRuntime",
 					"BlendStackEditor",
 					"BlendStack"
 					// ... add private dependencies that you statically link with here ...
