Export virtuals
#rnx

diff --git a/Engine/Plugins/Experimental/Chooser/Source/Chooser/Internal/AnimNode_ChooserPlayer.h b/Engine/Plugins/Experimental/Chooser/Source/Chooser/Internal/AnimNode_ChooserPlayer.h
--- a/Engine/Plugins/Experimental/Chooser/Source/Chooser/Internal/AnimNode_ChooserPlayer.h
+++ b/Engine/Plugins/Experimental/Chooser/Source/Chooser/Internal/AnimNode_ChooserPlayer.h
@@ -168,21 +168,21 @@ public:
 	TArray<FInstancedStruct> ChooserContextDefinition;
 
 	// FAnimNode_Base interface
-	virtual void Initialize_AnyThread(const FAnimationInitializeContext& Context) override;
-	virtual void Evaluate_AnyThread(FPoseContext& Output) override;
+	CHOOSER_API virtual void Initialize_AnyThread(const FAnimationInitializeContext& Context) override;
+	CHOOSER_API virtual void Evaluate_AnyThread(FPoseContext& Output) override;
 	// End of FAnimNode_Base interface
 	
 	// FAnimNode_AssetPlayerBase interface
-	virtual void UpdateAssetPlayer(const FAnimationUpdateContext& Context) override;
-	virtual FName GetGroupName() const override;
-	virtual EAnimGroupRole::Type GetGroupRole() const override;
-	virtual EAnimSyncMethod GetGroupMethod() const override;
-	virtual bool GetIgnoreForRelevancyTest() const override;
-	virtual bool IsLooping() const override;
-	virtual bool SetGroupName(FName InGroupName) override;
-	virtual bool SetGroupRole(EAnimGroupRole::Type InRole) override;
-	virtual bool SetGroupMethod(EAnimSyncMethod InMethod) override;
-	virtual bool SetIgnoreForRelevancyTest(bool bInIgnoreForRelevancyTest) override;
+	CHOOSER_API virtual void UpdateAssetPlayer(const FAnimationUpdateContext& Context) override;
+	CHOOSER_API virtual FName GetGroupName() const override;
+	CHOOSER_API virtual EAnimGroupRole::Type GetGroupRole() const override;
+	CHOOSER_API virtual EAnimSyncMethod GetGroupMethod() const override;
+	CHOOSER_API virtual bool GetIgnoreForRelevancyTest() const override;
+	CHOOSER_API virtual bool IsLooping() const override;
+	CHOOSER_API virtual bool SetGroupName(FName InGroupName) override;
+	CHOOSER_API virtual bool SetGroupRole(EAnimGroupRole::Type InRole) override;
+	CHOOSER_API virtual bool SetGroupMethod(EAnimSyncMethod InMethod) override;
+	CHOOSER_API virtual bool SetIgnoreForRelevancyTest(bool bInIgnoreForRelevancyTest) override;
 	// End of FAnimNode_AssetPlayerBase interface
 	
 private:
