Fix for struct property bindings in proxy table struct outputs resetting the struct value when binding a different variable of the same type

diff --git a/Engine/Plugins/Experimental/Chooser/Source/ProxyTableEditor/Private/StructOutputDataCustomization.cpp b/Engine/Plugins/Experimental/Chooser/Source/ProxyTableEditor/Private/StructOutputDataCustomization.cpp
--- a/Engine/Plugins/Experimental/Chooser/Source/ProxyTableEditor/Private/StructOutputDataCustomization.cpp
+++ b/Engine/Plugins/Experimental/Chooser/Source/ProxyTableEditor/Private/StructOutputDataCustomization.cpp
@@ -106,6 +106,7 @@ void FStructOutputDataCustomization::CustomizeHeader(TSharedRef<IPropertyHandle>
 								if (const FStructProperty* StructProperty = CastField<const FStructProperty>(Property))
 								{
 									StructPropertyBinding->StructType = StructProperty->Struct;
+									StructPropertyBinding->DisplayName = StructProperty->GetDisplayNameText().ToString();
 								}
          	 				}
          	 				else if (InBindingChain.Num() ==1)
@@ -123,7 +124,10 @@ void FStructOutputDataCustomization::CustomizeHeader(TSharedRef<IPropertyHandle>
 							// reset the type of the FInstancedStruct storing the value
 							ValueHandle->GetValueData(ValuePtr);
 							FInstancedStruct* ValueStruct = reinterpret_cast<FInstancedStruct*>(ValuePtr);
-							ValueStruct->InitializeAs(StructPropertyBinding->StructType);
+         	 				if (ValueStruct->GetScriptStruct() != StructPropertyBinding->StructType)
+         	 				{
+         	 					ValueStruct->InitializeAs(StructPropertyBinding->StructType);
+         	 				}
 	
          	 				BindingHandle->NotifyPostChange(EPropertyChangeType::ValueSet);
 							ValueHandle->NotifyPostChange(EPropertyChangeType::ValueSet);
