Chooser Table Editor crash fixes

diff --git a/Engine/Plugins/Experimental/Chooser/Source/ChooserEditor/Private/ChooserTableEditor.cpp b/Engine/Plugins/Experimental/Chooser/Source/ChooserEditor/Private/ChooserTableEditor.cpp
--- a/Engine/Plugins/Experimental/Chooser/Source/ChooserEditor/Private/ChooserTableEditor.cpp
+++ b/Engine/Plugins/Experimental/Chooser/Source/ChooserEditor/Private/ChooserTableEditor.cpp
@@ -539,6 +539,9 @@ public:
 
 	FReply OnDragDetected(const FGeometry& MyGeometry, const FPointerEvent& MouseEvent) override
 	{
+		// clear row selection so that delete key can't cause the selected row to be deleted
+		ChooserEditor->ClearSelectedRows();
+		
 		TSharedRef<FChooserRowDragDropOp> DragDropOp = FChooserRowDragDropOp::New(ChooserEditor, RowIndex);
 		return FReply::Handled().BeginDragDrop(DragDropOp);
 	}
@@ -698,14 +701,17 @@ public:
 		{
 			if (Chooser == Operation->ChooserEditor->GetChooser())
 			{
+				int NewRowIndex;
 				if (bDropAbove)
 				{
-					Editor->MoveRow(Operation->RowIndex, RowIndex->RowIndex);
+					NewRowIndex = Editor->MoveRow(Operation->RowIndex, RowIndex->RowIndex);
 				}
 				else
 				{
-					Editor->MoveRow(Operation->RowIndex, RowIndex->RowIndex+1);
+					NewRowIndex = Editor->MoveRow(Operation->RowIndex, RowIndex->RowIndex+1);
 				}
+				Editor->SelectRow(NewRowIndex);
+				
 				return FReply::Handled();		
 			}
 		}
@@ -739,7 +745,7 @@ void FChooserTableEditor::SelectRootProperties()
 	}
 }
 
-void FChooserTableEditor::MoveRow(int SourceRowIndex, int TargetRowIndex)
+int FChooserTableEditor::MoveRow(int SourceRowIndex, int TargetRowIndex)
 {
 	UChooserTable* Chooser = Cast<UChooserTable>(EditingObjects[0]);
 	TargetRowIndex = FMath::Min(TargetRowIndex,Chooser->ResultsStructs.Num());
@@ -763,9 +769,11 @@ void FChooserTableEditor::MoveRow(int SourceRowIndex, int TargetRowIndex)
 	Chooser->ResultsStructs.Insert(Result, TargetRowIndex);
 
 	UpdateTableRows();
+
+	return TargetRowIndex;
 }
 	
-void FChooserTableEditor::SelectRow(int32 RowIndex) const
+void FChooserTableEditor::SelectRow(int32 RowIndex)
 {
 	if (TableRows.IsValidIndex(RowIndex))
 	{
@@ -776,6 +784,13 @@ void FChooserTableEditor::SelectRow(int32 RowIndex) const
 		}
 	}
 }
+	
+void FChooserTableEditor::ClearSelectedRows() 
+{
+	SelectedRows.SetNum(0);
+	TableView->ClearSelection();
+	SelectRootProperties();
+}
 
 void FChooserTableEditor::UpdateTableColumns()
 {
@@ -1046,6 +1061,7 @@ TSharedRef<SDockTab> FChooserTableEditor::SpawnTableTab( const FSpawnTabArgs& Ar
     			.ListItemsSource(&TableRows)
 				.OnKeyDownHandler_Lambda([this](const FGeometry&, const FKeyEvent& Event)
 				{
+					
 					if (Event.GetKey() == EKeys::Delete)
 					{
 						const FScopedTransaction Transaction(LOCTEXT("Delete Row Transaction", "Delete Row"));
@@ -1378,7 +1394,7 @@ void FChooserRowDetails::CustomizeDetails(IDetailLayoutBuilder& DetailBuilder)
 	UChooserRowDetails* Row = Cast<UChooserRowDetails>(Objects[0]);
 	UChooserTable* Chooser = Row->Chooser;
 	
-	if (Chooser->ResultsStructs.Num() > Row->Row)
+	if (Chooser->ResultsStructs.IsValidIndex(Row->Row))
 	{
 		IDetailCategoryBuilder& PropertiesCategory = DetailBuilder.EditCategory("Row Properties");
 
diff --git a/Engine/Plugins/Experimental/Chooser/Source/ChooserEditor/Private/ChooserTableEditor.h b/Engine/Plugins/Experimental/Chooser/Source/ChooserEditor/Private/ChooserTableEditor.h
--- a/Engine/Plugins/Experimental/Chooser/Source/ChooserEditor/Private/ChooserTableEditor.h
+++ b/Engine/Plugins/Experimental/Chooser/Source/ChooserEditor/Private/ChooserTableEditor.h
@@ -103,7 +103,9 @@ namespace UE::ChooserEditor
 		void ClearSelectedColumn();
 		void DeleteColumn(int Index);
 		void AddColumn(const UScriptStruct* ColumnType);
-		void MoveRow(int SourceRowIndex, int TargetIndex);
+		int MoveRow(int SourceRowIndex, int TargetIndex);
+		void SelectRow(int32 RowIndex);
+		void ClearSelectedRows(); 
 	private:
 
 		void SelectRootProperties();
@@ -146,8 +148,7 @@ namespace UE::ChooserEditor
 		TSharedPtr<SHeaderRow> HeaderRow;
 		TSharedPtr<SListView<TSharedPtr<FChooserTableRow>>> TableView;
 	public:
-		void SelectRow(int32 RowIndex) const;
-		
+
 		TSharedPtr<SComboButton>& GetCreateRowComboButton() { return CreateRowComboButton; };
 
 		/** The name given to all instances of this type of editor */
diff --git a/Engine/Plugins/Experimental/Chooser/Source/ProxyTableEditor/Private/ProxyTableEditor.cpp b/Engine/Plugins/Experimental/Chooser/Source/ProxyTableEditor/Private/ProxyTableEditor.cpp
--- a/Engine/Plugins/Experimental/Chooser/Source/ProxyTableEditor/Private/ProxyTableEditor.cpp
+++ b/Engine/Plugins/Experimental/Chooser/Source/ProxyTableEditor/Private/ProxyTableEditor.cpp
@@ -422,6 +422,9 @@ public:
 
 	FReply OnDragDetected(const FGeometry& MyGeometry, const FPointerEvent& MouseEvent) override
 	{
+		// clear row selection so that delete key can't cause the selected row to be deleted
+		ProxyEditor->ClearSelectedRows();
+		
 		TSharedRef<FProxyRowDragDropOp> DragDropOp = FProxyRowDragDropOp::New(Row);
 		return FReply::Handled().BeginDragDrop(DragDropOp);
 	}
@@ -631,12 +634,14 @@ public:
 				if (Row->ProxyTable == Operation->Row->ProxyTable)
 				{
 					// move row within a proxy table
-					Editor->MoveRow(Operation->Row->RowIndex, InsertIndex);
+					int NewIndex = Editor->MoveRow(Operation->Row->RowIndex, InsertIndex);
+					Editor->SelectRow(NewIndex);
 					return FReply::Handled();		
 				}
 				else
 				{
 					Editor->InsertEntry(Operation->Row->ProxyTable->Entries[Operation->Row->RowIndex], InsertIndex);
+					Editor->SelectRow(InsertIndex);
 					return FReply::Handled();
 				}
 			}
@@ -717,6 +722,23 @@ void FProxyTableEditor::DeleteSelectedRows()
 	
 	UpdateTableRows();
 }
+	
+	
+void FProxyTableEditor::ClearSelectedRows() 
+{
+	SelectedRows.SetNum(0);
+	TableView->ClearSelection();
+	SelectRootProperties();
+}
+	
+void FProxyTableEditor::SelectRow(TSharedPtr<FProxyTableRow> Row)
+{
+	if (!TableView->IsItemSelected(Row))
+	{
+		TableView->ClearSelection();
+		TableView->SetItemSelection(Row, true, ESelectInfo::OnMouseClick);
+	}
+}
 
 void FProxyTableEditor::InsertEntry(FProxyEntry& Entry, int RowIndex)
 {
@@ -734,7 +756,7 @@ void FProxyTableEditor::InsertEntry(FProxyEntry& Entry, int RowIndex)
 	UpdateTableRows();
 }
 
-void FProxyTableEditor::MoveRow(int SourceRowIndex, int TargetRowIndex)
+int FProxyTableEditor::MoveRow(int SourceRowIndex, int TargetRowIndex)
 {
 	UProxyTable* Table = Cast<UProxyTable>(EditingObjects[0]);
 	TargetRowIndex = FMath::Min(TargetRowIndex,Table->Entries.Num());
@@ -752,6 +774,7 @@ void FProxyTableEditor::MoveRow(int SourceRowIndex, int TargetRowIndex)
 	Table->Entries.Insert(Entry, TargetRowIndex);
 
 	UpdateTableRows();
+	return TargetRowIndex;
 }
 
 void FProxyTableEditor::TreeViewExpansionChanged(TSharedPtr<FProxyTableEditor::FProxyTableRow> InItem, bool bShouldBeExpanded)
diff --git a/Engine/Plugins/Experimental/Chooser/Source/ProxyTableEditor/Private/ProxyTableEditor.h b/Engine/Plugins/Experimental/Chooser/Source/ProxyTableEditor/Private/ProxyTableEditor.h
--- a/Engine/Plugins/Experimental/Chooser/Source/ProxyTableEditor/Private/ProxyTableEditor.h
+++ b/Engine/Plugins/Experimental/Chooser/Source/ProxyTableEditor/Private/ProxyTableEditor.h
@@ -94,9 +94,12 @@ namespace UE::ProxyTableEditor
 
 		void UpdateTableColumns();
 		void UpdateTableRows();
-		void MoveRow(int SourceRowIndex, int TargetIndex);
+		int MoveRow(int SourceRowIndex, int TargetIndex);
 		void InsertEntry(FProxyEntry& Entry, int RowIndex);
 		void DeleteSelectedRows();
+		void ClearSelectedRows();
+		void SelectRow(int RowIndex) { if (TableRows.IsValidIndex(RowIndex)) { SelectRow(TableRows[RowIndex]); } }
+		void SelectRow(TSharedPtr<FProxyTableRow> Row);
 	private:
 		void AddInheritedRows(UProxyTable* ProxyTable);
 		void SelectRootProperties();
@@ -142,15 +145,7 @@ namespace UE::ProxyTableEditor
 		TMap<UProxyTable*, bool> ImportedTablesExpansionState;
 		TSet<FProxyEntry> ReferencedProxyEntries;
 		TSet<UProxyTable*> ReferencedProxyTables;
-		
-		void SelectRow(TSharedPtr<FProxyTableRow> Row)
-		{
-			if (!TableView->IsItemSelected(Row))
-			{
-				TableView->ClearSelection();
-				TableView->SetItemSelection(Row, true, ESelectInfo::OnMouseClick);
-			}
-		}
+
 		TSharedPtr<SComboButton>& GetCreateRowComboButton() { return CreateRowComboButton; };
 
 		/** The name given to all instances of this type of editor */
