ChooserPlayer support for choosing between different frames of the same sequence as poses, with PlaybackRate 0

diff --git a/Engine/Plugins/Experimental/Chooser/Source/Chooser/Internal/AnimNode_ChooserPlayer.h b/Engine/Plugins/Experimental/Chooser/Source/Chooser/Internal/AnimNode_ChooserPlayer.h
--- a/Engine/Plugins/Experimental/Chooser/Source/Chooser/Internal/AnimNode_ChooserPlayer.h
+++ b/Engine/Plugins/Experimental/Chooser/Source/Chooser/Internal/AnimNode_ChooserPlayer.h
@@ -189,6 +189,7 @@ private:
 	UAnimationAsset* ChooseAsset(const FAnimationUpdateContext& Context);
 	FChooserEvaluationContext ChooserContext;
 	UAnimationAsset* CurrentAsset = nullptr;
+	float CurrentStartTime = 0;
 	uint32 CurveOverridesIndex = 0;
 	TBaseBlendedCurve<FDefaultAllocator, UE::Anim::FCurveElement> OverrideCurves[2]; 
 	
diff --git a/Engine/Plugins/Experimental/Chooser/Source/Chooser/Private/AnimNode_ChooserPlayer.cpp b/Engine/Plugins/Experimental/Chooser/Source/Chooser/Private/AnimNode_ChooserPlayer.cpp
--- a/Engine/Plugins/Experimental/Chooser/Source/Chooser/Private/AnimNode_ChooserPlayer.cpp
+++ b/Engine/Plugins/Experimental/Chooser/Source/Chooser/Private/AnimNode_ChooserPlayer.cpp
@@ -80,10 +80,15 @@ void FAnimNode_ChooserPlayer::UpdateAssetPlayer(const FAnimationUpdateContext& C
 		NewAsset = ChooseAsset(Context);
 	}
 
-	if (bJustBecameRelevant || NewAsset != CurrentAsset)
-	{
-		const FChooserPlayerSettings& Settings = ChooserContext.Params[1].Get<FChooserPlayerSettings>();
-			
+	const FChooserPlayerSettings& Settings = ChooserContext.Params[1].Get<FChooserPlayerSettings>();
+
+	// Restart the animation:
+	// - if this node just became relevant
+	// - if we chose a new animation
+	// - for playback rate of 0, when the start time changes - for choosing poses as frames of an animation sequence
+	if (bJustBecameRelevant || NewAsset != CurrentAsset ||
+		(CurrentStartTime != Settings.StartTime && Settings.PlaybackRate == 0.0f)) 
+	{			
 		CurveOverridesIndex = (CurveOverridesIndex + 1) % 2;
 		OverrideCurves[CurveOverridesIndex].Empty();
 		if (!Settings.CurveOverrides.Values.IsEmpty())
@@ -116,6 +121,7 @@ void FAnimNode_ChooserPlayer::UpdateAssetPlayer(const FAnimationUpdateContext& C
 			GetGroupName(), GetGroupRole(), GetGroupMethod());
 		
 		CurrentAsset = NewAsset;
+		CurrentStartTime = Settings.StartTime;
 	}
 
 	// Update blend space parameters
