fallback result for chooser tables

diff --git a/Engine/Plugins/Experimental/Chooser/Source/Chooser/Internal/Chooser.h b/Engine/Plugins/Experimental/Chooser/Source/Chooser/Internal/Chooser.h
--- a/Engine/Plugins/Experimental/Chooser/Source/Chooser/Internal/Chooser.h
+++ b/Engine/Plugins/Experimental/Chooser/Source/Chooser/Internal/Chooser.h
@@ -76,6 +76,11 @@ public:
 	TArray<TScriptInterface<IChooserColumn>> Columns_DEPRECATED;
 #endif
 
+	
+	// FallbackResult will be used as the Result if there are no rows in the chooser which pass all filters.  If FallbackResult is not assigned, then the Chooser will return null in that case.
+	UPROPERTY(EditAnywhere, Meta = (ExcludeBaseStruct, BaseStruct = "/Script/Chooser.ObjectChooserBase"), Category = "Fallback")
+	FInstancedStruct FallbackResult;
+	
 	// Each possible Result (Rows of chooser table)
 	UPROPERTY(EditAnywhere, NoClear, DisplayName = "Results", Meta = (ExcludeBaseStruct, BaseStruct = "/Script/Chooser.ObjectChooserBase"), Category = "Hidden")
 	TArray<FInstancedStruct> ResultsStructs;
diff --git a/Engine/Plugins/Experimental/Chooser/Source/Chooser/Internal/OutputBoolColumn.h b/Engine/Plugins/Experimental/Chooser/Source/Chooser/Internal/OutputBoolColumn.h
--- a/Engine/Plugins/Experimental/Chooser/Source/Chooser/Internal/OutputBoolColumn.h
+++ b/Engine/Plugins/Experimental/Chooser/Source/Chooser/Internal/OutputBoolColumn.h
@@ -25,12 +25,16 @@ struct CHOOSER_API FOutputBoolColumn : public FChooserColumnBase
 	mutable bool TestValue=false;
 #endif
 	
+	// FallbackValue will be used as the output value if the all rows in the chooser fail, and the FallbackResult from the chooser is used.
+	UPROPERTY(EditAnywhere, Category=Data);
+	bool bFallbackValue = false;
+	
 #if WITH_EDITORONLY_DATA
-	UPROPERTY(EditAnywhere, Category=Runtime);
-	bool DefaultRowValue = 0;
+	UPROPERTY(EditAnywhere, Category=Data);
+	bool DefaultRowValue = false;
 #endif
 	
-	UPROPERTY(EditAnywhere, Category=Runtime);
+	UPROPERTY(EditAnywhere, Category=Data);
 	TArray<bool> RowValues; 
 
 	CHOOSER_COLUMN_BOILERPLATE(FChooserParameterBoolBase);
diff --git a/Engine/Plugins/Experimental/Chooser/Source/Chooser/Internal/OutputEnumColumn.h b/Engine/Plugins/Experimental/Chooser/Source/Chooser/Internal/OutputEnumColumn.h
--- a/Engine/Plugins/Experimental/Chooser/Source/Chooser/Internal/OutputEnumColumn.h
+++ b/Engine/Plugins/Experimental/Chooser/Source/Chooser/Internal/OutputEnumColumn.h
@@ -31,6 +31,10 @@ public:
 	UPROPERTY(EditAnywhere, NoClear, Meta = (ExcludeBaseStruct, BaseStruct = "/Script/Chooser.ChooserParameterEnumBase"), Category = "Data")
 	FInstancedStruct InputValue;
 
+	// FallbackValue will be used as the output value if the all rows in the chooser fail, and the FallbackResult from the chooser is used.
+	UPROPERTY(EditAnywhere, Category = "Data")
+	FChooserOutputEnumRowData FallbackValue;
+
 #if WITH_EDITORONLY_DATA
 	UPROPERTY(EditAnywhere, Category = "Data")
 	FChooserOutputEnumRowData DefaultRowValue;
diff --git a/Engine/Plugins/Experimental/Chooser/Source/Chooser/Internal/OutputFloatColumn.h b/Engine/Plugins/Experimental/Chooser/Source/Chooser/Internal/OutputFloatColumn.h
--- a/Engine/Plugins/Experimental/Chooser/Source/Chooser/Internal/OutputFloatColumn.h
+++ b/Engine/Plugins/Experimental/Chooser/Source/Chooser/Internal/OutputFloatColumn.h
@@ -24,13 +24,17 @@ struct CHOOSER_API FOutputFloatColumn : public FChooserColumnBase
 #if WITH_EDITOR
 	mutable double TestValue=false;
 #endif
+
+	// FallbackValue will be used as the output value if the all rows in the chooser fail, and the FallbackResult from the chooser is used.
+	UPROPERTY(EditAnywhere, Category=Data);
+	double FallbackValue = 0.0;
 	
 #if WITH_EDITORONLY_DATA
-	UPROPERTY(EditAnywhere, Category=Runtime);
+	UPROPERTY(EditAnywhere, Category=Data);
 	double DefaultRowValue = 0.0;
 #endif
 	
-	UPROPERTY(EditAnywhere, Category=Runtime);
+	UPROPERTY(EditAnywhere, Category=Data);
 	TArray<double> RowValues; 
 
 	CHOOSER_COLUMN_BOILERPLATE(FChooserParameterFloatBase);
diff --git a/Engine/Plugins/Experimental/Chooser/Source/Chooser/Internal/OutputStructColumn.h b/Engine/Plugins/Experimental/Chooser/Source/Chooser/Internal/OutputStructColumn.h
--- a/Engine/Plugins/Experimental/Chooser/Source/Chooser/Internal/OutputStructColumn.h
+++ b/Engine/Plugins/Experimental/Chooser/Source/Chooser/Internal/OutputStructColumn.h
@@ -60,12 +60,16 @@ struct CHOOSER_API FOutputStructColumn : public FChooserColumnBase
 	mutable FInstancedStruct TestValue;
 #endif
 	
+	// FallbackValue will be used as the output value if the all rows in the chooser fail, and the FallbackResult from the chooser is used.
+	UPROPERTY(EditAnywhere, Meta = (StructTypeConst), Category=Data);
+   	FInstancedStruct FallbackValue;
+	
 #if WITH_EDITORONLY_DATA
-	UPROPERTY(EditAnywhere, Meta = (StructTypeConst), Category=Runtime);
+	UPROPERTY(EditAnywhere, Meta = (StructTypeConst), Category=Data);
 	FInstancedStruct DefaultRowValue;
 #endif
 	
-	UPROPERTY(EditAnywhere, Meta = (StructTypeConst), Category=Runtime);
+	UPROPERTY(EditAnywhere, Meta = (StructTypeConst), Category=Data);
 	TArray<FInstancedStruct> RowValues; 
 
 	CHOOSER_COLUMN_BOILERPLATE(FChooserParameterStructBase);
diff --git a/Engine/Plugins/Experimental/Chooser/Source/Chooser/Private/Chooser.cpp b/Engine/Plugins/Experimental/Chooser/Source/Chooser/Private/Chooser.cpp
--- a/Engine/Plugins/Experimental/Chooser/Source/Chooser/Private/Chooser.cpp
+++ b/Engine/Plugins/Experimental/Chooser/Source/Chooser/Private/Chooser.cpp
@@ -239,29 +239,31 @@ FObjectChooserBase::EIteratorStatus UChooserTable::EvaluateChooser(FChooserEvalu
 		}
 	}
 	
-	// of the rows that passed all column filters, iterate through them calling the callback until it returns Stop
 	bool bSetOutputs = false;
-	for (uint32 SelectedIndex : *IndicesOut)
+	if (IndicesOut->IsEmpty())
 	{
-		if (Chooser->ResultsStructs.Num() > (int32)SelectedIndex)
+		// if no rows passed the filter columns then return the FallbackResult and output the FallbackValue from each output column
+		
+		#if WITH_EDITOR
+		if (Context.DebuggingInfo.bCurrentDebugTarget)
 		{
-			const FObjectChooserBase& SelectedResult = Chooser->ResultsStructs[SelectedIndex].Get<FObjectChooserBase>();
+			Chooser->SetDebugSelectedRow(-1);
+		}
+		#endif
+		
+		if (Chooser->FallbackResult.IsValid())
+		{
+			const FObjectChooserBase& SelectedResult = Chooser->FallbackResult.Get<FObjectChooserBase>();
 			FObjectChooserBase::EIteratorStatus Status = SelectedResult.ChooseMulti(Context, Callback);
 			if (Status != FObjectChooserBase::EIteratorStatus::Continue)
 			{
 				bSetOutputs = true;
-				// trigger all output columns
+				// trigger all output columns to set their default output value
 				for (const FInstancedStruct& ColumnData : Chooser->ColumnsStructs)
 				{
 					const FChooserColumnBase& Column = ColumnData.Get<FChooserColumnBase>();
-					Column.SetOutputs(Context, SelectedIndex);
-				}
-#if WITH_EDITOR
-				if (Context.DebuggingInfo.bCurrentDebugTarget)
-				{
-					Chooser->SetDebugSelectedRow(SelectedIndex);
+					Column.SetOutputs(Context, -1);
 				}
-#endif
 			}
 			if (Status == FObjectChooserBase::EIteratorStatus::Stop)
 			{
@@ -269,8 +271,39 @@ FObjectChooserBase::EIteratorStatus UChooserTable::EvaluateChooser(FChooserEvalu
 			}
 		}
 	}
-
-	// If this is a nested chooser make sure the parent also sets the output vales from the row that contained this chooser
+	else
+	{
+		// of the rows that passed all column filters, iterate through them calling the callback until it returns Stop
+		for (uint32 SelectedIndex : *IndicesOut)
+		{
+			if (Chooser->ResultsStructs.Num() > (int32)SelectedIndex)
+			{
+				const FObjectChooserBase& SelectedResult = Chooser->ResultsStructs[SelectedIndex].Get<FObjectChooserBase>();
+				FObjectChooserBase::EIteratorStatus Status = SelectedResult.ChooseMulti(Context, Callback);
+				if (Status != FObjectChooserBase::EIteratorStatus::Continue)
+				{
+					bSetOutputs = true;
+					// trigger all output columns
+					for (const FInstancedStruct& ColumnData : Chooser->ColumnsStructs)
+					{
+						const FChooserColumnBase& Column = ColumnData.Get<FChooserColumnBase>();
+						Column.SetOutputs(Context, SelectedIndex);
+					}
+	#if WITH_EDITOR
+					if (Context.DebuggingInfo.bCurrentDebugTarget)
+					{
+						Chooser->SetDebugSelectedRow(SelectedIndex);
+					}
+	#endif
+				}
+				if (Status == FObjectChooserBase::EIteratorStatus::Stop)
+				{
+					return FObjectChooserBase::EIteratorStatus::Stop;
+				}
+			}
+		}
+	}
+	// If this is a nested chooser make sure the parent also sets the output values from the row that contained this chooser
 	return bSetOutputs ? FObjectChooserBase::EIteratorStatus::ContinueWithOutputs : FObjectChooserBase::EIteratorStatus::Continue;
 }
 
diff --git a/Engine/Plugins/Experimental/Chooser/Source/Chooser/Private/OutputBoolColumn.cpp b/Engine/Plugins/Experimental/Chooser/Source/Chooser/Private/OutputBoolColumn.cpp
--- a/Engine/Plugins/Experimental/Chooser/Source/Chooser/Private/OutputBoolColumn.cpp
+++ b/Engine/Plugins/Experimental/Chooser/Source/Chooser/Private/OutputBoolColumn.cpp
@@ -12,7 +12,13 @@ void FOutputBoolColumn::SetOutputs(FChooserEvaluationContext& Context, int RowIn
 {
 	if (InputValue.IsValid())
 	{
-		InputValue.Get<FChooserParameterBoolBase>().SetValue(Context, RowValues[RowIndex]);
+		bool bOutputValue = bFallbackValue;
+		if (RowValues.IsValidIndex(RowIndex))
+		{
+			bOutputValue = RowValues[RowIndex];
+		}
+		
+		InputValue.Get<FChooserParameterBoolBase>().SetValue(Context, bOutputValue);
 	}
 	
 #if WITH_EDITOR
diff --git a/Engine/Plugins/Experimental/Chooser/Source/Chooser/Private/OutputEnumColumn.cpp b/Engine/Plugins/Experimental/Chooser/Source/Chooser/Private/OutputEnumColumn.cpp
--- a/Engine/Plugins/Experimental/Chooser/Source/Chooser/Private/OutputEnumColumn.cpp
+++ b/Engine/Plugins/Experimental/Chooser/Source/Chooser/Private/OutputEnumColumn.cpp
@@ -14,7 +14,12 @@ void FOutputEnumColumn::SetOutputs(FChooserEvaluationContext& Context, int RowIn
 {
 	if (InputValue.IsValid())
 	{
-		InputValue.Get<FChooserParameterEnumBase>().SetValue(Context, RowValues[RowIndex].Value);
+		uint8 OutputValue = FallbackValue.Value;
+		if (RowValues.IsValidIndex(RowIndex))
+		{
+			OutputValue = RowValues[RowIndex].Value;
+		}
+		InputValue.Get<FChooserParameterEnumBase>().SetValue(Context, OutputValue);
 	}
 	
 #if WITH_EDITOR
diff --git a/Engine/Plugins/Experimental/Chooser/Source/Chooser/Private/OutputFloatColumn .cpp b/Engine/Plugins/Experimental/Chooser/Source/Chooser/Private/OutputFloatColumn .cpp
--- a/Engine/Plugins/Experimental/Chooser/Source/Chooser/Private/OutputFloatColumn .cpp	
+++ b/Engine/Plugins/Experimental/Chooser/Source/Chooser/Private/OutputFloatColumn .cpp	
@@ -12,7 +12,13 @@ void FOutputFloatColumn::SetOutputs(FChooserEvaluationContext& Context, int RowI
 {
 	if (InputValue.IsValid())
 	{
-		InputValue.Get<FChooserParameterFloatBase>().SetValue(Context, RowValues[RowIndex]);
+		double OutputValue = FallbackValue;
+		if (RowValues.IsValidIndex(RowIndex))
+		{
+			OutputValue = RowValues[RowIndex];
+		}
+	
+		InputValue.Get<FChooserParameterFloatBase>().SetValue(Context, OutputValue);
 	}
 	
 #if WITH_EDITOR
diff --git a/Engine/Plugins/Experimental/Chooser/Source/Chooser/Private/OutputStructColumn.cpp b/Engine/Plugins/Experimental/Chooser/Source/Chooser/Private/OutputStructColumn.cpp
--- a/Engine/Plugins/Experimental/Chooser/Source/Chooser/Private/OutputStructColumn.cpp
+++ b/Engine/Plugins/Experimental/Chooser/Source/Chooser/Private/OutputStructColumn.cpp
@@ -31,6 +31,11 @@ void FOutputStructColumn::StructTypeChanged()
 			DefaultRowValue.InitializeAs(Struct);
 		}
 
+		if (FallbackValue.GetScriptStruct() != Struct)
+		{
+			FallbackValue.InitializeAs(Struct);
+		}
+
 		for (FInstancedStruct& RowValue : RowValues)
 		{
 			if (RowValue.GetScriptStruct() != Struct)
@@ -52,7 +57,16 @@ void FOutputStructColumn::SetOutputs(FChooserEvaluationContext& Context, int Row
 {
 	if (InputValue.IsValid())
 	{
-		InputValue.Get<FChooserParameterStructBase>().SetValue(Context, RowValues[RowIndex]);
+		const FInstancedStruct* OutputValue = &FallbackValue;
+		if (RowValues.IsValidIndex(RowIndex))
+		{
+			OutputValue = &RowValues[RowIndex];
+		}
+
+		if (OutputValue && OutputValue->IsValid())
+		{
+			InputValue.Get<FChooserParameterStructBase>().SetValue(Context, *OutputValue);
+		}
 	}
 	
 #if WITH_EDITOR
diff --git a/Engine/Plugins/Experimental/Chooser/Source/Chooser/Private/RandomizeColumn.cpp b/Engine/Plugins/Experimental/Chooser/Source/Chooser/Private/RandomizeColumn.cpp
--- a/Engine/Plugins/Experimental/Chooser/Source/Chooser/Private/RandomizeColumn.cpp
+++ b/Engine/Plugins/Experimental/Chooser/Source/Chooser/Private/RandomizeColumn.cpp
@@ -88,7 +88,7 @@ void FRandomizeColumn::Filter(FChooserEvaluationContext& Context, const TArray<u
 
 void FRandomizeColumn::SetOutputs(FChooserEvaluationContext& Context, int RowIndex) const
 {
-	if (InputValue.IsValid())
+	if (InputValue.IsValid() && RowValues.IsValidIndex((RowIndex)))
 	{
 		const FChooserRandomizationContext* ConstRandomizationContext = nullptr;
 		InputValue.Get<FChooserParameterRandomizeBase>().GetValue(Context,ConstRandomizationContext);
