ChooserPlayer: improved override curve blending, trigger blend on mirror change, trigger blend on override curves change

diff --git a/Engine/Plugins/Experimental/Animation/BlendStack/Source/Runtime/Private/AnimNode_BlendStack.cpp b/Engine/Plugins/Experimental/Animation/BlendStack/Source/Runtime/Private/AnimNode_BlendStack.cpp
--- a/Engine/Plugins/Experimental/Animation/BlendStack/Source/Runtime/Private/AnimNode_BlendStack.cpp
+++ b/Engine/Plugins/Experimental/Animation/BlendStack/Source/Runtime/Private/AnimNode_BlendStack.cpp
@@ -130,6 +130,8 @@ void FBlendStackAnimPlayer::Initialize(const FAnimationInitializeContext& Contex
 
 	UpdateSourceLinkNode();
 	PoseLinkIndex = InPoseLinkIdx;
+
+	OverrideCurve.Empty();
 }
 
 void FBlendStackAnimPlayer::UpdatePlayRate(float PlayRate)
@@ -226,6 +228,11 @@ void FBlendStackAnimPlayer::Evaluate_AnyThread(FPoseContext& Output)
 	{
 		UpdateSourceLinkNode();
 		MirrorNode.Evaluate_AnyThread(Output);
+
+		if (OverrideCurve.Num() != 0)
+		{
+			UE::Anim::FNamedValueArrayUtils::Union(Output.Curve, OverrideCurve);
+		}
 	}
 	else
 	{
diff --git a/Engine/Plugins/Experimental/Animation/BlendStack/Source/Runtime/Public/BlendStack/AnimNode_BlendStack.h b/Engine/Plugins/Experimental/Animation/BlendStack/Source/Runtime/Public/BlendStack/AnimNode_BlendStack.h
--- a/Engine/Plugins/Experimental/Animation/BlendStack/Source/Runtime/Public/BlendStack/AnimNode_BlendStack.h
+++ b/Engine/Plugins/Experimental/Animation/BlendStack/Source/Runtime/Public/BlendStack/AnimNode_BlendStack.h
@@ -59,6 +59,10 @@ public:
 	// the output FPoseContext will be from StoredPose, StoredCurve, StoredAttributes
 	FCompactHeapPose StoredPose;
 	FBlendedHeapCurve StoredCurve;
+
+	// Curves to add to the pose after the player evaluates
+	TBaseBlendedCurve<FDefaultAllocator, UE::Anim::FCurveElement> OverrideCurve;
+	
 	UE::Anim::FHeapAttributeContainer StoredAttributes;
 	// We need to store the bone container, in case we have a LOD swap during a blend that uses the stored pose.
 	FBoneContainer StoredBoneContainer;
diff --git a/Engine/Plugins/Experimental/Chooser/Source/Chooser/Internal/AnimNode_ChooserPlayer.h b/Engine/Plugins/Experimental/Chooser/Source/Chooser/Internal/AnimNode_ChooserPlayer.h
--- a/Engine/Plugins/Experimental/Chooser/Source/Chooser/Internal/AnimNode_ChooserPlayer.h
+++ b/Engine/Plugins/Experimental/Chooser/Source/Chooser/Internal/AnimNode_ChooserPlayer.h
@@ -9,7 +9,6 @@
 #include "InstancedStruct.h"
 #include "IObjectChooser.h"
 #include "BlendStack/AnimNode_BlendStack.h"
-#include "Styling/SlateBrush.h"
 
 #include "AnimNode_ChooserPlayer.generated.h"
 
@@ -22,53 +21,6 @@ enum class EChooserEvaluationFrequency
 	OnUpdate
 };
 
-/** The random player node holds a list of sequences and parameter ranges which will be played continuously
-  * In a random order. If shuffle mode is enabled then each entry will be played once before repeating any
-  */
-USTRUCT(BlueprintInternalUseOnly)
-struct FChooserPlayerSequenceEntry
-{
-	GENERATED_BODY()
-
-	FChooserPlayerSequenceEntry()
-	    : Sequence(nullptr)
-	    , ChanceToPlay(1.0f)
-	    , MinLoopCount(0)
-	    , MaxLoopCount(0)
-	    , MinPlayRate(1.0f)
-	    , MaxPlayRate(1.0f)
-	{
-	}
-
-	/** Sequence to play when this entry is picked */
-	UPROPERTY(EditAnywhere, BlueprintReadWrite, Category = "Settings", meta = (DisallowedClasses = "/Script/Engine.AnimMontage"))
-	TObjectPtr<UAnimSequenceBase> Sequence;
-
-	/** When not in shuffle mode, this is the chance this entry will play (normalized against all other sample chances) */
-	UPROPERTY(EditAnywhere, BlueprintReadWrite, Category = "Settings", meta = (UIMin = "0", ClampMin = "0"))
-	float ChanceToPlay;
-
-	/** Minimum number of times this entry will loop before ending */
-	UPROPERTY(EditAnywhere, BlueprintReadWrite, Category = "Settings", meta = (UIMin = "0", ClampMin = "0"))
-	int32 MinLoopCount;
-
-	/** Maximum number of times this entry will loop before ending */
-	UPROPERTY(EditAnywhere, BlueprintReadWrite, Category = "Settings", meta = (UIMin = "0", ClampMin = "0"))
-	int32 MaxLoopCount;
-
-	/** Minimum playrate for this entry */
-	UPROPERTY(EditAnywhere, BlueprintReadWrite, Category = "Settings", meta = (UIMin = "0", ClampMin = "0"))
-	float MinPlayRate;
-
-	/** Maximum playrate for this entry */
-	UPROPERTY(EditAnywhere, BlueprintReadWrite, Category = "Settings", meta = (UIMin = "0", ClampMin = "0"))
-	float MaxPlayRate;
-
-	/** Blending properties used when this entry is blending in ontop of another entry */
-	UPROPERTY(EditAnywhere, Category = "Settings")
-	FAlphaBlend BlendIn;
-};
-
 USTRUCT(BlueprintType)
 struct FAnimCurveOverride
 {
@@ -81,14 +33,26 @@ struct FAnimCurveOverride
 	float CurveValue;
 };
 
+template <>
+struct TTypeTraits<FAnimCurveOverride> : public TTypeTraitsBase < FAnimCurveOverride >
+{
+	enum { IsBytewiseComparable = true };
+};
+
 USTRUCT(BlueprintType)
 struct FAnimCurveOverrideList
 {
 	GENERATED_BODY()
 	UPROPERTY(EditAnywhere, BlueprintReadWrite, Category = "Values")
 	TArray<FAnimCurveOverride> Values;
+
+	bool operator != (const FAnimCurveOverrideList& Other) const
+	{
+		return Values != Other.Values;
+	}
 };
 
+
 USTRUCT(BlueprintType)
 struct FChooserPlayerSettings
 {
@@ -148,7 +112,7 @@ public:
 	// if set and bMirrored MirrorDataTable will be used for mirroring the aniamtion
 	UPROPERTY(EditAnywhere, Category = "Settings")
 	TObjectPtr<UMirrorDataTable> MirrorDataTable;
-
+	
 	// requested blend space blend X parameter (if AnimationAsset is a blend space)
 	UPROPERTY(EditAnywhere, BlueprintReadWrite, Category = Settings, meta = (PinHiddenByDefault))
 	float BlendSpaceX = 0;
@@ -190,9 +154,9 @@ private:
 	FChooserEvaluationContext ChooserContext;
 	UAnimationAsset* CurrentAsset = nullptr;
 	float CurrentStartTime = 0;
-	uint32 CurveOverridesIndex = 0;
-	TBaseBlendedCurve<FDefaultAllocator, UE::Anim::FCurveElement> OverrideCurves[2]; 
-	
+	bool CurrentMirror = false;
+	FAnimCurveOverrideList CurrentCurveOverrides;
+
 	// Update Counter for detecting being relevant
 	FGraphTraversalCounter UpdateCounter;
 	
diff --git a/Engine/Plugins/Experimental/Chooser/Source/Chooser/Private/AnimNode_ChooserPlayer.cpp b/Engine/Plugins/Experimental/Chooser/Source/Chooser/Private/AnimNode_ChooserPlayer.cpp
--- a/Engine/Plugins/Experimental/Chooser/Source/Chooser/Private/AnimNode_ChooserPlayer.cpp
+++ b/Engine/Plugins/Experimental/Chooser/Source/Chooser/Private/AnimNode_ChooserPlayer.cpp
@@ -85,21 +85,17 @@ void FAnimNode_ChooserPlayer::UpdateAssetPlayer(const FAnimationUpdateContext& C
 	// Restart the animation:
 	// - if this node just became relevant
 	// - if we chose a new animation
+	// - if the mirror setting has changed
 	// - for playback rate of 0, when the start time changes - for choosing poses as frames of an animation sequence
+	// - if the curve values are different
 	if (bJustBecameRelevant || NewAsset != CurrentAsset ||
-		(CurrentStartTime != Settings.StartTime && Settings.PlaybackRate == 0.0f)) 
-	{			
-		CurveOverridesIndex = (CurveOverridesIndex + 1) % 2;
-		OverrideCurves[CurveOverridesIndex].Empty();
-		if (!Settings.CurveOverrides.Values.IsEmpty())
-		{
-			OverrideCurves[CurveOverridesIndex].Reserve(Settings.CurveOverrides.Values.Num());
-			for (const FAnimCurveOverride& CurveOverride : Settings.CurveOverrides.Values)
-			{
-				OverrideCurves[CurveOverridesIndex].Add(CurveOverride.CurveName, CurveOverride.CurveValue);
-			}
-		}
-
+		CurrentMirror != Settings.bMirror ||
+		(CurrentStartTime != Settings.StartTime && Settings.PlaybackRate == 0.0f) ||
+		CurrentCurveOverrides != Settings.CurveOverrides)
+	{
+		CurrentCurveOverrides = Settings.CurveOverrides;
+		CurrentMirror = Settings.bMirror;
+		
 		float BlendTime = bJustBecameRelevant ? 0 : Settings.BlendTime;
 
 		bool bLoop = Settings.bForceLooping;
@@ -120,6 +116,16 @@ void FAnimNode_ChooserPlayer::UpdateAssetPlayer(const FAnimationUpdateContext& C
 			Settings.BlendTime, Settings.BlendTime, Settings.BlendTime,
 			Settings.BlendProfile, Settings.BlendOption, Settings.bUseInertialBlend, FVector::Zero(), Settings.PlaybackRate,
 			GetGroupName(), GetGroupRole(), GetGroupMethod());
+
+		if (!Settings.CurveOverrides.Values.IsEmpty())
+		{
+			TBaseBlendedCurve<FDefaultAllocator, UE::Anim::FCurveElement>& OverrideCurve = AnimPlayers.First().OverrideCurve;
+			OverrideCurve.Reserve(Settings.CurveOverrides.Values.Num());
+			for (const FAnimCurveOverride& CurveOverride : Settings.CurveOverrides.Values)
+			{
+				OverrideCurve.Add(CurveOverride.CurveName, CurveOverride.CurveValue);
+			}
+		}
 		
 		CurrentAsset = NewAsset;
 		CurrentStartTime = Settings.StartTime;
@@ -161,22 +167,6 @@ void FAnimNode_ChooserPlayer::Evaluate_AnyThread(FPoseContext& Output)
 	FAnimNode_BlendStack_Standalone::Evaluate_AnyThread(Output);
 
 	const FChooserPlayerSettings& Settings = ChooserContext.Params[1].Get<FChooserPlayerSettings>();
-
-
-	
-	if (AnimPlayers.Num() == 1)
-	{
-		UE::Anim::FNamedValueArrayUtils::Union(Output.Curve, OverrideCurves[CurveOverridesIndex]);	
-	}
-	if (AnimPlayers.Num() > 1)
-	{
-		const float Weight = FAlphaBlend::AlphaToBlendOption(AnimPlayers[0].GetBlendInPercentage(), AnimPlayers[0].GetBlendOption());
-		const uint32 PrevCurveOverridesIndex = (CurveOverridesIndex + 1) % 2;
-		UE::Anim::FNamedValueArrayUtils::Union(Output.Curve, OverrideCurves[PrevCurveOverridesIndex]);
-		Output.Curve.LerpTo(OverrideCurves[CurveOverridesIndex], Weight);
-	}
-
-	
 }
 
 FName FAnimNode_ChooserPlayer::GetGroupName() const
