Custom details editor for curve overrides in ChooserPlayerSettings
Optimize check for changed curve values using hash computed at edit time
temporary data conversion codepath for converting to ChooserPlayerSettings from user defined struct

diff --git a/Engine/Plugins/Experimental/Chooser/Source/Chooser/Internal/AnimNode_ChooserPlayer.h b/Engine/Plugins/Experimental/Chooser/Source/Chooser/Internal/AnimNode_ChooserPlayer.h
--- a/Engine/Plugins/Experimental/Chooser/Source/Chooser/Internal/AnimNode_ChooserPlayer.h
+++ b/Engine/Plugins/Experimental/Chooser/Source/Chooser/Internal/AnimNode_ChooserPlayer.h
@@ -43,13 +43,13 @@ USTRUCT(BlueprintType)
 struct FAnimCurveOverrideList
 {
 	GENERATED_BODY()
-	UPROPERTY(EditAnywhere, BlueprintReadWrite, Category = "Values")
+	UPROPERTY(EditAnywhere, BlueprintReadWrite, Category = "Values", meta=(ExpandByDefault=true))
 	TArray<FAnimCurveOverride> Values;
 
-	bool operator != (const FAnimCurveOverrideList& Other) const
-	{
-		return Values != Other.Values;
-	}
+	UPROPERTY()
+	uint32 Hash = 0;
+	
+	CHOOSER_API void ComputeHash();
 };
 
 
@@ -75,7 +75,7 @@ struct FChooserPlayerSettings
     float PlaybackRate = 1.0;
 
 	// List of curve values to set 
-	UPROPERTY(EditAnywhere, BlueprintReadWrite, Category = "Settings")
+	UPROPERTY(EditAnywhere, BlueprintReadWrite, Category = "Settings", meta=(ExpandByDefault=true))
 	FAnimCurveOverrideList CurveOverrides;
 	
 	UPROPERTY(EditAnywhere, BlueprintReadWrite, Category = "Blending")
@@ -155,7 +155,7 @@ private:
 	UAnimationAsset* CurrentAsset = nullptr;
 	float CurrentStartTime = 0;
 	bool CurrentMirror = false;
-	FAnimCurveOverrideList CurrentCurveOverrides;
+	uint32 CurrentCurveOverridesHash = 0;
 
 	// Update Counter for detecting being relevant
 	FGraphTraversalCounter UpdateCounter;
diff --git a/Engine/Plugins/Experimental/Chooser/Source/Chooser/Private/AnimNode_ChooserPlayer.cpp b/Engine/Plugins/Experimental/Chooser/Source/Chooser/Private/AnimNode_ChooserPlayer.cpp
--- a/Engine/Plugins/Experimental/Chooser/Source/Chooser/Private/AnimNode_ChooserPlayer.cpp
+++ b/Engine/Plugins/Experimental/Chooser/Source/Chooser/Private/AnimNode_ChooserPlayer.cpp
@@ -3,16 +3,20 @@
 #include "AnimNode_ChooserPlayer.h"
 
 #include "IObjectChooser.h"
-#include "../../../../Animation/BlendStack/Source/Runtime/Public/BlendStack/AnimNode_BlendStack.h"
+#include "BlendStack/AnimNode_BlendStack.h"
 #include "Animation/AnimInstanceProxy.h"
-#include "Animation/AnimSequence.h"
 #include "Animation/AnimTrace.h"
 #include "Animation/AnimStats.h"
-#include "Animation/AnimSyncScope.h"
 #include "Animation/BlendSpace.h"
+#include "StructUtilsTypes.h"
 
 #include UE_INLINE_GENERATED_CPP_BY_NAME(AnimNode_ChooserPlayer)
 
+void FAnimCurveOverrideList::ComputeHash()
+{
+	Hash = UE::StructUtils::GetStructCrc32(*StaticStruct(), reinterpret_cast<const uint8*>(this));
+}
+
 FAnimNode_ChooserPlayer::FAnimNode_ChooserPlayer()
 {
 }
@@ -91,9 +95,9 @@ void FAnimNode_ChooserPlayer::UpdateAssetPlayer(const FAnimationUpdateContext& C
 	if (bJustBecameRelevant || NewAsset != CurrentAsset ||
 		CurrentMirror != Settings.bMirror ||
 		(CurrentStartTime != Settings.StartTime && Settings.PlaybackRate == 0.0f) ||
-		CurrentCurveOverrides != Settings.CurveOverrides)
+		CurrentCurveOverridesHash != Settings.CurveOverrides.Hash)
 	{
-		CurrentCurveOverrides = Settings.CurveOverrides;
+		CurrentCurveOverridesHash = Settings.CurveOverrides.Hash;
 		CurrentMirror = Settings.bMirror;
 		
 		float BlendTime = bJustBecameRelevant ? 0 : Settings.BlendTime;
diff --git a/Engine/Plugins/Experimental/Chooser/Source/ChooserEditor/ChooserEditor.Build.cs b/Engine/Plugins/Experimental/Chooser/Source/ChooserEditor/ChooserEditor.Build.cs
--- a/Engine/Plugins/Experimental/Chooser/Source/ChooserEditor/ChooserEditor.Build.cs
+++ b/Engine/Plugins/Experimental/Chooser/Source/ChooserEditor/ChooserEditor.Build.cs
@@ -31,6 +31,7 @@ namespace UnrealBuildTool.Rules
 					"GameplayTagsEditor",
 					"StructUtils",
 					"KismetCompiler",
+					"BlendStack",
 					// ... add private dependencies that you statically link with here ...
 				}
 			);
diff --git a/Engine/Plugins/Experimental/Chooser/Source/ChooserEditor/Private/ChooserEditorModule.cpp b/Engine/Plugins/Experimental/Chooser/Source/ChooserEditor/Private/ChooserEditorModule.cpp
--- a/Engine/Plugins/Experimental/Chooser/Source/ChooserEditor/Private/ChooserEditorModule.cpp
+++ b/Engine/Plugins/Experimental/Chooser/Source/ChooserEditor/Private/ChooserEditorModule.cpp
@@ -1,6 +1,8 @@
 // Copyright Epic Games, Inc. All Rights Reserved.
 
 #include "ChooserEditorModule.h"
+
+#include "AnimNode_ChooserPlayer.h"
 #include "IAssetTools.h"
 #include "ChooserTableEditor.h"
 #include "BoolColumnEditor.h"
@@ -14,6 +16,7 @@
 #include "PropertyEditorModule.h"
 #include "RandomizeColumnEditor.h"
 #include "OutputStructColumnEditor.h"
+#include "CurveOverrideCustomization.h"
 
 #define LOCTEXT_NAMESPACE "ChooserEditorModule"
 
@@ -36,6 +39,8 @@ void FModule::StartupModule()
 
 	FPropertyEditorModule& PropertyModule = FModuleManager::LoadModuleChecked<FPropertyEditorModule>("PropertyEditor");
 	
+	PropertyModule.RegisterCustomPropertyTypeLayout(FAnimCurveOverride::StaticStruct()->GetFName(), FOnGetPropertyTypeCustomizationInstance::CreateLambda([] { return MakeShared<FCurveOverrideCustomization>(); }));
+	PropertyModule.RegisterCustomPropertyTypeLayout(FAnimCurveOverrideList::StaticStruct()->GetFName(), FOnGetPropertyTypeCustomizationInstance::CreateLambda([] { return MakeShared<FCurveOverrideListCustomization>(); }));
 	PropertyModule.RegisterCustomPropertyTypeLayout(FChooserPropertyBinding::StaticStruct()->GetFName(), FOnGetPropertyTypeCustomizationInstance::CreateLambda([] { return MakeShared<FPropertyAccessChainCustomization>(); }));
 	PropertyModule.RegisterCustomPropertyTypeLayout(FChooserEnumPropertyBinding::StaticStruct()->GetFName(), FOnGetPropertyTypeCustomizationInstance::CreateLambda([] { return MakeShared<FPropertyAccessChainCustomization>(); }));
 	PropertyModule.RegisterCustomPropertyTypeLayout(FChooserObjectPropertyBinding::StaticStruct()->GetFName(), FOnGetPropertyTypeCustomizationInstance::CreateLambda([] { return MakeShared<FPropertyAccessChainCustomization>(); }));
diff --git a/Engine/Plugins/Experimental/Chooser/Source/ChooserEditor/Private/CurveOverrideCustomization.cpp b/Engine/Plugins/Experimental/Chooser/Source/ChooserEditor/Private/CurveOverrideCustomization.cpp
new file mode 100644
--- /dev/null
+++ b/Engine/Plugins/Experimental/Chooser/Source/ChooserEditor/Private/CurveOverrideCustomization.cpp
@@ -0,0 +1,213 @@
+﻿// Copyright Epic Games, Inc. All Rights Reserved.
+
+#include "CurveOverrideCustomization.h"
+
+#include "AnimNode_ChooserPlayer.h"
+#include "PropertyHandle.h"
+#include "DetailWidgetRow.h"
+#include "SClassViewer.h"
+#include "IDetailChildrenBuilder.h"
+#include "Widgets/Input/SEditableTextBox.h"
+#include "Widgets/Input/SNumericEntryBox.h"
+#include "ScopedTransaction.h"
+#include "Editor/PropertyEditor/Private/PropertyNode.h"
+
+#define LOCTEXT_NAMESPACE "CurveOverrideCustomization"
+
+namespace UE::ChooserEditor
+{
+
+	
+void RebuildHash(TSharedPtr<IPropertyHandle> CurveOverrideListProperty)
+{
+	if (CurveOverrideListProperty)
+	{
+		if (FStructProperty* StructProperty = CastField<FStructProperty>(CurveOverrideListProperty->GetProperty()))
+		{
+			if (StructProperty->Struct == FAnimCurveOverrideList::StaticStruct())
+			{
+				CurveOverrideListProperty->EnumerateRawData([](void* RawData, const int32 DataIndex, const int32 NumDatas)
+				{
+					FAnimCurveOverrideList* AnimCurveOverrideList = static_cast<FAnimCurveOverrideList*>(RawData);
+					AnimCurveOverrideList->ComputeHash();
+					return true;
+				});
+			}
+		}
+		else
+		{
+			if (TSharedPtr<IPropertyHandleStruct> StructPropertyHandle = CurveOverrideListProperty->AsStruct())
+			{
+				if (StructPropertyHandle->GetStructData()->GetStruct() == FAnimCurveOverrideList::StaticStruct())
+				{
+					CurveOverrideListProperty->EnumerateRawData([](void* RawData, const int32 DataIndex, const int32 NumDatas)
+					{
+						FAnimCurveOverrideList* AnimCurveOverrideList = static_cast<FAnimCurveOverrideList*>(RawData);
+						AnimCurveOverrideList->ComputeHash();
+						return true;
+					});
+				}
+			}
+		}
+	}
+}
+	
+	
+void FCurveOverrideListCustomization::CustomizeHeader(TSharedRef<IPropertyHandle> PropertyHandle, FDetailWidgetRow& HeaderRow, IPropertyTypeCustomizationUtils& CustomizationUtils)
+{
+	auto ValuesPropertyHandle = PropertyHandle->GetChildHandle("Values");
+	
+	HeaderRow
+    	.NameContent()
+    	[
+    		PropertyHandle->CreatePropertyNameWidget()
+    	]
+    	.ValueContent()
+		[
+			ValuesPropertyHandle->CreatePropertyValueWidget()
+		];
+}
+
+
+/** 
+ * Node builder for Array children.
+ * Expects Property handle to be an Array
+ */
+class FArrayDataDetails : public IDetailCustomNodeBuilder, public TSharedFromThis<FArrayDataDetails>
+{
+public:
+	FArrayDataDetails(TSharedPtr<IPropertyHandle> InArrayProperty)
+	{
+		ArrayProperty = InArrayProperty;
+	}
+	virtual ~FArrayDataDetails() override { };
+
+	//~ Begin IDetailCustomNodeBuilder interface
+	virtual void SetOnRebuildChildren(FSimpleDelegate InOnRegenerateChildren) override
+	{
+		OnRegenerateChildren = InOnRegenerateChildren;
+	};
+	
+	virtual void GenerateHeaderRowContent(FDetailWidgetRow& NodeRow) override {};
+	virtual void GenerateChildContent(IDetailChildrenBuilder& ChildrenBuilder) override
+	{
+		ArrayProperty->GetNumChildren(CachedNum);
+		for(uint32 i=0;i<CachedNum;i++)
+		{
+			if (auto ChildProperty = ArrayProperty->GetChildHandle(i))
+			{
+				ChildrenBuilder.AddProperty(ChildProperty.ToSharedRef());
+			}
+		}
+	};
+	
+	virtual void Tick(float DeltaTime) override
+	{
+		// If the array num changes regenerate
+		uint32 NumChildren = 0;
+		ArrayProperty->GetNumChildren(NumChildren);
+		if (ArrayProperty && NumChildren != CachedNum)
+		{
+			OnRegenerateChildren.ExecuteIfBound();
+		}	
+	}
+	
+	virtual bool RequiresTick() const override { return true; }
+	virtual bool InitiallyCollapsed() const override { return false; }
+	virtual FName GetName() const override
+	{
+		static const FName Name("InstancedStructDataDetails");
+		return Name;
+	}
+	//~ End IDetailCustomNodeBuilder interface
+private:
+	/** Handle to the array property being edited */
+	TSharedPtr<IPropertyHandle> ArrayProperty;
+
+	/** Delegate that can be used to refresh the child rows of the current struct (eg, when changing struct type) */
+	FSimpleDelegate OnRegenerateChildren;
+
+	uint32 CachedNum = 0;
+};
+
+	
+	
+void FCurveOverrideListCustomization::CustomizeChildren(TSharedRef<IPropertyHandle> PropertyHandle, IDetailChildrenBuilder& ChildBuilder, IPropertyTypeCustomizationUtils& CustomizationUtils)
+{
+	if (auto ValuesPropertyHandle = PropertyHandle->GetChildHandle("Values"))
+	{
+		ChildBuilder.AddCustomBuilder(MakeShared<FArrayDataDetails>(ValuesPropertyHandle));
+	}
+}
+
+void FCurveOverrideCustomization::CustomizeHeader(TSharedRef<IPropertyHandle> PropertyHandle, FDetailWidgetRow& HeaderRow, IPropertyTypeCustomizationUtils& CustomizationUtils)
+{
+	auto NamePropertyHandle = PropertyHandle->GetChildHandle("CurveName");
+	auto ValuePropertyHandle = PropertyHandle->GetChildHandle("CurveValue");
+
+	HeaderRow
+		.NameContent().HAlign(EHorizontalAlignment::HAlign_Right).MinDesiredWidth(300)
+		[
+			SNew(SEditableTextBox)
+			.Text_Lambda([NamePropertyHandle]
+			{
+				FName Name;
+				NamePropertyHandle->GetValue(Name);
+				return FText::FromName(Name);
+			})
+			.OnTextCommitted_Lambda([NamePropertyHandle](FText NewValue, ETextCommit::Type CommitType)
+			{
+				FScopedTransaction ScopedTransaction(LOCTEXT("Set Curve Override Name", "Set Curve Override Name"));
+				NamePropertyHandle->NotifyPreChange();
+				NamePropertyHandle->SetValue(FName(NewValue.ToString()));
+				NamePropertyHandle->NotifyPostChange(EPropertyChangeType::ValueSet);
+
+				if (TSharedPtr<IPropertyHandle> ParentStructProperty = NamePropertyHandle->GetParentHandle())
+				{
+					if (TSharedPtr<IPropertyHandle> ParentArrayProperty = ParentStructProperty->GetParentHandle())
+					{
+						if (TSharedPtr<IPropertyHandle> CurveOverrideListProperty = ParentArrayProperty->GetParentHandle())
+						{
+							RebuildHash(CurveOverrideListProperty);
+						}
+					}
+				}
+			})
+		]
+		.ValueContent()
+		[
+			SNew(SNumericEntryBox<float>)
+			.Value_Lambda([ValuePropertyHandle]()
+			{
+				float Value;
+				ValuePropertyHandle->GetValue(Value);
+				return Value;
+			})
+			.OnValueCommitted_Lambda([ValuePropertyHandle](float NewValue, ETextCommit::Type CommitType)
+			{
+				FScopedTransaction ScopedTransaction(LOCTEXT("Set Curve Override Value", "Set Curve Override Value"));
+				ValuePropertyHandle->NotifyPreChange();
+				ValuePropertyHandle->SetValue(NewValue);
+				ValuePropertyHandle->NotifyPostChange(EPropertyChangeType::ValueSet);
+				
+				if (TSharedPtr<IPropertyHandle> ParentStructProperty = ValuePropertyHandle->GetParentHandle())
+				{
+					if (TSharedPtr<IPropertyHandle> ParentArrayProperty = ParentStructProperty->GetParentHandle())
+					{
+						if (TSharedPtr<IPropertyHandle> CurveOverrideListProperty = ParentArrayProperty->GetParentHandle())
+						{
+							RebuildHash(CurveOverrideListProperty);
+						}
+					}
+				}
+			})
+		];
+}
+	
+void FCurveOverrideCustomization::CustomizeChildren(TSharedRef<IPropertyHandle> PropertyHandle, IDetailChildrenBuilder& ChildBuilder, IPropertyTypeCustomizationUtils& CustomizationUtils)
+{
+}
+
+}
+
+#undef LOCTEXT_NAMESPACE
\ No newline at end of file
diff --git a/Engine/Plugins/Experimental/Chooser/Source/ChooserEditor/Private/CurveOverrideCustomization.h b/Engine/Plugins/Experimental/Chooser/Source/ChooserEditor/Private/CurveOverrideCustomization.h
new file mode 100644
--- /dev/null
+++ b/Engine/Plugins/Experimental/Chooser/Source/ChooserEditor/Private/CurveOverrideCustomization.h
@@ -0,0 +1,27 @@
+﻿// Copyright Epic Games, Inc. All Rights Reserved.
+
+#pragma once
+
+#include "IPropertyTypeCustomization.h"
+
+namespace UE::ChooserEditor
+{
+	
+class FCurveOverrideListCustomization : public IPropertyTypeCustomization
+{
+public:
+	// IPropertyTypeCustomization interface
+	virtual void CustomizeHeader(TSharedRef<IPropertyHandle> PropertyHandle, FDetailWidgetRow& HeaderRow, IPropertyTypeCustomizationUtils& CustomizationUtils) override;
+	virtual void CustomizeChildren(TSharedRef<IPropertyHandle> PropertyHandle, IDetailChildrenBuilder& ChildBuilder, IPropertyTypeCustomizationUtils& CustomizationUtils) override;
+};
+
+
+class FCurveOverrideCustomization : public IPropertyTypeCustomization
+{
+public:
+	// IPropertyTypeCustomization interface
+	virtual void CustomizeHeader(TSharedRef<IPropertyHandle> PropertyHandle, FDetailWidgetRow& HeaderRow, IPropertyTypeCustomizationUtils& CustomizationUtils) override;
+	virtual void CustomizeChildren(TSharedRef<IPropertyHandle> PropertyHandle, IDetailChildrenBuilder& ChildBuilder, IPropertyTypeCustomizationUtils& CustomizationUtils) override;
+};
+
+}
\ No newline at end of file
diff --git a/Engine/Plugins/Experimental/Chooser/Source/ProxyTableEditor/Private/StructOutputDataCustomization.cpp b/Engine/Plugins/Experimental/Chooser/Source/ProxyTableEditor/Private/StructOutputDataCustomization.cpp
--- a/Engine/Plugins/Experimental/Chooser/Source/ProxyTableEditor/Private/StructOutputDataCustomization.cpp
+++ b/Engine/Plugins/Experimental/Chooser/Source/ProxyTableEditor/Private/StructOutputDataCustomization.cpp
@@ -2,6 +2,7 @@
 
 #include "StructOutputDataCustomization.h"
 
+#include "AnimNode_ChooserPlayer.h"
 #include "Chooser.h"
 #include "ChooserPropertyAccess.h"
 #include "PropertyHandle.h"
@@ -124,6 +125,42 @@ void FStructOutputDataCustomization::CustomizeHeader(TSharedRef<IPropertyHandle>
 							// reset the type of the FInstancedStruct storing the value
 							ValueHandle->GetValueData(ValuePtr);
 							FInstancedStruct* ValueStruct = reinterpret_cast<FInstancedStruct*>(ValuePtr);
+
+         	 				// begin Temporary data conversion codepath
+         	 				if (StructPropertyBinding->StructType == FChooserPlayerSettings::StaticStruct() &&
+								(ValueStruct->IsValid() && ValueStruct->GetScriptStruct()->GetName() == "ProxyPropPoseParams"))
+         	 				{
+         	 					const UScriptStruct* Struct = ValueStruct->GetScriptStruct();
+         	 					const void* StructData = ValueStruct->GetMemory();
+
+         	 					FDoubleProperty* FrameProperty = CastField<FDoubleProperty>(Struct->CustomFindProperty("Frame"));
+         	 					double Frame = 0.0;
+         	 					FrameProperty->GetValue_InContainer(StructData, &Frame);
+         	 							
+         	 					FBoolProperty* IsMirroredProperty = CastField<FBoolProperty>(Struct->CustomFindProperty("IsMirrored"));
+         	 					bool bIsMirrored = false;
+         	 					IsMirroredProperty->GetValue_InContainer(StructData, &bIsMirrored);
+         	 						
+         	 					FMapProperty* CurvesProperty = CastField<FMapProperty>(Struct->CustomFindProperty("Curves"));
+         	 					// copy Map
+         	 					TMap<FName,double> Map = *reinterpret_cast<const TMap<FName, double>*>(static_cast<const uint8*>(StructData) + CurvesProperty->GetOffset_ForInternal());
+         	 						
+								ValueStruct->InitializeAs(StructPropertyBinding->StructType);
+
+         	 					FChooserPlayerSettings* ChooserPlayerSettings = reinterpret_cast<FChooserPlayerSettings*>(ValueStruct->GetMutableMemory());
+
+         	 					ChooserPlayerSettings->PlaybackRate = 0.0;
+         	 					ChooserPlayerSettings->StartTime = Frame / 30.0;
+         	 					ChooserPlayerSettings->bMirror = bIsMirrored;
+
+         	 					for (auto& pair : Map)
+         	 					{
+         	 						ChooserPlayerSettings->CurveOverrides.Values.Add(FAnimCurveOverride{pair.Key, (float)pair.Value});
+         	 					}
+         	 						
+         	 				}
+         	 				else
+       	 					// END Temporary data conversion codepath
          	 				if (ValueStruct->GetScriptStruct() != StructPropertyBinding->StructType)
          	 				{
          	 					ValueStruct->InitializeAs(StructPropertyBinding->StructType);
diff --git a/Engine/Plugins/Experimental/Chooser/Source/ProxyTableEditor/ProxyTableEditor.Build.cs b/Engine/Plugins/Experimental/Chooser/Source/ProxyTableEditor/ProxyTableEditor.Build.cs
--- a/Engine/Plugins/Experimental/Chooser/Source/ProxyTableEditor/ProxyTableEditor.Build.cs
+++ b/Engine/Plugins/Experimental/Chooser/Source/ProxyTableEditor/ProxyTableEditor.Build.cs
@@ -32,7 +32,8 @@ namespace UnrealBuildTool.Rules
 					"GameplayTags",
 					"GameplayTagsEditor",
 					"StructUtils",
-					"StructUtilsEditor"
+					"StructUtilsEditor",
+					"BlendStack"
 					// ... add private dependencies that you statically link with here ...
 				}
 			);
