FrameTime details customization for setting StartTime and BlendTime on ChooserPlayerSettings

diff --git a/Engine/Plugins/Experimental/Chooser/Source/Chooser/Internal/AnimNode_ChooserPlayer.h b/Engine/Plugins/Experimental/Chooser/Source/Chooser/Internal/AnimNode_ChooserPlayer.h
--- a/Engine/Plugins/Experimental/Chooser/Source/Chooser/Internal/AnimNode_ChooserPlayer.h
+++ b/Engine/Plugins/Experimental/Chooser/Source/Chooser/Internal/AnimNode_ChooserPlayer.h
@@ -63,7 +63,7 @@ struct FChooserPlayerSettings
 	bool bMirror = false;
 	
 	// Start offset when starting the Animation Asset
-	UPROPERTY(EditAnywhere, BlueprintReadWrite, Category = "Settings")
+	UPROPERTY(EditAnywhere, BlueprintReadWrite, Category = "Settings", meta=(FrameTimeEditor))
     float StartTime = 0;
 
 	// Loop the animation asset, even if the asset is not set as looping
@@ -78,7 +78,7 @@ struct FChooserPlayerSettings
 	UPROPERTY(EditAnywhere, BlueprintReadWrite, Category = "Settings", meta=(ExpandByDefault=true))
 	FAnimCurveOverrideList CurveOverrides;
 	
-	UPROPERTY(EditAnywhere, BlueprintReadWrite, Category = "Blending")
+	UPROPERTY(EditAnywhere, BlueprintReadWrite, Category = "Blending", meta=(FrameTimeEditor))
 	float BlendTime = 0.2;
 	
 	// Set Blend Profiles (editable in the skeleton) to determine how the blending is distributed among your character's bones. It could be used to differentiate between upper body and lower body to blend timing.
diff --git a/Engine/Plugins/Experimental/Chooser/Source/ChooserEditor/Private/ChooserEditorModule.cpp b/Engine/Plugins/Experimental/Chooser/Source/ChooserEditor/Private/ChooserEditorModule.cpp
--- a/Engine/Plugins/Experimental/Chooser/Source/ChooserEditor/Private/ChooserEditorModule.cpp
+++ b/Engine/Plugins/Experimental/Chooser/Source/ChooserEditor/Private/ChooserEditorModule.cpp
@@ -17,11 +17,14 @@
 #include "RandomizeColumnEditor.h"
 #include "OutputStructColumnEditor.h"
 #include "CurveOverrideCustomization.h"
+#include "FrameTimeCustomization.h"
 
 #define LOCTEXT_NAMESPACE "ChooserEditorModule"
 
 namespace UE::ChooserEditor
 {
+	
+
 
 void FModule::StartupModule()
 {
@@ -39,6 +42,8 @@ void FModule::StartupModule()
 
 	FPropertyEditorModule& PropertyModule = FModuleManager::LoadModuleChecked<FPropertyEditorModule>("PropertyEditor");
 	
+	
+	PropertyModule.RegisterCustomPropertyTypeLayout("FloatProperty", FOnGetPropertyTypeCustomizationInstance::CreateLambda([] { return MakeShared<FFrameTimeCustomization>(); }), MakeShared<FFrameTimePropertyTypeIdentifier>());
 	PropertyModule.RegisterCustomPropertyTypeLayout(FAnimCurveOverride::StaticStruct()->GetFName(), FOnGetPropertyTypeCustomizationInstance::CreateLambda([] { return MakeShared<FCurveOverrideCustomization>(); }));
 	PropertyModule.RegisterCustomPropertyTypeLayout(FAnimCurveOverrideList::StaticStruct()->GetFName(), FOnGetPropertyTypeCustomizationInstance::CreateLambda([] { return MakeShared<FCurveOverrideListCustomization>(); }));
 	PropertyModule.RegisterCustomPropertyTypeLayout(FChooserPropertyBinding::StaticStruct()->GetFName(), FOnGetPropertyTypeCustomizationInstance::CreateLambda([] { return MakeShared<FPropertyAccessChainCustomization>(); }));
diff --git a/Engine/Plugins/Experimental/Chooser/Source/ChooserEditor/Private/FrameTimeCustomization.cpp b/Engine/Plugins/Experimental/Chooser/Source/ChooserEditor/Private/FrameTimeCustomization.cpp
new file mode 100644
--- /dev/null
+++ b/Engine/Plugins/Experimental/Chooser/Source/ChooserEditor/Private/FrameTimeCustomization.cpp
@@ -0,0 +1,91 @@
+﻿// Copyright Epic Games, Inc. All Rights Reserved.
+
+#include "FrameTimeCustomization.h"
+
+#include "AnimNode_ChooserPlayer.h"
+#include "AudioDevice.h"
+#include "PropertyHandle.h"
+#include "DetailWidgetRow.h"
+#include "IDetailChildrenBuilder.h"
+#include "IDocumentation.h"
+#include "ScopedTransaction.h"
+#include "SEnumCombo.h"
+#include "Widgets/Input/SNumericEntryBox.h"
+#include "Widgets/Layout/SBox.h"
+
+#define LOCTEXT_NAMESPACE "FrameTimeCustomization"
+
+namespace UE::ChooserEditor
+{
+	
+static const float SecondsToFrames = 30.0f;
+static const float FramesToSeconds = 1.0f/30.0f;
+static ETimeFloatFormat GTimeFloatDisplayFormat = ETimeFloatFormat::Frames;
+	
+void FFrameTimeCustomization::CustomizeHeader(TSharedRef<IPropertyHandle> PropertyHandle, FDetailWidgetRow& HeaderRow, IPropertyTypeCustomizationUtils& CustomizationUtils)
+{
+	UEnum* TimeFloatFormatEnum = StaticEnum<ETimeFloatFormat>();
+	HeaderRow
+    	.NameContent()
+    	[
+    		PropertyHandle->CreatePropertyNameWidget()
+    	]
+    	.ValueContent()
+		[
+			SNew(SHorizontalBox)
+			+SHorizontalBox::Slot()
+			[
+				SNew(SBox).WidthOverride(125).HeightOverride(20)
+             	[
+					SNew(SNumericEntryBox<float>)
+						.MinFractionalDigits(0)
+						.MaxFractionalDigits(3)
+						.Font(FAppStyle::GetFontStyle("PropertyWindow.BoldFont"))
+						.OverrideTextMargin(FMargin(8,0,8,0))
+						.Value_Lambda([PropertyHandle]()
+						{
+							float Value;
+							PropertyHandle->GetValue(Value);
+							return GTimeFloatDisplayFormat == ETimeFloatFormat::Frames ? Value * SecondsToFrames : Value;
+						})
+						.OnValueCommitted_Lambda([PropertyHandle](float NewValue, ETextCommit::Type CommitType)
+						{
+							FScopedTransaction ScopedTransaction(LOCTEXT("Set Frame Time Value", "Set Frame Time Value"));
+							PropertyHandle->NotifyPreChange();
+							PropertyHandle->EnumerateRawData([NewValue](void* RawData,int32 Index, int32 Num)
+							{
+								*static_cast<float*>(RawData) = GTimeFloatDisplayFormat == ETimeFloatFormat::Frames ? NewValue*FramesToSeconds : NewValue;
+								return true;
+							});
+							PropertyHandle->NotifyPostChange(EPropertyChangeType::ValueSet);
+							
+						})
+				]
+			]
+			+SHorizontalBox::Slot().AutoWidth()
+			[
+				SNew(SBox).WidthOverride(80).HeightOverride(20)
+				[
+					SNew(SEnumComboBox, TimeFloatFormatEnum)
+					.Font(FAppStyle::GetFontStyle("PropertyWindow.NormalFont"))
+					.CurrentValue_Lambda([]{return static_cast<int32>(GTimeFloatDisplayFormat);})
+					.ContentPadding(FMargin(0.0f, 0.0f))
+					.ToolTipText(LOCTEXT("FrameTimeToolTip","Select whether to display and edit the time as Seconds, or Frames (at 30 fps)"))
+					.OnEnumSelectionChanged_Lambda([](uint32 NewValue, ESelectInfo::Type)
+					{
+						FSlateApplication::Get().ClearKeyboardFocus(EFocusCause::Cleared);
+						GTimeFloatDisplayFormat = static_cast<ETimeFloatFormat>(NewValue);
+					})
+				]
+			]
+		];
+}
+
+
+void FFrameTimeCustomization::CustomizeChildren(TSharedRef<IPropertyHandle> PropertyHandle, IDetailChildrenBuilder& ChildBuilder, IPropertyTypeCustomizationUtils& CustomizationUtils)
+{
+}
+
+}
+
+#undef LOCTEXT_NAMESPACE
\ No newline at end of file
diff --git a/Engine/Plugins/Experimental/Chooser/Source/ChooserEditor/Private/FrameTimeCustomization.h b/Engine/Plugins/Experimental/Chooser/Source/ChooserEditor/Private/FrameTimeCustomization.h
new file mode 100644
--- /dev/null
+++ b/Engine/Plugins/Experimental/Chooser/Source/ChooserEditor/Private/FrameTimeCustomization.h
@@ -0,0 +1,37 @@
+﻿// Copyright Epic Games, Inc. All Rights Reserved.
+
+#pragma once
+
+#include "IPropertyTypeCustomization.h"
+#include "PropertyEditorModule.h"
+#include "PropertyHandle.h"
+#include "FrameTimeCustomization.generated.h"
+
+UENUM()
+enum class ETimeFloatFormat
+{
+	Frames,
+	Seconds
+};
+
+namespace UE::ChooserEditor
+{
+
+class FFrameTimePropertyTypeIdentifier : public IPropertyTypeIdentifier
+{
+public:
+	virtual bool IsPropertyTypeCustomized(const IPropertyHandle& PropertyHandle) const override
+	{
+		return PropertyHandle.HasMetaData(TEXT("FrameTimeEditor"));
+	}
+};
+
+class FFrameTimeCustomization : public IPropertyTypeCustomization
+{
+public:
+	// IPropertyTypeCustomization interface
+	virtual void CustomizeHeader(TSharedRef<IPropertyHandle> PropertyHandle, FDetailWidgetRow& HeaderRow, IPropertyTypeCustomizationUtils& CustomizationUtils) override;
+	virtual void CustomizeChildren(TSharedRef<IPropertyHandle> PropertyHandle, IDetailChildrenBuilder& ChildBuilder, IPropertyTypeCustomizationUtils& CustomizationUtils) override;
+};
+
+}
\ No newline at end of file
