"Any" setting for Chooser Enum Column

diff --git a/Engine/Plugins/Experimental/Chooser/Source/Chooser/Internal/EnumColumn.h b/Engine/Plugins/Experimental/Chooser/Source/Chooser/Internal/EnumColumn.h
--- a/Engine/Plugins/Experimental/Chooser/Source/Chooser/Internal/EnumColumn.h
+++ b/Engine/Plugins/Experimental/Chooser/Source/Chooser/Internal/EnumColumn.h
@@ -10,6 +10,16 @@
 
 struct FBindingChainElement;
 
+UENUM()
+enum class EEnumColumnCellValueComparison
+{
+	MatchEqual,
+	MatchNotEqual,
+	MatchAny,
+
+	Modulus // used for cycling through the other values
+};
+
 USTRUCT(DisplayName = "Enum Property Binding")
 struct CHOOSER_API FEnumContextProperty : public FChooserParameterEnumBase
 {
@@ -71,8 +81,13 @@ struct FChooserEnumRowData
 {
 	GENERATED_BODY()
 
-	UPROPERTY(EditAnywhere, Category = Runtime)
-	bool CompareNotEqual = false;
+#if WITH_EDITORONLY_DATA
+	UPROPERTY()
+	bool CompareNotEqual_DEPRECATED = false;
+#endif
+	
+	UPROPERTY(EditAnywhere, Category = Runtime, Meta = (ValidEnumValues = "MatchEqual, MatchNotEqual, MatchAny"))
+	EEnumColumnCellValueComparison Comparison = EEnumColumnCellValueComparison::MatchEqual;
 	
 	UPROPERTY(EditAnywhere, Category = Runtime)
 	uint8 Value = 0;
@@ -113,16 +128,8 @@ public:
 	
 	CHOOSER_COLUMN_BOILERPLATE(FChooserParameterEnumBase);
 
-#if WITH_EDITOR
-	virtual void PostLoad() override
-	{
-		Super::PostLoad();
-		
-		if (InputValue.IsValid())
-		{
-			InputValue.GetMutable<FChooserParameterBase>().PostLoad();
-		}
-	}
+#if WITH_EDITORONLY_DATA
+	virtual void PostLoad() override;
 #endif
 };
 
diff --git a/Engine/Plugins/Experimental/Chooser/Source/Chooser/Internal/ObjectColumn.h b/Engine/Plugins/Experimental/Chooser/Source/Chooser/Internal/ObjectColumn.h
--- a/Engine/Plugins/Experimental/Chooser/Source/Chooser/Internal/ObjectColumn.h
+++ b/Engine/Plugins/Experimental/Chooser/Source/Chooser/Internal/ObjectColumn.h
@@ -48,7 +48,7 @@ struct FChooserObjectRowData
 {
 	GENERATED_BODY()
 
-	UPROPERTY(EditAnywhere, Category = "Runtime")
+	UPROPERTY(EditAnywhere, Category = Runtime, Meta = (ValidEnumValues = "MatchEqual, MatchNotEqual, MatchAny"))
 	EObjectColumnCellValueComparison Comparison = EObjectColumnCellValueComparison::MatchEqual;
 
 	UPROPERTY(EditAnywhere, Category = "Runtime")
diff --git a/Engine/Plugins/Experimental/Chooser/Source/Chooser/Private/EnumColumn.cpp b/Engine/Plugins/Experimental/Chooser/Source/Chooser/Private/EnumColumn.cpp
--- a/Engine/Plugins/Experimental/Chooser/Source/Chooser/Private/EnumColumn.cpp
+++ b/Engine/Plugins/Experimental/Chooser/Source/Chooser/Private/EnumColumn.cpp
@@ -2,9 +2,6 @@
 #include "EnumColumn.h"
 #include "ChooserPropertyAccess.h"
 
-#if WITH_EDITOR
-#include "IPropertyAccessEditor.h"
-#endif
 
 bool FEnumContextProperty::GetValue(FChooserEvaluationContext& Context, uint8& OutResult) const
 {
@@ -23,10 +20,45 @@ FEnumColumn::FEnumColumn()
 #endif
 }
 
+#if WITH_EDITORONLY_DATA
+void FEnumColumn::PostLoad()
+
+{
+	Super::PostLoad();
+	
+	if (InputValue.IsValid())
+	{
+		InputValue.GetMutable<FChooserParameterBase>().PostLoad();
+	}
+
+	// upgrade data for "Any" support
+	for(FChooserEnumRowData& CellData : RowValues)
+	{
+		if (CellData.CompareNotEqual_DEPRECATED)
+		{
+			CellData.CompareNotEqual_DEPRECATED = false;
+			CellData.Comparison = EEnumColumnCellValueComparison::MatchNotEqual;
+		}
+	}
+}
+#endif
+
 bool FChooserEnumRowData::Evaluate(const uint8 LeftHandSide) const
 {
-	bool Equal = LeftHandSide == Value;
-	return Equal ^ CompareNotEqual;
+	switch (Comparison)
+	{
+		case EEnumColumnCellValueComparison::MatchEqual:
+			return LeftHandSide == Value;
+
+		case EEnumColumnCellValueComparison::MatchNotEqual:
+			return LeftHandSide != Value;
+
+		case EEnumColumnCellValueComparison::MatchAny:
+			return true;
+
+		default:
+			return false;
+	}
 }
 
 void FEnumColumn::Filter(FChooserEvaluationContext& Context, const TArray<uint32>& IndexListIn, TArray<uint32>& IndexListOut) const
diff --git a/Engine/Plugins/Experimental/Chooser/Source/ChooserEditor/Private/EnumColumnEditor.cpp b/Engine/Plugins/Experimental/Chooser/Source/ChooserEditor/Private/EnumColumnEditor.cpp
--- a/Engine/Plugins/Experimental/Chooser/Source/ChooserEditor/Private/EnumColumnEditor.cpp
+++ b/Engine/Plugins/Experimental/Chooser/Source/ChooserEditor/Private/EnumColumnEditor.cpp
@@ -191,13 +191,24 @@ TSharedRef<SWidget> CreateEnumColumnWidget(UChooserTable* Chooser, FChooserColum
 	return SNew(SHorizontalBox)
     		+ SHorizontalBox::Slot().AutoWidth()
     		[
-    			SNew(SBox).WidthOverride(Row < 0 ? 0 : 45)
+    			SNew(SBox).WidthOverride(Row < 0 ? 0 : 55)
     			[
     				SNew(SButton).ButtonStyle(FAppStyle::Get(),"FlatButton").TextStyle(FAppStyle::Get(),"RichTextBlock.Bold").HAlign(HAlign_Center)
     				.Visibility(Row < 0 ? EVisibility::Hidden : EVisibility::Visible)
 					.Text_Lambda([EnumColumn, Row]()
 					{
-						return (EnumColumn->RowValues.IsValidIndex(Row) && EnumColumn->RowValues[Row].CompareNotEqual ? LOCTEXT("Not Equal", "!=") : LOCTEXT("Equal", "="));
+						switch (EnumColumn->RowValues[Row].Comparison)
+						{
+						case EEnumColumnCellValueComparison::MatchEqual:
+							return LOCTEXT("CompEqual", "=");
+
+						case EEnumColumnCellValueComparison::MatchNotEqual:
+							return LOCTEXT("CompNotEqual", "!=");
+
+						case EEnumColumnCellValueComparison::MatchAny:
+							return LOCTEXT("CompAny", "Any");
+						}
+						return FText::GetEmpty();
 					})
 					.OnClicked_Lambda([EnumColumn, Chooser, Row]()
 					{
@@ -205,7 +216,10 @@ TSharedRef<SWidget> CreateEnumColumnWidget(UChooserTable* Chooser, FChooserColum
 						{
 							const FScopedTransaction Transaction(LOCTEXT("Edit Comparison", "Edit Comparison Operation"));
 							Chooser->Modify(true);
-							EnumColumn->RowValues[Row].CompareNotEqual = !EnumColumn->RowValues[Row].CompareNotEqual;
+							// cycle through comparison options
+							EEnumColumnCellValueComparison& Comparison = EnumColumn->RowValues[Row].Comparison;
+							const int32 NextComparison = (static_cast<int32>(Comparison) + 1) % static_cast<int32>(EEnumColumnCellValueComparison::Modulus);
+							Comparison = static_cast<EEnumColumnCellValueComparison>(NextComparison);
 						}
 						return FReply::Handled();
 					})
@@ -214,6 +228,13 @@ TSharedRef<SWidget> CreateEnumColumnWidget(UChooserTable* Chooser, FChooserColum
 			+ SHorizontalBox::Slot().FillWidth(1)
 			[
 				SNew(SEnumCell<FEnumColumn>).TransactionObject(Chooser).EnumColumn(EnumColumn).RowIndex(Row)
+				.Visibility_Lambda([EnumColumn,Column, Row]()
+					{
+						return (EnumColumn->RowValues.IsValidIndex(Row) &&
+								EnumColumn->RowValues[Row].Comparison == EEnumColumnCellValueComparison::MatchAny)
+								   ? EVisibility::Collapsed
+								   : EVisibility::Visible;
+					})
 			];
 }
 
