Fix crash when changing struct binding in Proxy Table Struct outputs

diff --git a/Engine/Plugins/Experimental/Chooser/Source/ProxyTableEditor/Private/StructOutputDataCustomization.cpp b/Engine/Plugins/Experimental/Chooser/Source/ProxyTableEditor/Private/StructOutputDataCustomization.cpp
--- a/Engine/Plugins/Experimental/Chooser/Source/ProxyTableEditor/Private/StructOutputDataCustomization.cpp
+++ b/Engine/Plugins/Experimental/Chooser/Source/ProxyTableEditor/Private/StructOutputDataCustomization.cpp
@@ -15,6 +15,7 @@
 #include "SPropertyAccessChainWidget.h"
 #include "InstancedStructDetails.h"
 #include "ProxyTable.h"
+#include "IPropertyUtilities.h"
 
 #define LOCTEXT_NAMESPACE "StructOutputCustomization"
 
@@ -29,6 +30,9 @@ void FStructOutputDataCustomization::CustomizeHeader(TSharedRef<IPropertyHandle>
 
 	static FName ProxyPropertyName = "Proxy";
 
+	
+	TSharedPtr<IPropertyUtilities> PropUtils = CustomizationUtils.GetPropertyUtilities();
+         	 					
 	IHasContextClass* HasContext = nullptr;
 	// this details customization is hard coded to work in a ProxyTable where the StructOutputs array will be the parent
 	if (TSharedPtr<IPropertyHandle> ArrayHandle = PropertyHandle->GetParentHandle())
@@ -78,7 +82,7 @@ void FStructOutputDataCustomization::CustomizeHeader(TSharedRef<IPropertyHandle>
 			BindingHandle->GetValueData(data);
 			return reinterpret_cast<FChooserPropertyBinding*>(data);
 		})
-		.OnAddBinding_Lambda([HasContext, ValueHandle, BindingHandle](FName InPropertyName, const TArray<FBindingChainElement>& InBindingChain)
+		.OnAddBinding_Lambda([HasContext, ValueHandle, BindingHandle, PropUtils](FName InPropertyName, const TArray<FBindingChainElement>& InBindingChain)
          	 	{
          	 		TArray<UObject*> OuterObjects;
          	 		BindingHandle->GetOuterObjects(OuterObjects);
@@ -164,6 +168,11 @@ void FStructOutputDataCustomization::CustomizeHeader(TSharedRef<IPropertyHandle>
          	 				if (ValueStruct->GetScriptStruct() != StructPropertyBinding->StructType)
          	 				{
          	 					ValueStruct->InitializeAs(StructPropertyBinding->StructType);
+								if (PropUtils.IsValid())
+								{
+									PropUtils->ForceRefresh();
+								}
+         	 					
          	 				}
 	
          	 				BindingHandle->NotifyPostChange(EPropertyChangeType::ValueSet);
