chooser - adding IPoseSearchProvider support

diff --git a/Engine/Plugins/Experimental/Chooser/Source/Chooser/Internal/AnimNode_ChooserPlayer.h b/Engine/Plugins/Experimental/Chooser/Source/Chooser/Internal/AnimNode_ChooserPlayer.h
--- a/Engine/Plugins/Experimental/Chooser/Source/Chooser/Internal/AnimNode_ChooserPlayer.h
+++ b/Engine/Plugins/Experimental/Chooser/Source/Chooser/Internal/AnimNode_ChooserPlayer.h
@@ -131,6 +131,10 @@ public:
 	UPROPERTY()
 	TArray<FInstancedStruct> ChooserContextDefinition;
 
+	// Use pose matching to choose the start position. Requires experimental PoseSearch plugin.
+	UPROPERTY(EditAnywhere, Category = PoseMatching, meta = (PinHiddenByDefault))
+	bool bStartFromMatchingPose = false;
+
 	// FAnimNode_Base interface
 	CHOOSER_API virtual void Initialize_AnyThread(const FAnimationInitializeContext& Context) override;
 	CHOOSER_API virtual void Evaluate_AnyThread(FPoseContext& Output) override;
diff --git a/Engine/Plugins/Experimental/Chooser/Source/Chooser/Private/AnimNode_ChooserPlayer.cpp b/Engine/Plugins/Experimental/Chooser/Source/Chooser/Private/AnimNode_ChooserPlayer.cpp
--- a/Engine/Plugins/Experimental/Chooser/Source/Chooser/Private/AnimNode_ChooserPlayer.cpp
+++ b/Engine/Plugins/Experimental/Chooser/Source/Chooser/Private/AnimNode_ChooserPlayer.cpp
@@ -2,12 +2,13 @@
 
 #include "AnimNode_ChooserPlayer.h"
 
-#include "IObjectChooser.h"
-#include "BlendStack/AnimNode_BlendStack.h"
 #include "Animation/AnimInstanceProxy.h"
-#include "Animation/AnimTrace.h"
+#include "Animation/AnimPoseSearchProvider.h"
 #include "Animation/AnimStats.h"
+#include "Animation/AnimTrace.h"
 #include "Animation/BlendSpace.h"
+#include "BlendStack/AnimNode_BlendStack.h"
+#include "IObjectChooser.h"
 #include "StructUtilsTypes.h"
 
 #include UE_INLINE_GENERATED_CPP_BY_NAME(AnimNode_ChooserPlayer)
@@ -27,10 +28,34 @@ UAnimationAsset* FAnimNode_ChooserPlayer::ChooseAsset(const FAnimationUpdateCont
 	{
 		// reset settings to default
 		FChooserPlayerSettings& Settings = ChooserContext.Params[1].GetMutable<FChooserPlayerSettings>();
+		
 		Settings = DefaultSettings;
-		UAnimationAsset* Result = Cast<UAnimationAsset>(Chooser.Get<FObjectChooserBase>().ChooseObject(ChooserContext));
-	
-		return Result;
+		
+		const FObjectChooserBase& ChooserBase = Chooser.Get<FObjectChooserBase>();
+		if (bStartFromMatchingPose)
+		{
+			if (UE::Anim::IPoseSearchProvider* PoseSearchProvider = UE::Anim::IPoseSearchProvider::Get())
+			{
+				TArray<UAnimationAsset*, TInlineAllocator<128>> AnimationAssets;
+				ChooserBase.ChooseMulti(ChooserContext, FObjectChooserBase::FObjectChooserIteratorCallback::CreateLambda([&AnimationAssets](UObject* InResult)
+					{
+						if (UAnimationAsset* AnimationAsset = Cast<UAnimationAsset>(InResult))
+						{
+							AnimationAssets.Add(AnimationAsset);
+						}
+						return FObjectChooserBase::EIteratorStatus::Continue;
+					}));
+
+				const UE::Anim::IPoseSearchProvider::FSearchResult SearchResult = PoseSearchProvider->Search(Context, AnimationAssets);
+				if (SearchResult.AnimationAsset)
+				{
+					Settings.StartTime = SearchResult.TimeOffsetSeconds;
+					return SearchResult.AnimationAsset;
+				}
+			}
+		}
+		
+		return Cast<UAnimationAsset>(ChooserBase.ChooseObject(ChooserContext));
 	}
 	return nullptr;
 }
diff --git a/Engine/Plugins/Experimental/Chooser/Source/ChooserUncooked/Private/AnimGraphNode_ChooserPlayer.cpp b/Engine/Plugins/Experimental/Chooser/Source/ChooserUncooked/Private/AnimGraphNode_ChooserPlayer.cpp
--- a/Engine/Plugins/Experimental/Chooser/Source/ChooserUncooked/Private/AnimGraphNode_ChooserPlayer.cpp
+++ b/Engine/Plugins/Experimental/Chooser/Source/ChooserUncooked/Private/AnimGraphNode_ChooserPlayer.cpp
@@ -5,7 +5,9 @@
 #include <chrono>
 
 #include "EditorCategoryUtils.h"
+#include "DetailLayoutBuilder.h"
 #include "Animation/AnimAttributes.h"
+#include "Animation/AnimPoseSearchProvider.h"
 #include "Animation/AnimRootMotionProvider.h"
 #include "ChooserPropertyAccess.h"
 #include "Elements/Interfaces/TypedElementDataStorageCompatibilityInterface.h"
@@ -54,6 +56,16 @@ FText UAnimGraphNode_ChooserPlayer::GetMenuCategory() const
 	return LOCTEXT("NodeCategory", "Animation|Sequences");
 }
 
+void UAnimGraphNode_ChooserPlayer::CustomizeDetails(IDetailLayoutBuilder& DetailBuilder)
+{
+	Super::CustomizeDetails(DetailBuilder);
+
+	if (!UE::Anim::IPoseSearchProvider::IsAvailable())
+	{
+		DetailBuilder.HideCategory(TEXT("PoseMatching"));
+	}
+}
+
 void UAnimGraphNode_ChooserPlayer::GetOutputLinkAttributes(FNodeAttributeArray& OutAttributes) const
 {
 	OutAttributes.Add(UE::Anim::FAttributes::Curves);
diff --git a/Engine/Plugins/Experimental/Chooser/Source/ChooserUncooked/Public/AnimGraphNode_ChooserPlayer.h b/Engine/Plugins/Experimental/Chooser/Source/ChooserUncooked/Public/AnimGraphNode_ChooserPlayer.h
--- a/Engine/Plugins/Experimental/Chooser/Source/ChooserUncooked/Public/AnimGraphNode_ChooserPlayer.h
+++ b/Engine/Plugins/Experimental/Chooser/Source/ChooserUncooked/Public/AnimGraphNode_ChooserPlayer.h
@@ -35,6 +35,7 @@ public:
 	virtual FAnimNode_BlendStack_Standalone* GetBlendStackNode() const override { return (FAnimNode_BlendStack_Standalone*)(&Node); }
 
 	// UAnimGraphNode_Base interface
+	virtual void CustomizeDetails(IDetailLayoutBuilder& DetailBuilder) override;
 	virtual void GetOutputLinkAttributes(FNodeAttributeArray& OutAttributes) const override;
 	// End of UAnimGraphNode_Base interface
 
