Export Proxy table API
Add helper functions to add parameters to FChooserEvalutionContext

diff --git a/Engine/Plugins/Experimental/Chooser/Source/Chooser/Public/IObjectChooser.h b/Engine/Plugins/Experimental/Chooser/Source/Chooser/Public/IObjectChooser.h
--- a/Engine/Plugins/Experimental/Chooser/Source/Chooser/Public/IObjectChooser.h
+++ b/Engine/Plugins/Experimental/Chooser/Source/Chooser/Public/IObjectChooser.h
@@ -39,6 +39,18 @@ struct CHOOSER_API FChooserEvaluationInputObject
 USTRUCT(BlueprintType)
 struct CHOOSER_API FChooserEvaluationContext
 {
+	FChooserEvaluationContext() {}
+	FChooserEvaluationContext(UObject* ContextObject)
+	{
+		AddParam(ContextObject);
+	}
+	void AddParam(UObject* ContextObject)
+	{
+		Params.AddDefaulted();
+		Params.Last().InitializeAs(FChooserEvaluationInputObject::StaticStruct());
+		Params.Last().GetMutable<FChooserEvaluationInputObject>().Object = ContextObject;
+	}
+
 	#if WITH_EDITOR
     	FChooserDebuggingInfo DebuggingInfo;
     #endif
diff --git a/Engine/Plugins/Experimental/Chooser/Source/ProxyTable/Public/ProxyTable.h b/Engine/Plugins/Experimental/Chooser/Source/ProxyTable/Public/ProxyTable.h
--- a/Engine/Plugins/Experimental/Chooser/Source/ProxyTable/Public/ProxyTable.h
+++ b/Engine/Plugins/Experimental/Chooser/Source/ProxyTable/Public/ProxyTable.h
@@ -63,8 +63,8 @@ struct PROXYTABLE_API FRuntimeProxyValue
 	TArray<FProxyStructOutput> OutputStructData;
 };
 
-UCLASS(MinimalAPI,BlueprintType)
-class UProxyTable : public UObject
+UCLASS(BlueprintType)
+class PROXYTABLE_API UProxyTable : public UObject
 {
 	GENERATED_UCLASS_BODY()
 public:
