Chooser Property Access - make editor use non-compiled codepath by default

diff --git a/Engine/Plugins/Experimental/Chooser/Source/Chooser/Internal/ObjectColumn.h b/Engine/Plugins/Experimental/Chooser/Source/Chooser/Internal/ObjectColumn.h
--- a/Engine/Plugins/Experimental/Chooser/Source/Chooser/Internal/ObjectColumn.h
+++ b/Engine/Plugins/Experimental/Chooser/Source/Chooser/Internal/ObjectColumn.h
@@ -19,6 +19,11 @@ struct CHOOSER_API FObjectContextProperty : public FChooserParameterObjectBase
 	FChooserObjectPropertyBinding Binding;
 
 	virtual bool GetValue(FChooserEvaluationContext& Context, FSoftObjectPath& OutResult) const override;
+	
+	virtual void Compile(IHasContextClass* Owner, bool bForce) override
+	{
+		Binding.Compile(Owner, bForce);
+	};
 
 #if WITH_EDITOR
 	virtual void GetDisplayName(FText& OutName) const override
diff --git a/Engine/Plugins/Experimental/Chooser/Source/Chooser/Private/ChooserPropertyAccess.cpp b/Engine/Plugins/Experimental/Chooser/Source/Chooser/Private/ChooserPropertyAccess.cpp
--- a/Engine/Plugins/Experimental/Chooser/Source/Chooser/Private/ChooserPropertyAccess.cpp
+++ b/Engine/Plugins/Experimental/Chooser/Source/Chooser/Private/ChooserPropertyAccess.cpp
@@ -4,6 +4,7 @@
 
 #include "IObjectChooser.h"
 #include "Engine/UserDefinedStruct.h"
+#include "HAL/IConsoleManager.h"
 #include "UObject/UnrealType.h"
 #include "Logging/LogMacros.h"
 
@@ -13,6 +14,12 @@
 #endif
 
 #define LOCTEXT_NAMESPACE "ChooserPropertyAccess"
+
+TAutoConsoleVariable<bool> CVarUseCompiledPropertyChainsInEditor(
+	TEXT("Choosers.UseCompiledPropertyChainsInEditor"),
+	false,
+	TEXT("Enable optimized property access on Choosers and Proxy Tables in Editor. \n0: Disable (default), 1: Enable"),
+	ECVF_Default);
 		
 FName FChooserPropertyBinding::GetUniqueId() const
 {
@@ -220,19 +227,23 @@ void FChooserPropertyBinding::Compile(IHasContextClass* Owner, bool bForce)
 			if (const FBoolProperty* BoolProperty = CastField<FBoolProperty>(BaseProperty))
 			{
 				OutCompiledBinding.CompiledChain.Last().Mask = BoolProperty->GetFieldMask();
-				OutCompiledBinding.PropertyType = UE::Chooser::EPropertyNumericalType::BOOL;
+				OutCompiledBinding.PropertyType = UE::Chooser::EChooserPropertyAccessType::Bool;
 			}
 			if (BaseProperty->IsA<FFloatProperty>())
 			{
-				OutCompiledBinding.PropertyType = UE::Chooser::EPropertyNumericalType::FLOAT;
+				OutCompiledBinding.PropertyType = UE::Chooser::EChooserPropertyAccessType::Float;
 			}
 			else if (BaseProperty->IsA<FDoubleProperty>())
 			{
-				OutCompiledBinding.PropertyType = UE::Chooser::EPropertyNumericalType::DOUBLE;
+				OutCompiledBinding.PropertyType = UE::Chooser::EChooserPropertyAccessType::Double;
 			}
 			else if (BaseProperty->IsA<FIntProperty>())
 			{
-				OutCompiledBinding.PropertyType = UE::Chooser::EPropertyNumericalType::INT32;
+				OutCompiledBinding.PropertyType = UE::Chooser::EChooserPropertyAccessType::Int32;
+			}
+			else if (BaseProperty->IsA<FSoftObjectProperty>())
+			{
+				OutCompiledBinding.PropertyType = UE::Chooser::EChooserPropertyAccessType::SoftObjectRef;
 			}
 		}
 		else
@@ -248,15 +259,15 @@ void FChooserPropertyBinding::Compile(IHasContextClass* Owner, bool bForce)
 					OutCompiledBinding.CompiledChain.Add(UE::Chooser::FCompiledBindingElement(Function));
 					if (ReturnProperty->IsA<FFloatProperty>())
 					{
-						OutCompiledBinding.PropertyType = UE::Chooser::EPropertyNumericalType::FLOAT;
+						OutCompiledBinding.PropertyType = UE::Chooser::EChooserPropertyAccessType::Float;
 					}
 					else if (ReturnProperty->IsA<FDoubleProperty>())
 					{
-						OutCompiledBinding.PropertyType = UE::Chooser::EPropertyNumericalType::DOUBLE;
+						OutCompiledBinding.PropertyType = UE::Chooser::EChooserPropertyAccessType::Double;
 					}
 					else if (ReturnProperty->IsA<FIntProperty>())
 					{
-						OutCompiledBinding.PropertyType = UE::Chooser::EPropertyNumericalType::INT32;
+						OutCompiledBinding.PropertyType = UE::Chooser::EChooserPropertyAccessType::Int32;
 					}
 				}
 			}
@@ -285,21 +296,28 @@ void FChooserPropertyBinding::Compile(IHasContextClass* Owner, bool bForce)
 namespace UE::Chooser
 {
 
-	uint8* ResolveCompiledPropertyChain(FChooserEvaluationContext& Context, const FCompiledBinding& CompiledBinding)
+	bool ResolveCompiledPropertyChain(FChooserEvaluationContext& Context, const FChooserPropertyBinding& Binding, FResolvedPropertyChainResult& Result)
 	{
-		uint8* Result = nullptr;
 		const UStruct* InputType = nullptr;
+		uint8* Container = nullptr;
+		
+		if (!Binding.CompiledBinding.IsValid())
+		{
+			return false;
+		}
+		
+		const FCompiledBinding& CompiledBinding = *Binding.CompiledBinding.Get();
 
 		if(!Context.Params.IsValidIndex(CompiledBinding.ContextIndex))
 		{
 			UE_LOG(LogChooser, Warning, TEXT("Invalid Index {%d} while resolving compiled property chain."), CompiledBinding.ContextIndex);
-			return nullptr;
+			return false;
 		}
 		
 		if (FChooserEvaluationInputObject* ObjectInput = Context.Params[CompiledBinding.ContextIndex].GetMutablePtr<FChooserEvaluationInputObject>())
 		{
 			UObject* Object = ObjectInput->Object.Get();
-			Result = reinterpret_cast<uint8*>(Object);
+			Container = reinterpret_cast<uint8*>(Object);
 			if (Object)
 			{
 				InputType = Object->GetClass();
@@ -307,102 +325,88 @@ namespace UE::Chooser
 		}
 		else
 		{
-			Result = Context.Params[CompiledBinding.ContextIndex].GetMutableMemory();
+			Container = Context.Params[CompiledBinding.ContextIndex].GetMutableMemory();
 			InputType = Context.Params[CompiledBinding.ContextIndex].GetScriptStruct();
 		}
 
-		if (Result == nullptr || InputType == nullptr)
+		if (Container == nullptr || InputType == nullptr)
 		{
-			return nullptr;
+			return false;
 		}
 
 		if (!InputType->IsChildOf(CompiledBinding.TargetType))
 		{
 			UE_LOG(LogChooser, Warning, TEXT("Property Binding compiled for type: {%s} is being evaluated on incompatible type: {%s}."), ToCStr(CompiledBinding.TargetType->GetName()), ToCStr(InputType->GetName()));
-			return nullptr;
+			return false;
 		}
 		
 		for (int i = 0; i<CompiledBinding.CompiledChain.Num() - 1; i++)
 		{
-			if (Result)
+			if (Container)
 			{
 				const FCompiledBindingElement& Element = CompiledBinding.CompiledChain[i];
 				if (!Element.bIsFunction)
 				{
-					Result = *reinterpret_cast<uint8**>(Result + CompiledBinding.CompiledChain[i].Offset);
+					Container = *reinterpret_cast<uint8**>(Container + CompiledBinding.CompiledChain[i].Offset);
 				}
 				else
 				{
-					UObject* Object = reinterpret_cast<UObject*>(Result);
+					UObject* Object = reinterpret_cast<UObject*>(Container);
 					if (Element.Function->IsNative())
 					{
 						FFrame Stack(Object, Element.Function, nullptr, nullptr, Element.Function->ChildProperties);
-						Element.Function->Invoke(Object, Stack, &Result);
+						Element.Function->Invoke(Object, Stack, &Container);
 					}
 					else
 					{
-						Object->ProcessEvent(Element.Function, &Result);
+						Object->ProcessEvent(Element.Function, &Container);
 					}
 				}
 			}
 		}
 
-		return Result;
+		Result.Container = Container;
+		const FCompiledBindingElement& Last = CompiledBinding.CompiledChain.Last();
+		if (Last.bIsFunction)
+		{
+			Result.Function = Last.Function;
+		}
+		else
+		{
+			Result.PropertyOffset = Last.Offset;
+			Result.Mask = Last.Mask;
+		}
+		Result.PropertyType = CompiledBinding.PropertyType;
+		return true;
 	}
 	
-	bool ResolvePropertyChain(FChooserEvaluationContext& Context, const FChooserPropertyBinding& PropertyBinding, const void*& OutContainer, const UStruct*& OutStructType)
+	bool ResolvePropertyChain(uint8* Container, const UStruct*& StructType, const FChooserPropertyBinding& PropertyBinding, FResolvedPropertyChainResult& Result)
 	{
-		if (Context.Params.IsValidIndex(PropertyBinding.ContextIndex))
+		if (PropertyBinding.PropertyBindingChain.Num() == 0)
 		{
-			if (FChooserEvaluationInputObject* ObjectParam = Context.Params[PropertyBinding.ContextIndex].GetMutablePtr<FChooserEvaluationInputObject>())
+			if (PropertyBinding.IsBoundToRoot)
 			{
-				OutContainer = ObjectParam->Object;
-				if (OutContainer)
-				{
-					OutStructType = ObjectParam->Object->GetClass();
-				}
-				else
-				{
-					OutStructType = nullptr;
-				}
+				Result.Container = Container;
+				return true;
 			}
 			else
-			{
-				OutContainer = Context.Params[PropertyBinding.ContextIndex].GetMutableMemory();
-				OutStructType = Context.Params[PropertyBinding.ContextIndex].GetScriptStruct();
-			}
-			
-
-			if (OutContainer == nullptr || OutStructType == nullptr)
 			{
 				return false;
 			}
-
-			return ResolvePropertyChain(OutContainer, OutStructType, PropertyBinding.PropertyBindingChain);
-		}
-
-		return false;
-	}
-	
-	bool ResolvePropertyChain(const void*& Container, const UStruct*& StructType, const TArray<FName>& PropertyBindingChain)
-	{
-		if (PropertyBindingChain.Num() == 0)
-		{
-			return false;
 		}
 	
-		const int PropertyChainLength = PropertyBindingChain.Num();
+		const int PropertyChainLength = PropertyBinding.PropertyBindingChain.Num();
 		for(int PropertyChainIndex = 0; PropertyChainIndex < PropertyChainLength - 1; PropertyChainIndex++)
 		{
-			if (const FStructProperty* StructProperty = FindFProperty<FStructProperty>(StructType, PropertyBindingChain[PropertyChainIndex]))
+			if (const FStructProperty* StructProperty = FindFProperty<FStructProperty>(StructType, PropertyBinding.PropertyBindingChain[PropertyChainIndex]))
 			{
 				StructType = StructProperty->Struct;
-				Container = StructProperty->ContainerPtrToValuePtr<void>(Container);
+				Container = StructProperty->ContainerPtrToValuePtr<uint8>(Container);
 			}
-			else if (const FObjectProperty* ObjectProperty = FindFProperty<FObjectProperty>(StructType, PropertyBindingChain[PropertyChainIndex]))
+			else if (const FObjectProperty* ObjectProperty = FindFProperty<FObjectProperty>(StructType, PropertyBinding.PropertyBindingChain[PropertyChainIndex]))
 			{
 				StructType = ObjectProperty->PropertyClass;
-				Container = *ObjectProperty->ContainerPtrToValuePtr<TObjectPtr<UObject>>(Container);
+				Container = static_cast<uint8*>(*ObjectProperty->ContainerPtrToValuePtr<TObjectPtr<UObject>>(Container));
 				if (Container == nullptr)
 				{
 					return false;
@@ -413,9 +417,9 @@ namespace UE::Chooser
 				// check if it's a member function
 				if (const UClass* ClassType = Cast<const UClass>(StructType))
 				{
-					if (UFunction* Function = ClassType->FindFunctionByName(PropertyBindingChain[PropertyChainIndex]))
+					if (UFunction* Function = ClassType->FindFunctionByName(PropertyBinding.PropertyBindingChain[PropertyChainIndex]))
 					{
-						UObject* Object = reinterpret_cast<UObject*>(const_cast<void*>(Container));
+						UObject* Object = reinterpret_cast<UObject*>(Container);
 						if (Function->IsNative())
 						{
 							FFrame Stack(Object, Function, nullptr, nullptr, Function->ChildProperties);
@@ -432,7 +436,7 @@ namespace UE::Chooser
 						}
 						else
 						{
-							StructType = reinterpret_cast<UObject*>(const_cast<void*>(Container))->GetClass();
+							StructType = reinterpret_cast<UObject*>(Container)->GetClass();
 						}
 					}
 					else
@@ -446,9 +450,117 @@ namespace UE::Chooser
 				}
 			}
 		}
+
+
+		bool bFound = false;
+
+		if (const FProperty* BaseProperty = FindFProperty<FProperty>(StructType, PropertyBinding.PropertyBindingChain.Last()))
+		{
+			bFound = true;
+
+			Result.Container = Container;
+			Result.PropertyOffset = BaseProperty->GetOffset_ForInternal();
+			
+			if (const FBoolProperty* BoolProperty = CastField<FBoolProperty>(BaseProperty))
+			{
+				Result.Mask = BoolProperty->GetFieldMask();
+				Result.PropertyType = UE::Chooser::EChooserPropertyAccessType::Bool;
+			}
+			if (BaseProperty->IsA<FFloatProperty>())
+			{
+				Result.PropertyType = UE::Chooser::EChooserPropertyAccessType::Float;
+			}
+			else if (BaseProperty->IsA<FDoubleProperty>())
+			{
+				Result.PropertyType = UE::Chooser::EChooserPropertyAccessType::Double;
+			}
+			else if (BaseProperty->IsA<FIntProperty>())
+			{
+				Result.PropertyType = UE::Chooser::EChooserPropertyAccessType::Int32;
+			}
+			else if (BaseProperty->IsA<FSoftObjectProperty>())
+			{
+				Result.PropertyType = UE::Chooser::EChooserPropertyAccessType::SoftObjectRef;
+			}
+		}
+		else
+		{
+			// handle function calls 
+			if (const UClass* ClassType = Cast<const UClass>(StructType))
+			{
+				if (UFunction* Function = ClassType->FindFunctionByName(PropertyBinding.PropertyBindingChain.Last()))
+				{
+					bFound = true;
+					
+					Result.Container = Container;
+					Result.Function = Function;
+					
+					const FProperty* ReturnProperty = Function->GetReturnProperty();
+					if (ReturnProperty->IsA<FFloatProperty>())
+					{
+						Result.PropertyType = UE::Chooser::EChooserPropertyAccessType::Float;
+					}
+					else if (ReturnProperty->IsA<FDoubleProperty>())
+					{
+						Result.PropertyType = UE::Chooser::EChooserPropertyAccessType::Double;
+					}
+					else if (ReturnProperty->IsA<FIntProperty>())
+					{
+						Result.PropertyType = UE::Chooser::EChooserPropertyAccessType::Int32;
+					}
+				}
+			}
+		}
 	
-		return true;
+		return bFound;
 	}
+
+	bool ResolvePropertyChain(FChooserEvaluationContext& Context, const FChooserPropertyBinding& PropertyBinding, FResolvedPropertyChainResult& Result)
+	{
+		bool bUseCompiledChain = PropertyBinding.CompiledBinding.IsValid();
+
+#if WITH_EDITOR
+		if (!CVarUseCompiledPropertyChainsInEditor.GetValueOnAnyThread())
+		{
+			bUseCompiledChain = false;
+		}
+#endif
+		
+		if (bUseCompiledChain)
+		{
+			return ResolveCompiledPropertyChain(Context,PropertyBinding,Result);
+		}
+		
+		if (Context.Params.IsValidIndex(PropertyBinding.ContextIndex))
+		{
+			uint8* Container = nullptr;
+			const UStruct* StructType = nullptr;
+			if (FChooserEvaluationInputObject* ObjectParam = Context.Params[PropertyBinding.ContextIndex].GetMutablePtr<FChooserEvaluationInputObject>())
+			{
+				Container = static_cast<uint8*>(ObjectParam->Object);
+				if (Container)
+				{
+					StructType = ObjectParam->Object->GetClass();
+				}
+			}
+			else
+			{
+				Container = Context.Params[PropertyBinding.ContextIndex].GetMutableMemory();
+				StructType = Context.Params[PropertyBinding.ContextIndex].GetScriptStruct();
+			}
+
+
+			if (Container == nullptr || StructType == nullptr)
+			{
+				return false;
+			}
+
+			return ResolvePropertyChain(Container, StructType, PropertyBinding, Result);
+		}
+
+		return false;
+	}
+
 	
 #if WITH_EDITOR
 	void CopyPropertyChain(const TArray<FBindingChainElement>& InBindingChain, FChooserPropertyBinding& OutPropertyBinding)
diff --git a/Engine/Plugins/Experimental/Chooser/Source/Chooser/Private/ObjectColumn.cpp b/Engine/Plugins/Experimental/Chooser/Source/Chooser/Private/ObjectColumn.cpp
--- a/Engine/Plugins/Experimental/Chooser/Source/Chooser/Private/ObjectColumn.cpp
+++ b/Engine/Plugins/Experimental/Chooser/Source/Chooser/Private/ObjectColumn.cpp
@@ -8,24 +8,21 @@
 
 bool FObjectContextProperty::GetValue(FChooserEvaluationContext& Context, FSoftObjectPath& OutResult) const
 {
-	const UStruct* StructType = nullptr;
-	const void* Container = nullptr;
-
-	// todo: update this to use the compiled property chain
-	if (UE::Chooser::ResolvePropertyChain(Context, Binding, Container, StructType))
+	UE::Chooser::FResolvedPropertyChainResult Result;
+	if (UE::Chooser::ResolvePropertyChain(Context, Binding, Result))
 	{
-		if (const FObjectPropertyBase* ObjectProperty = FindFProperty<FObjectPropertyBase>(StructType, Binding.PropertyBindingChain.Last()))
+		if (Result.Function == nullptr)
 		{
 			// if the property is a soft object property, get the path directly
-			if (ObjectProperty->IsA<FSoftObjectProperty>())
+			if (Result.PropertyType == UE::Chooser::EChooserPropertyAccessType::SoftObjectRef)
 			{
-				const FSoftObjectPtr& SoftObjectPtr = *ObjectProperty->ContainerPtrToValuePtr<FSoftObjectPtr>(Container);
+				const FSoftObjectPtr& SoftObjectPtr = *reinterpret_cast<const FSoftObjectPtr*>(Result.Container + Result.PropertyOffset);
 				OutResult = SoftObjectPtr.ToSoftObjectPath();
 				return true;
 			}
-			
+		
 			// otherwise get the value from the object property and convert to a soft object path
-			const UObject* LoadedObject = ObjectProperty->GetObjectPropertyValue_InContainer(Container);
+			const UObject* LoadedObject = reinterpret_cast<const UObject*>(Result.Container + Result.PropertyOffset);
 			OutResult = LoadedObject;
 			return true;
 		}
diff --git a/Engine/Plugins/Experimental/Chooser/Source/Chooser/Public/ChooserPropertyAccess.h b/Engine/Plugins/Experimental/Chooser/Source/Chooser/Public/ChooserPropertyAccess.h
--- a/Engine/Plugins/Experimental/Chooser/Source/Chooser/Public/ChooserPropertyAccess.h
+++ b/Engine/Plugins/Experimental/Chooser/Source/Chooser/Public/ChooserPropertyAccess.h
@@ -47,19 +47,20 @@ namespace UE::Chooser
 	};
 
 	// property type, for numerical conversions
-	enum class EPropertyNumericalType
+	enum class EChooserPropertyAccessType
 	{
-		NONE,
-		BOOL,
-		INT32,
-		FLOAT,
-		DOUBLE,
+		None,
+		Bool,
+		Int32,
+		Float,
+		Double,
+		SoftObjectRef
 	};
 
 	struct FCompiledBinding
 	{
 		int ContextIndex = 0;
-		EPropertyNumericalType PropertyType = EPropertyNumericalType::NONE;
+		EChooserPropertyAccessType PropertyType = EChooserPropertyAccessType::None;
 		// type info for number and enum conversions
 		TArray<FCompiledBindingElement> CompiledChain;
 		const UStruct* TargetType = nullptr;
@@ -179,9 +180,16 @@ struct FContextObjectTypeStruct : public FContextObjectTypeBase
 
 namespace UE::Chooser
 {
-	CHOOSER_API uint8* ResolveCompiledPropertyChain(FChooserEvaluationContext& Context, const FCompiledBinding& CompiledBinding);
-	CHOOSER_API bool ResolvePropertyChain(const void*& Container, const UStruct*& StructType, const TArray<FName>& PropertyBindingChain);
-	CHOOSER_API bool ResolvePropertyChain(FChooserEvaluationContext& Context, const FChooserPropertyBinding& Binding, const void*& OutContainer, const UStruct*& OutStructType);
+	struct FResolvedPropertyChainResult
+	{
+		uint8* Container = nullptr;
+		uint32 PropertyOffset = 0;
+		UFunction* Function = nullptr;
+		EChooserPropertyAccessType PropertyType = EChooserPropertyAccessType::None;
+		uint8 Mask = 0;
+	};
+	
+	CHOOSER_API bool ResolvePropertyChain(FChooserEvaluationContext& Context, const FChooserPropertyBinding& Binding, FResolvedPropertyChainResult& Result);
 
 #if WITH_EDITOR
 	CHOOSER_API void CopyPropertyChain(const TArray<FBindingChainElement>& InBindingChain, FChooserPropertyBinding& OutPropertyBinding);
@@ -192,19 +200,13 @@ template <typename T>
 bool FChooserPropertyBinding::GetValuePtr(FChooserEvaluationContext& Context, T*& OutResult) const
 {
 	using namespace  UE::Chooser;
-	if (!CompiledBinding.IsValid())
-	{
-		return false;
-	}
 
-	const FCompiledBinding& Binding = *CompiledBinding.Get();
-		
-	if (uint8* Result = ResolveCompiledPropertyChain(Context, Binding))
+	FResolvedPropertyChainResult Result;
+	if (ResolvePropertyChain(Context, *this, Result))
 	{
-		const FCompiledBindingElement& Element = Binding.CompiledChain.Last();
-		if (!Element.bIsFunction)
+		if (Result.Function == nullptr)
 		{
-			OutResult = reinterpret_cast<T*>(Result + Binding.CompiledChain.Last().Offset);
+			OutResult = reinterpret_cast<T*>(Result.Container + Result.PropertyOffset);
 			return true;
 		}
 	}
@@ -215,111 +217,104 @@ template <typename T>
 bool FChooserPropertyBinding::GetValue(FChooserEvaluationContext& Context, T& OutResult) const
 {
 	using namespace  UE::Chooser;
-	if (!CompiledBinding.IsValid())
-	{
-		return false;
-	}
 
-	const FCompiledBinding& Binding = *CompiledBinding.Get();
-		
-	if (uint8* Result = ResolveCompiledPropertyChain(Context, Binding))
+	FResolvedPropertyChainResult Result;
+	if (ResolvePropertyChain(Context, *this, Result))
 	{
-		const FCompiledBindingElement& Element = Binding.CompiledChain.Last();
-		if (!Element.bIsFunction)
+		if (Result.Function == nullptr)
 		{
-			switch (Binding.PropertyType)
+			switch (Result.PropertyType)
 			{
-			case EPropertyNumericalType::FLOAT:
+			case EChooserPropertyAccessType::Float:
 				{
-					float FloatResult = *reinterpret_cast<float*>(Result + Binding.CompiledChain.Last().Offset);
+					float FloatResult = *reinterpret_cast<float*>(Result.Container + Result.PropertyOffset);
 					OutResult = static_cast<T>(FloatResult);
 					break;
 				}
-			case EPropertyNumericalType::DOUBLE:
+			case EChooserPropertyAccessType::Double:
 				{
-					double DoubleResult = *reinterpret_cast<double*>(Result + Binding.CompiledChain.Last().Offset);
+					double DoubleResult = *reinterpret_cast<double*>(Result.Container + Result.PropertyOffset);
 					OutResult = static_cast<T>(DoubleResult);
 					break;
 				}
-			case EPropertyNumericalType::INT32:
+			case EChooserPropertyAccessType::Int32:
 				{
-					int32 IntResult = *reinterpret_cast<int*>(Result + Binding.CompiledChain.Last().Offset);
+					int32 IntResult = *reinterpret_cast<int*>(Result.Container + Result.PropertyOffset);
 					OutResult = static_cast<T>(IntResult);
 					break;
 				}
-			case EPropertyNumericalType::BOOL:
+			case EChooserPropertyAccessType::Bool:
 				{
-					const FCompiledBindingElement& Last = Binding.CompiledChain.Last();
-					uint8* ByteValue = Result + Last.Offset;
-					OutResult = !!(*ByteValue & Last.Mask);
+					uint8* ByteValue = Result.Container + Result.PropertyOffset;
+					OutResult = !!(*ByteValue & Result.Mask);
 					break;
 				}
 			default:
-				OutResult = *reinterpret_cast<T*>(Result + Binding.CompiledChain.Last().Offset);
+				OutResult = *reinterpret_cast<T*>(Result.Container + Result.PropertyOffset);
 				break;
 			}
 		}
 		else
 		{
-			UObject* Object = reinterpret_cast<UObject*>(Result);
-			if (Element.Function->IsNative())
+			UObject* Object = reinterpret_cast<UObject*>(Result.Container);
+			if (Result.Function->IsNative())
 			{
-				FFrame Stack(Object, Element.Function, nullptr, nullptr, Element.Function->ChildProperties);
-				switch (Binding.PropertyType)
+				FFrame Stack(Object, Result.Function, nullptr, nullptr, Result.Function->ChildProperties);
+				switch (Result.PropertyType)
 				{
-				case EPropertyNumericalType::FLOAT:
+				case EChooserPropertyAccessType::Float:
 					{
 						float FloatResult = 0;
-						Element.Function->Invoke(Object, Stack, &FloatResult);
+						Result.Function->Invoke(Object, Stack, &FloatResult);
 						OutResult = static_cast<T>(FloatResult);
 						break;
 					}
-				case EPropertyNumericalType::DOUBLE:
+				case EChooserPropertyAccessType::Double:
 					{
 						double DoubleResult = 0;
-						Element.Function->Invoke(Object, Stack, &DoubleResult);
+						Result.Function->Invoke(Object, Stack, &DoubleResult);
 						OutResult = static_cast<T>(DoubleResult);
 						break;
 					}
-				case EPropertyNumericalType::INT32:
+				case EChooserPropertyAccessType::Int32:
 					{
 						int32 IntResult = 0;
-						Element.Function->Invoke(Object, Stack, &IntResult);
+						Result.Function->Invoke(Object, Stack, &IntResult);
 						OutResult = static_cast<T>(IntResult);
 						break;
 					}
 				default:
-					Element.Function->Invoke(Object, Stack, &OutResult);
+					Result.Function->Invoke(Object, Stack, &OutResult);
 					break;
 				}
 			}
 			else
 			{
-				switch (Binding.PropertyType)
+				switch (Result.PropertyType)
 				{
-				case EPropertyNumericalType::FLOAT:
+				case EChooserPropertyAccessType::Float:
 					{
 						float FloatResult = 0;
-						Object->ProcessEvent(Element.Function, &FloatResult);
+						Object->ProcessEvent(Result.Function, &FloatResult);
 						OutResult = static_cast<T>(FloatResult);
 						break;
 					}
-				case EPropertyNumericalType::DOUBLE:
+				case EChooserPropertyAccessType::Double:
 					{
 						double DoubleResult = 0;
-						Object->ProcessEvent(Element.Function, &DoubleResult);
+						Object->ProcessEvent(Result.Function, &DoubleResult);
 						OutResult = static_cast<T>(DoubleResult);
 						break;
 					}
-				case EPropertyNumericalType::INT32:
+				case EChooserPropertyAccessType::Int32:
 					{
 						int32 IntResult = 0;
-						Object->ProcessEvent(Element.Function, &IntResult);
+						Object->ProcessEvent(Result.Function, &IntResult);
 						OutResult = static_cast<T>(IntResult);
 						break;
 					}
 				default:
-					Object->ProcessEvent(Element.Function, &OutResult);
+					Object->ProcessEvent(Result.Function, &OutResult);
 					break;
 				}
 			}
@@ -333,51 +328,44 @@ template <typename T>
 bool FChooserPropertyBinding::SetValue(FChooserEvaluationContext& Context, const T& InValue) const
 {
 	using namespace  UE::Chooser;
-	if (!CompiledBinding.IsValid())
-	{
-		return false;
-	}
 
-	const FCompiledBinding& Binding = *CompiledBinding.Get();
-		
-	if (uint8* Result = ResolveCompiledPropertyChain(Context, Binding))
+	FResolvedPropertyChainResult Result;
+	if (ResolvePropertyChain(Context, *this, Result))
 	{
-		const FCompiledBindingElement& Element = Binding.CompiledChain.Last();
-		if (!Element.bIsFunction)
+		if (Result.Function == nullptr)
 		{
-			switch (Binding.PropertyType)
+			switch (Result.PropertyType)
 			{
-			case EPropertyNumericalType::FLOAT:
+			case EChooserPropertyAccessType::Float:
 				{
-					*reinterpret_cast<float*>(Result + Binding.CompiledChain.Last().Offset) = static_cast<float>(InValue);
+					*reinterpret_cast<float*>(Result.Container + Result.PropertyOffset) = static_cast<float>(InValue);
 					break;
 				}
-			case EPropertyNumericalType::DOUBLE:
+			case EChooserPropertyAccessType::Double:
 				{
-					*reinterpret_cast<double*>(Result + Binding.CompiledChain.Last().Offset) = static_cast<double>(InValue);
+					*reinterpret_cast<double*>(Result.Container + Result.PropertyOffset) = static_cast<double>(InValue);
 					break;
 				}
-			case EPropertyNumericalType::INT32:
+			case EChooserPropertyAccessType::Int32:
 				{
-					*reinterpret_cast<int*>(Result + Binding.CompiledChain.Last().Offset) = static_cast<int>(InValue);
+					*reinterpret_cast<int*>(Result.Container + Result.PropertyOffset) = static_cast<int>(InValue);
 					break;
 				}
-			case EPropertyNumericalType::BOOL:
+			case EChooserPropertyAccessType::Bool:
 			{
-				const FCompiledBindingElement& Last = Binding.CompiledChain.Last();
-				uint8* ByteValue = Result + Last.Offset;
-				if (Last.Mask == 255) // regular bool
+				uint8* ByteValue = Result.Container + Result.PropertyOffset;
+				if (Result.Mask == 255) // regular bool
 				{
-					*reinterpret_cast<bool*>(Result + Binding.CompiledChain.Last().Offset) = static_cast<bool>(InValue);
+					*reinterpret_cast<bool*>(Result.Container + Result.PropertyOffset) = static_cast<bool>(InValue);
 				}
 				else // bitset bool
 				{
-					*ByteValue = ((*ByteValue) & ~Last.Mask) | (InValue ? Last.Mask : 0);
+					*ByteValue = ((*ByteValue) & ~Result.Mask) | (InValue ? Result.Mask : 0);
 				}
 				break;
 			}
 			default:
-				*reinterpret_cast<T*>(Result + Binding.CompiledChain.Last().Offset) = InValue;
+				*reinterpret_cast<T*>(Result.Container + Result.PropertyOffset) = InValue;
 				break;
 			}
 			return true;
diff --git a/Engine/Plugins/Experimental/Chooser/Source/ProxyTable/Private/LookupProxy.cpp b/Engine/Plugins/Experimental/Chooser/Source/ProxyTable/Private/LookupProxy.cpp
--- a/Engine/Plugins/Experimental/Chooser/Source/ProxyTable/Private/LookupProxy.cpp
+++ b/Engine/Plugins/Experimental/Chooser/Source/ProxyTable/Private/LookupProxy.cpp
@@ -71,38 +71,16 @@ FObjectChooserBase::EIteratorStatus FLookupProxyWithOverrideTable::ChooseMulti(F
 
 bool FProxyTableContextProperty::GetValue(FChooserEvaluationContext& Context, const UProxyTable*& OutResult) const
 {
-	if (Binding.CompiledBinding)
+	UProxyTable** ProxyTableReference;
+	if (Binding.GetValuePtr(Context, ProxyTableReference))
 	{
-		UProxyTable** ProxyTableReference;
-		if (Binding.GetValuePtr(Context, ProxyTableReference))
-		{
-			OutResult = *ProxyTableReference;
-			return true;
-		}
-		else
-		{
-			return false;
-		}
+		OutResult = *ProxyTableReference;
+		return true;
 	}
 	else
 	{
-		// for temporary backwards compatibility: ProxyTableContextProperties on UProxyAsset are being phased out,
-		// but they don't get compiled so we need to keep this code-path temporarily
-		const UStruct* StructType = nullptr;
-		const void* Container = nullptr;
-	
-		if (UE::Chooser::ResolvePropertyChain(Context, Binding,Container, StructType))
-		{
-			if (const FObjectProperty* Property = FindFProperty<FObjectProperty>(StructType, Binding.PropertyBindingChain.Last()))
-			{
-				OutResult = *Property->ContainerPtrToValuePtr<UProxyTable*>(Container);
-				return true;
-			}
-		}
+		return false;
 	}
-	
-
-	return false;
 }
 
 void FLookupProxy::Compile(IHasContextClass* HasContext, bool bForce)
diff --git a/Engine/Plugins/Experimental/Chooser/Source/ProxyTable/Private/ProxyTable.cpp b/Engine/Plugins/Experimental/Chooser/Source/ProxyTable/Private/ProxyTable.cpp
--- a/Engine/Plugins/Experimental/Chooser/Source/ProxyTable/Private/ProxyTable.cpp
+++ b/Engine/Plugins/Experimental/Chooser/Source/ProxyTable/Private/ProxyTable.cpp
@@ -176,42 +176,11 @@ static void OutputStructData(const FRuntimeProxyValue& EntryValueData, FChooserE
 {
 	for (const FProxyStructOutput& StructOutput : EntryValueData.OutputStructData)
 	{
-		// todo: should be updated to the following, but ProxyTable needs to trigger compiling of it's output struct bindings first
-		// void* TargetData;
-		// if (StructOutput.Binding.GetValuePtr(Context, TargetData))
-		// {
-		// 	StructOutput.Value.GetScriptStruct()->CopyScriptStruct(TargetData, StructOutput.Value.GetMemory());
-		// }	
-		
-		const void* Container = nullptr;
-		const UStruct* StructType;
-
 		// copy each struct output value
-
-		if (StructOutput.Binding.PropertyBindingChain.IsEmpty())
-		{
-			if(Context.Params.IsValidIndex((StructOutput.Binding.ContextIndex)))
-			{
-				// directly bound to context struct
-				if (Context.Params[StructOutput.Binding.ContextIndex].GetScriptStruct() == StructOutput.Value.GetScriptStruct())
-				{
-					void* TargetData = Context.Params[StructOutput.Binding.ContextIndex].GetMutableMemory();
-					StructOutput.Value.GetScriptStruct()->CopyScriptStruct(TargetData, StructOutput.Value.GetMemory());
-				}
-			}
-		}
-		else if (UE::Chooser::ResolvePropertyChain(Context, StructOutput.Binding, Container, StructType))
+		void* TargetData;
+		if (StructOutput.Binding.GetValuePtr(Context, TargetData))
 		{
-			if (FStructProperty* Property = FindFProperty<FStructProperty>(StructType, StructOutput.Binding.PropertyBindingChain.Last()))
-			{
-				// const cast is here just because ResolvePropertyChain expects a const void*&
-				void* TargetData = Property->ContainerPtrToValuePtr<void>(const_cast<void*>(Container));
-				
-				if (Property->Struct == StructOutput.Value.GetScriptStruct())
-				{
-					Property->Struct->CopyScriptStruct(TargetData, StructOutput.Value.GetMemory());
-				}
-			}
+			StructOutput.Value.GetScriptStruct()->CopyScriptStruct(TargetData, StructOutput.Value.GetMemory());
 		}
 	}
 }
