Chooser Debugging with Rewind Debugger

diff --git a/Engine/Plugins/Animation/GameplayInsights/Source/RewindDebugger/Private/RewindDebugger.h b/Engine/Plugins/Animation/GameplayInsights/Source/RewindDebugger/Private/RewindDebugger.h
--- a/Engine/Plugins/Animation/GameplayInsights/Source/RewindDebugger/Private/RewindDebugger.h
+++ b/Engine/Plugins/Animation/GameplayInsights/Source/RewindDebugger/Private/RewindDebugger.h
@@ -33,6 +33,7 @@ public:
 
 	// IRewindDebugger interface
 	virtual double CurrentTraceTime() const override { return TraceTime.Get(); }
+	virtual double GetScrubTime() const override { return CurrentScrubTime; }
 	virtual const TRange<double>& GetCurrentTraceRange() const override { return CurrentTraceRange; }
 	virtual const TRange<double>& GetCurrentViewRange() const override { return CurrentViewRange; }
 	virtual const TraceServices::IAnalysisSession* GetAnalysisSession() const override;
@@ -93,7 +94,6 @@ public:
 
 	bool CanScrub() const;
 	void ScrubToTime(double ScrubTime, bool bIsScrubbing);
-	double GetScrubTime() { return CurrentScrubTime; }
 
 	// Tick function: While recording, update recording duration.  While paused, and we have recorded data, update skinned mesh poses for the current frame, and handle playback.
 	void Tick(float DeltaTime);
diff --git a/Engine/Plugins/Experimental/Chooser/Chooser.uplugin b/Engine/Plugins/Experimental/Chooser/Chooser.uplugin
--- a/Engine/Plugins/Experimental/Chooser/Chooser.uplugin
+++ b/Engine/Plugins/Experimental/Chooser/Chooser.uplugin
@@ -27,6 +27,10 @@
 		{
 			"Name": "BlendStack",
 			"Enabled": true
+		},
+		{
+			"Name": "GameplayInsights",
+			"Enabled": true
 		}
 	],
 	"Modules" :
diff --git a/Engine/Plugins/Experimental/Chooser/Source/Chooser/Chooser.Build.cs b/Engine/Plugins/Experimental/Chooser/Source/Chooser/Chooser.Build.cs
--- a/Engine/Plugins/Experimental/Chooser/Source/Chooser/Chooser.Build.cs
+++ b/Engine/Plugins/Experimental/Chooser/Source/Chooser/Chooser.Build.cs
@@ -23,7 +23,8 @@ namespace UnrealBuildTool.Rules
 					"GameplayTags",
 					"AnimationCore",
 					"AnimGraphRuntime",
-					"BlendStack"
+					"BlendStack",
+					"TraceLog"
 					// ... add private dependencies that you statically link with here ...
 				}
 			);
diff --git a/Engine/Plugins/Experimental/Chooser/Source/Chooser/Internal/BoolColumn.h b/Engine/Plugins/Experimental/Chooser/Source/Chooser/Internal/BoolColumn.h
--- a/Engine/Plugins/Experimental/Chooser/Source/Chooser/Internal/BoolColumn.h
+++ b/Engine/Plugins/Experimental/Chooser/Source/Chooser/Internal/BoolColumn.h
@@ -6,6 +6,7 @@
 #include "IChooserParameterBool.h"
 #include "ChooserPropertyAccess.h"
 #include "InstancedStruct.h"
+#include "Serialization/MemoryReader.h"
 #include "BoolColumn.generated.h"
 
 UENUM()
@@ -84,11 +85,17 @@ struct CHOOSER_API FBoolColumn : public FChooserColumnBase
 	virtual void Filter(FChooserEvaluationContext& Context, const TArray<uint32>& IndexListIn, TArray<uint32>& IndexListOut) const override;
 
 #if WITH_EDITOR
-	mutable bool TestValue;
+	mutable bool TestValue = false;
 	virtual bool EditorTestFilter(int32 RowIndex) const override
 	{
 		return RowValuesWithAny.IsValidIndex(RowIndex) && (RowValuesWithAny[RowIndex] == EBoolColumnCellValue::MatchAny || TestValue == static_cast<bool>(RowValuesWithAny[RowIndex]));
 	}
+	
+	virtual void SetTestValue(TArrayView<const uint8> Value) override
+	{
+		FMemoryReaderView Reader(Value);
+		Reader << TestValue;
+	}
 #endif
 
 	virtual void PostLoad() override
diff --git a/Engine/Plugins/Experimental/Chooser/Source/Chooser/Internal/Chooser.h b/Engine/Plugins/Experimental/Chooser/Source/Chooser/Internal/Chooser.h
--- a/Engine/Plugins/Experimental/Chooser/Source/Chooser/Internal/Chooser.h
+++ b/Engine/Plugins/Experimental/Chooser/Source/Chooser/Internal/Chooser.h
@@ -34,10 +34,22 @@ public:
 	
 	void SetDebugSelectedRow(int32 Index) const { DebugSelectedRow = Index; }
 	int32 GetDebugSelectedRow() const { return DebugSelectedRow; }
-	bool HasDebugTarget() const { return DebugTarget != nullptr; }
+	bool HasDebugTarget() const { return !DebugTargetName.IsEmpty(); }
 	const UObject* GetDebugTarget() const { return DebugTarget.Get(); }
-	void SetDebugTarget(TWeakObjectPtr<const UObject> Target) { DebugTarget = Target; }
-	void ResetDebugTarget() { DebugTarget.Reset(); }
+	const FString& GetDebugTargetName() const { return DebugTargetName; }
+	void SetDebugTarget(TWeakObjectPtr<const UObject> Target)
+	{
+		DebugTarget = Target;
+		if (const UObject* DebugObject = DebugTarget.Get())
+		{
+			DebugTargetName = DebugObject->GetName();
+		}
+	}
+	void ResetDebugTarget()
+	{
+		DebugTarget.Reset();
+		DebugTargetName = "";
+	}
 	void IterateRecentContextObjects(TFunction<void(const UObject*)> Callback) const;
 	void UpdateDebugging(FChooserEvaluationContext& Context) const;
 	
@@ -54,7 +66,9 @@ private:
 	mutable TSet<TWeakObjectPtr<const UObject>> RecentContextObjects;
 	mutable FCriticalSection DebugLock;
 	// reference to the UObject in PIE  which we want to get debug info for
-	TWeakObjectPtr<const UObject> DebugTarget;
+	mutable TWeakObjectPtr<const UObject> DebugTarget;
+	FString DebugTargetName;
+	
 	// Row which was selected last time this chooser was evaluated on DebugTarget
 	mutable int32 DebugSelectedRow = -1;
 
diff --git a/Engine/Plugins/Experimental/Chooser/Source/Chooser/Internal/ChooserTrace.h b/Engine/Plugins/Experimental/Chooser/Source/Chooser/Internal/ChooserTrace.h
new file mode 100644
--- /dev/null
+++ b/Engine/Plugins/Experimental/Chooser/Source/Chooser/Internal/ChooserTrace.h
@@ -0,0 +1,43 @@
+// Copyright Epic Games, Inc. All Rights Reserved.
+
+#pragma once
+
+#include "CoreTypes.h"
+#include "Trace/Config.h"
+#include "TraceFilter.h"
+#include "Misc/Variant.h"
+#include "IObjectChooser.h"
+#include "Serialization/BufferArchive.h"
+
+#if CHOOSER_TRACE_ENABLED
+
+class UObject;
+
+struct FChooserEvaluationContext;
+
+struct FChooserTrace
+{
+	CHOOSER_API static void OutputChooserEvaluation(const UObject* ChooserAsset, const FChooserEvaluationContext& Context, uint32 SelectedIndex);
+	
+	CHOOSER_API static void OutputChooserValueArchive(const FChooserEvaluationContext& Context, const TCHAR* Key, const FBufferArchive& ValueArchive);
+	
+	template<typename T>
+	CHOOSER_API static void OutputChooserValue(const FChooserEvaluationContext& Context, const TCHAR* Key, const T& Value)
+	{
+		FBufferArchive Archive;
+		Archive << const_cast<T&>(Value);
+		OutputChooserValueArchive(Context,Key,Archive);
+	}
+};
+
+#define TRACE_CHOOSER_EVALUATION(Chooser, Context, SelectedIndex) \
+	FChooserTrace::OutputChooserEvaluation(Chooser, Context, SelectedIndex)
+	
+#define TRACE_CHOOSER_VALUE(Context, Key, Value) \
+	FChooserTrace::OutputChooserValue(Context, Key, Value);
+#else
+
+#define TRACE_CHOOSER_EVALUATION(...)
+#define TRACE_CHOOSER_VALUE(...)
+
+#endif
\ No newline at end of file
diff --git a/Engine/Plugins/Experimental/Chooser/Source/Chooser/Internal/EnumColumn.h b/Engine/Plugins/Experimental/Chooser/Source/Chooser/Internal/EnumColumn.h
--- a/Engine/Plugins/Experimental/Chooser/Source/Chooser/Internal/EnumColumn.h
+++ b/Engine/Plugins/Experimental/Chooser/Source/Chooser/Internal/EnumColumn.h
@@ -6,6 +6,7 @@
 #include "IChooserParameterEnum.h"
 #include "InstancedStruct.h"
 #include "ChooserPropertyAccess.h"
+#include "Serialization/MemoryReader.h"
 #include "EnumColumn.generated.h"
 
 struct FBindingChainElement;
@@ -119,11 +120,17 @@ public:
 	virtual void Filter(FChooserEvaluationContext& Context, const TArray<uint32>& IndexListIn, TArray<uint32>& IndexListOut) const override;
 	
 #if WITH_EDITOR
-	mutable int32 TestValue;
+	mutable uint8 TestValue = 0;
 	virtual bool EditorTestFilter(int32 RowIndex) const override
 	{
 		return RowValues.IsValidIndex(RowIndex) && RowValues[RowIndex].Evaluate(TestValue);
 	}
+	
+	virtual void SetTestValue(TArrayView<const uint8> Value) override
+	{
+		FMemoryReaderView Reader(Value);
+		Reader << TestValue;
+	}
 #endif
 	
 	CHOOSER_COLUMN_BOILERPLATE(FChooserParameterEnumBase);
diff --git a/Engine/Plugins/Experimental/Chooser/Source/Chooser/Internal/FloatRangeColumn.h b/Engine/Plugins/Experimental/Chooser/Source/Chooser/Internal/FloatRangeColumn.h
--- a/Engine/Plugins/Experimental/Chooser/Source/Chooser/Internal/FloatRangeColumn.h
+++ b/Engine/Plugins/Experimental/Chooser/Source/Chooser/Internal/FloatRangeColumn.h
@@ -8,6 +8,7 @@
 #include "IChooserParameterFloat.h"
 #include "ChooserPropertyAccess.h"
 #include "InstancedStruct.h"
+#include "Serialization/MemoryReader.h"
 #include "FloatRangeColumn.generated.h"
 
 USTRUCT(DisplayName = "Float Property Binding")
@@ -89,11 +90,17 @@ struct CHOOSER_API FFloatRangeColumn : public FChooserColumnBase
 	virtual void Filter(FChooserEvaluationContext& Context, const TArray<uint32>& IndexListIn, TArray<uint32>& IndexListOut) const override;
 
 #if WITH_EDITOR
-	mutable float TestValue;
+	mutable double TestValue = 0.0;
 	virtual bool EditorTestFilter(int32 RowIndex) const override
 	{
 		return RowValues.IsValidIndex(RowIndex) && TestValue >= RowValues[RowIndex].Min && TestValue <= RowValues[RowIndex].Max;
 	}
+	
+	virtual void SetTestValue(TArrayView<const uint8> Value) override
+	{
+		FMemoryReaderView Reader(Value);
+		Reader << TestValue;
+	}
 #endif
 	
 	virtual void PostLoad() override
diff --git a/Engine/Plugins/Experimental/Chooser/Source/Chooser/Internal/GameplayTagColumn.h b/Engine/Plugins/Experimental/Chooser/Source/Chooser/Internal/GameplayTagColumn.h
--- a/Engine/Plugins/Experimental/Chooser/Source/Chooser/Internal/GameplayTagColumn.h
+++ b/Engine/Plugins/Experimental/Chooser/Source/Chooser/Internal/GameplayTagColumn.h
@@ -7,6 +7,7 @@
 #include "ChooserPropertyAccess.h"
 #include "GameplayTagContainer.h"
 #include "InstancedStruct.h"
+#include "Serialization/MemoryReader.h"
 #include "GameplayTagColumn.generated.h"
 
 struct FBindingChainElement;
@@ -97,6 +98,14 @@ struct CHOOSER_API FGameplayTagColumn : public FChooserColumnBase
 		}
 		return false;
 	}
+
+	virtual void SetTestValue(TArrayView<const uint8> Value) override
+    {
+    	FMemoryReaderView Reader(Value);
+		FString Tags;
+    	Reader << Tags;
+		TestValue.FromExportString(Tags);
+    }
 #endif
 
 	virtual void PostLoad() override
diff --git a/Engine/Plugins/Experimental/Chooser/Source/Chooser/Internal/ObjectColumn.h b/Engine/Plugins/Experimental/Chooser/Source/Chooser/Internal/ObjectColumn.h
--- a/Engine/Plugins/Experimental/Chooser/Source/Chooser/Internal/ObjectColumn.h
+++ b/Engine/Plugins/Experimental/Chooser/Source/Chooser/Internal/ObjectColumn.h
@@ -6,6 +6,7 @@
 #include "IChooserColumn.h"
 #include "IChooserParameterObject.h"
 #include "InstancedStruct.h"
+#include "Serialization/MemoryReader.h"
 #include "ObjectColumn.generated.h"
 
 struct FBindingChainElement;
@@ -90,6 +91,13 @@ struct CHOOSER_API FObjectColumn : public FChooserColumnBase
 	{
 		return RowValues.IsValidIndex(RowIndex) && RowValues[RowIndex].Evaluate(TestValue);
 	}
+	virtual void SetTestValue(TArrayView<const uint8> Value) override
+	{
+		FMemoryReaderView Reader(Value);
+		FString Path;
+		Reader << Path;
+		TestValue.SetPath(Path);
+	}
 #endif
 
 	CHOOSER_COLUMN_BOILERPLATE(FChooserParameterObjectBase);
diff --git a/Engine/Plugins/Experimental/Chooser/Source/Chooser/Private/BoolColumn.cpp b/Engine/Plugins/Experimental/Chooser/Source/Chooser/Private/BoolColumn.cpp
--- a/Engine/Plugins/Experimental/Chooser/Source/Chooser/Private/BoolColumn.cpp
+++ b/Engine/Plugins/Experimental/Chooser/Source/Chooser/Private/BoolColumn.cpp
@@ -2,6 +2,7 @@
 #include "BoolColumn.h"
 #include "ChooserPropertyAccess.h"
 #include "Chooser.h"
+#include "ChooserTrace.h"
 
 bool FBoolContextProperty::GetValue(FChooserEvaluationContext& Context, bool& OutResult) const
 {
@@ -25,6 +26,8 @@ void FBoolColumn::Filter(FChooserEvaluationContext& Context, const TArray<uint32
 		bool Result = false;
 		InputValue.Get<FChooserParameterBoolBase>().GetValue(Context,Result);
 
+		TRACE_CHOOSER_VALUE(Context, ToCStr(InputValue.Get<FChooserParameterBase>().GetDebugName()), Result);
+
 	#if WITH_EDITOR
 		if (Context.DebuggingInfo.bCurrentDebugTarget)
 		{
diff --git a/Engine/Plugins/Experimental/Chooser/Source/Chooser/Private/Chooser.cpp b/Engine/Plugins/Experimental/Chooser/Source/Chooser/Private/Chooser.cpp
--- a/Engine/Plugins/Experimental/Chooser/Source/Chooser/Private/Chooser.cpp
+++ b/Engine/Plugins/Experimental/Chooser/Source/Chooser/Private/Chooser.cpp
@@ -2,6 +2,8 @@
 #include "Chooser.h"
 #include "ChooserFunctionLibrary.h"
 #include "ChooserPropertyAccess.h"
+#include "ChooserTrace.h"
+#include "ObjectTrace.h"
 #include "Engine/UserDefinedStruct.h"
 #include "Engine/Blueprint.h"
 
@@ -181,17 +183,30 @@ void UChooserTable::IterateRecentContextObjects(TFunction<void(const UObject*)>
 void UChooserTable::UpdateDebugging(FChooserEvaluationContext& Context) const
 {
 	FScopeLock Lock(&DebugLock);
+	
 	for (const FInstancedStruct& Param : Context.Params)
 	{
 		if (const FChooserEvaluationInputObject* ObjectParam = Param.GetPtr<FChooserEvaluationInputObject>())
 		{
-			UObject* ContextObject = ObjectParam->Object;
-			RecentContextObjects.Add(MakeWeakObjectPtr(ContextObject));
-			if (ContextObject == DebugTarget)
+			if (UObject* ContextObject = ObjectParam->Object)
 			{
-				bDebugTestValuesValid = true;
-				Context.DebuggingInfo.bCurrentDebugTarget = true;
-				return;
+				RecentContextObjects.Add(MakeWeakObjectPtr(ContextObject));
+				
+				if (DebugTarget == nullptr && !DebugTargetName.IsEmpty())
+				{
+					// if the DebugTargetName is set, but not the DebugTarget, it means that PIE has been restarted, so try matching by name
+					if (ContextObject->GetName() == DebugTargetName)
+					{
+						DebugTarget = ContextObject;
+					}
+				}
+				
+				if (ContextObject == DebugTarget)
+				{
+					bDebugTestValuesValid = true;
+					Context.DebuggingInfo.bCurrentDebugTarget = true;
+					return;
+				}
 			}
 		}
 	}
@@ -217,6 +232,11 @@ FObjectChooserBase::EIteratorStatus UChooserTable::EvaluateChooser(FChooserEvalu
 		Chooser->UpdateDebugging(Context);
 	}
 #endif
+	
+#if CHOOSER_DEBUGGING_ENABLED
+	Context.DebuggingInfo.CurrentChooser = Chooser;
+#endif
+	
 
 	TArray<uint32> Indices1;
 	TArray<uint32> Indices2;
@@ -252,7 +272,8 @@ FObjectChooserBase::EIteratorStatus UChooserTable::EvaluateChooser(FChooserEvalu
 			Chooser->SetDebugSelectedRow(-1);
 		}
 		#endif
-		
+		TRACE_CHOOSER_EVALUATION(Chooser, Context, -1);
+	
 		if (Chooser->FallbackResult.IsValid())
 		{
 			const FObjectChooserBase& SelectedResult = Chooser->FallbackResult.Get<FObjectChooserBase>();
@@ -291,12 +312,13 @@ FObjectChooserBase::EIteratorStatus UChooserTable::EvaluateChooser(FChooserEvalu
 						const FChooserColumnBase& Column = ColumnData.Get<FChooserColumnBase>();
 						Column.SetOutputs(Context, SelectedIndex);
 					}
-	#if WITH_EDITOR
+					#if WITH_EDITOR
 					if (Context.DebuggingInfo.bCurrentDebugTarget)
 					{
 						Chooser->SetDebugSelectedRow(SelectedIndex);
 					}
-	#endif
+					#endif
+					TRACE_CHOOSER_EVALUATION(Chooser, Context, SelectedIndex);
 				}
 				if (Status == FObjectChooserBase::EIteratorStatus::Stop)
 				{
diff --git a/Engine/Plugins/Experimental/Chooser/Source/Chooser/Private/ChooserTrace.cpp b/Engine/Plugins/Experimental/Chooser/Source/Chooser/Private/ChooserTrace.cpp
new file mode 100644
--- /dev/null
+++ b/Engine/Plugins/Experimental/Chooser/Source/Chooser/Private/ChooserTrace.cpp
@@ -0,0 +1,94 @@
+// Copyright Epic Games, Inc. All Rights Reserved.
+
+#include "ChooserTrace.h"
+#include "ObjectTrace.h"
+#include "IObjectChooser.h"
+#include "Serialization/BufferArchive.h"
+
+#if CHOOSER_TRACE_ENABLED
+
+UE_TRACE_CHANNEL(ChooserChannel)
+
+UE_TRACE_EVENT_BEGIN(Chooser, ChooserEvaluation)
+	UE_TRACE_EVENT_FIELD(double, RecordingTime)
+	UE_TRACE_EVENT_FIELD(uint64, ChooserId)
+	UE_TRACE_EVENT_FIELD(uint64, OwnerId)
+	UE_TRACE_EVENT_FIELD(int32, SelectedIndex)
+UE_TRACE_EVENT_END()
+
+UE_TRACE_EVENT_BEGIN(Chooser, ChooserValue)
+	UE_TRACE_EVENT_FIELD(double, RecordingTime)
+	UE_TRACE_EVENT_FIELD(uint64, ChooserId)
+	UE_TRACE_EVENT_FIELD(uint64, OwnerId)
+	UE_TRACE_EVENT_FIELD(uint8[], Value)
+	UE_TRACE_EVENT_FIELD(UE::Trace::WideString, Key)
+UE_TRACE_EVENT_END()
+
+void FChooserTrace::OutputChooserValueArchive(const FChooserEvaluationContext& Context, const TCHAR* InKey, const FBufferArchive& InValueArchive)
+{
+	UObject* ContextObject = nullptr;
+	if (!Context.Params.IsEmpty())
+	{
+		if (const FChooserEvaluationInputObject* ContextObjectInput = Context.Params[0].GetPtr<FChooserEvaluationInputObject>())
+		{
+			ContextObject = ContextObjectInput->Object;
+		}
+	}
+	
+	bool bChannelEnabled = UE_TRACE_CHANNELEXPR_IS_ENABLED(ChooserChannel);
+
+	if (!bChannelEnabled || ContextObject == nullptr || Context.DebuggingInfo.CurrentChooser == nullptr)
+	{
+		return;
+	}
+
+	if (CANNOT_TRACE_OBJECT(ContextObject->GetWorld()))
+	{
+		return;
+	}
+	
+	TRACE_OBJECT(Context.DebuggingInfo.CurrentChooser);
+	TRACE_OBJECT(ContextObject);
+
+	UE_TRACE_LOG(Chooser, ChooserValue, ChooserChannel)
+		<< ChooserValue.RecordingTime(FObjectTrace::GetWorldElapsedTime(ContextObject->GetWorld()))
+    	<< ChooserValue.ChooserId(FObjectTrace::GetObjectId(Context.DebuggingInfo.CurrentChooser))
+    	<< ChooserValue.OwnerId(FObjectTrace::GetObjectId(ContextObject))
+    	<< ChooserValue.Key(InKey)
+    	<< ChooserValue.Value(InValueArchive.GetData(), InValueArchive.Num());
+}
+
+void FChooserTrace::OutputChooserEvaluation(const UObject* InChooser, const FChooserEvaluationContext& Context, uint32 InSelectedIndex)
+{
+	UObject* ContextObject = nullptr;
+	if (!Context.Params.IsEmpty())
+	{
+		if (const FChooserEvaluationInputObject* ContextObjectInput = Context.Params[0].GetPtr<FChooserEvaluationInputObject>())
+		{
+			ContextObject = ContextObjectInput->Object;
+		}
+	}
+
+	bool bChannelEnabled = UE_TRACE_CHANNELEXPR_IS_ENABLED(ChooserChannel);
+	if (!bChannelEnabled || ContextObject == nullptr || InChooser == nullptr)
+	{
+		return;
+	}
+
+	if (CANNOT_TRACE_OBJECT(ContextObject->GetWorld()))
+	{
+		return;
+	}
+
+	TRACE_OBJECT(InChooser);
+	TRACE_OBJECT(ContextObject);
+
+	UE_TRACE_LOG(Chooser, ChooserEvaluation, ChooserChannel)
+	<< ChooserEvaluation.RecordingTime(FObjectTrace::GetWorldElapsedTime(ContextObject->GetWorld()))
+	<< ChooserEvaluation.ChooserId(FObjectTrace::GetObjectId(InChooser))
+	<< ChooserEvaluation.OwnerId(FObjectTrace::GetObjectId(ContextObject))
+	<< ChooserEvaluation.SelectedIndex(InSelectedIndex);
+}
+
+#endif
+
diff --git a/Engine/Plugins/Experimental/Chooser/Source/Chooser/Private/EnumColumn.cpp b/Engine/Plugins/Experimental/Chooser/Source/Chooser/Private/EnumColumn.cpp
--- a/Engine/Plugins/Experimental/Chooser/Source/Chooser/Private/EnumColumn.cpp
+++ b/Engine/Plugins/Experimental/Chooser/Source/Chooser/Private/EnumColumn.cpp
@@ -1,6 +1,7 @@
 // Copyright Epic Games, Inc. All Rights Reserved.
 #include "EnumColumn.h"
 #include "ChooserPropertyAccess.h"
+#include "ChooserTrace.h"
 
 
 bool FEnumContextProperty::GetValue(FChooserEvaluationContext& Context, uint8& OutResult) const
@@ -73,6 +74,8 @@ void FEnumColumn::Filter(FChooserEvaluationContext& Context, const TArray<uint32
 			TestValue = Result;
 		}
 #endif
+
+		TRACE_CHOOSER_VALUE(Context, ToCStr(InputValue.Get<FChooserParameterBase>().GetDebugName()), Result);
 		
 		for (const uint32 Index : IndexListIn)
 		{
diff --git a/Engine/Plugins/Experimental/Chooser/Source/Chooser/Private/FloatRangeColumn.cpp b/Engine/Plugins/Experimental/Chooser/Source/Chooser/Private/FloatRangeColumn.cpp
--- a/Engine/Plugins/Experimental/Chooser/Source/Chooser/Private/FloatRangeColumn.cpp
+++ b/Engine/Plugins/Experimental/Chooser/Source/Chooser/Private/FloatRangeColumn.cpp
@@ -1,6 +1,7 @@
 // Copyright Epic Games, Inc. All Rights Reserved.
 #include "FloatRangeColumn.h"
 #include "ChooserPropertyAccess.h"
+#include "ChooserTrace.h"
 
 
 
@@ -25,6 +26,8 @@ void FFloatRangeColumn::Filter(FChooserEvaluationContext& Context, const TArray<
 	{
 		double Result = 0.0f;
 		InputValue.Get<FChooserParameterFloatBase>().GetValue(Context, Result);
+		
+		TRACE_CHOOSER_VALUE(Context, ToCStr(InputValue.Get<FChooserParameterBase>().GetDebugName()), Result);
 
 #if WITH_EDITOR
 		if (Context.DebuggingInfo.bCurrentDebugTarget)
diff --git a/Engine/Plugins/Experimental/Chooser/Source/Chooser/Private/GameplayTagColumn.cpp b/Engine/Plugins/Experimental/Chooser/Source/Chooser/Private/GameplayTagColumn.cpp
--- a/Engine/Plugins/Experimental/Chooser/Source/Chooser/Private/GameplayTagColumn.cpp
+++ b/Engine/Plugins/Experimental/Chooser/Source/Chooser/Private/GameplayTagColumn.cpp
@@ -1,6 +1,7 @@
 // Copyright Epic Games, Inc. All Rights Reserved.
 #include "GameplayTagColumn.h"
 #include "ChooserPropertyAccess.h"
+#include "ChooserTrace.h"
 
 bool FGameplayTagContextProperty::GetValue(FChooserEvaluationContext& Context, const FGameplayTagContainer*& OutResult) const
 {
@@ -17,6 +18,8 @@ void FGameplayTagColumn::Filter(FChooserEvaluationContext& Context, const TArray
 	const FGameplayTagContainer* Result = nullptr;
 	if (InputValue.IsValid() && InputValue.Get<FChooserParameterGameplayTagBase>().GetValue(Context,Result))
 	{
+		TRACE_CHOOSER_VALUE(Context, ToCStr(InputValue.Get<FChooserParameterBase>().GetDebugName()), Result->ToString());
+
 #if WITH_EDITOR
 		if (Context.DebuggingInfo.bCurrentDebugTarget)
 		{
diff --git a/Engine/Plugins/Experimental/Chooser/Source/Chooser/Private/ObjectColumn.cpp b/Engine/Plugins/Experimental/Chooser/Source/Chooser/Private/ObjectColumn.cpp
--- a/Engine/Plugins/Experimental/Chooser/Source/Chooser/Private/ObjectColumn.cpp
+++ b/Engine/Plugins/Experimental/Chooser/Source/Chooser/Private/ObjectColumn.cpp
@@ -1,6 +1,7 @@
 // Copyright Epic Games, Inc. All Rights Reserved.
 #include "ObjectColumn.h"
 #include "ChooserPropertyAccess.h"
+#include "ChooserTrace.h"
 
 #if WITH_EDITOR
 	#include "IPropertyAccessEditor.h"
@@ -62,7 +63,8 @@ void FObjectColumn::Filter(FChooserEvaluationContext& Context, const TArray<uint
 	if (InputValue.IsValid() &&
 		InputValue.Get<FChooserParameterObjectBase>().GetValue(Context, Result))
 	{
-		
+		TRACE_CHOOSER_VALUE(Context, ToCStr(InputValue.Get<FChooserParameterBase>().GetDebugName()), Result.ToString());
+	
 #if WITH_EDITOR
 		if (Context.DebuggingInfo.bCurrentDebugTarget)
 		{
diff --git a/Engine/Plugins/Experimental/Chooser/Source/Chooser/Public/IChooserColumn.h b/Engine/Plugins/Experimental/Chooser/Source/Chooser/Public/IChooserColumn.h
--- a/Engine/Plugins/Experimental/Chooser/Source/Chooser/Public/IChooserColumn.h
+++ b/Engine/Plugins/Experimental/Chooser/Source/Chooser/Public/IChooserColumn.h
@@ -24,6 +24,7 @@ public:
 };
 
 class UChooserTable;
+class FVariant;
 
 USTRUCT()
 struct CHOOSER_API FChooserColumnBase
@@ -48,6 +49,7 @@ public:
 		}
 	};
 
+	virtual void SetTestValue(TArrayView<const uint8> Value) { }
 	virtual FChooserParameterBase* GetInputValue() { return nullptr; };
 #if WITH_EDITOR
 	virtual FName RowValuesPropertyName() { return FName(); }
diff --git a/Engine/Plugins/Experimental/Chooser/Source/Chooser/Public/IChooserParameterBase.h b/Engine/Plugins/Experimental/Chooser/Source/Chooser/Public/IChooserParameterBase.h
--- a/Engine/Plugins/Experimental/Chooser/Source/Chooser/Public/IChooserParameterBase.h
+++ b/Engine/Plugins/Experimental/Chooser/Source/Chooser/Public/IChooserParameterBase.h
@@ -12,6 +12,13 @@ struct FChooserParameterBase
 {
 	GENERATED_BODY()
 
+	virtual FString GetDebugName() const
+	{
+		FText Name;
+		GetDisplayName(Name);
+		return Name.ToString();
+	}
+	
 	virtual void GetDisplayName(FText& OutName) const { }
 
 	virtual void PostLoad() {};
diff --git a/Engine/Plugins/Experimental/Chooser/Source/Chooser/Public/IObjectChooser.h b/Engine/Plugins/Experimental/Chooser/Source/Chooser/Public/IObjectChooser.h
--- a/Engine/Plugins/Experimental/Chooser/Source/Chooser/Public/IObjectChooser.h
+++ b/Engine/Plugins/Experimental/Chooser/Source/Chooser/Public/IObjectChooser.h
@@ -8,6 +8,14 @@
 #include "InstancedStruct.h"
 #include "IObjectChooser.generated.h"
 
+#if UE_TRACE_ENABLED && !IS_PROGRAM && !UE_BUILD_SHIPPING && !UE_BUILD_TEST
+#define CHOOSER_TRACE_ENABLED 1
+#else
+#define CHOOSER_TRACE_ENABLED 0
+#endif
+
+#define CHOOSER_DEBUGGING_ENABLED ((CHOOSER_TRACE_ENABLED) || (WITH_EDITOR))
+
 UINTERFACE(NotBlueprintType, meta = (CannotImplementInterfaceInBlueprint))
 class CHOOSER_API UObjectChooser : public UInterface
 {
@@ -21,9 +29,10 @@ public:
 	virtual void ConvertToInstancedStruct(FInstancedStruct& OutInstancedStruct) const { }
 };
 
-#if WITH_EDITOR
+#if CHOOSER_DEBUGGING_ENABLED
 struct CHOOSER_API FChooserDebuggingInfo
 {
+	const UObject* CurrentChooser = nullptr;
 	bool bCurrentDebugTarget = false;
 };
 #endif
@@ -51,7 +60,7 @@ struct CHOOSER_API FChooserEvaluationContext
 		Params.Last().GetMutable<FChooserEvaluationInputObject>().Object = ContextObject;
 	}
 
-	#if WITH_EDITOR
+	#if CHOOSER_DEBUGGING_ENABLED
     	FChooserDebuggingInfo DebuggingInfo;
     #endif
 	
diff --git a/Engine/Plugins/Experimental/Chooser/Source/ChooserEditor/ChooserEditor.Build.cs b/Engine/Plugins/Experimental/Chooser/Source/ChooserEditor/ChooserEditor.Build.cs
--- a/Engine/Plugins/Experimental/Chooser/Source/ChooserEditor/ChooserEditor.Build.cs
+++ b/Engine/Plugins/Experimental/Chooser/Source/ChooserEditor/ChooserEditor.Build.cs
@@ -32,6 +32,12 @@ namespace UnrealBuildTool.Rules
 					"StructUtils",
 					"KismetCompiler",
 					"BlendStack",
+					"TraceAnalysis",
+					"TraceLog",
+					"TraceServices",
+					"TraceInsights",
+					"RewindDebuggerInterface",
+					"GameplayInsights",
 					// ... add private dependencies that you statically link with here ...
 				}
 			);
diff --git a/Engine/Plugins/Experimental/Chooser/Source/ChooserEditor/Private/ChooserAnalyzer.cpp b/Engine/Plugins/Experimental/Chooser/Source/ChooserEditor/Private/ChooserAnalyzer.cpp
new file mode 100644
--- /dev/null
+++ b/Engine/Plugins/Experimental/Chooser/Source/ChooserEditor/Private/ChooserAnalyzer.cpp
@@ -0,0 +1,56 @@
+// Copyright Epic Games, Inc. All Rights Reserved.
+
+#include "ChooserAnalyzer.h"
+
+#include "HAL/LowLevelMemTracker.h"
+#include "Serialization/MemoryReader.h"
+#include "TraceServices/Model/AnalysisSession.h"
+#include "ChooserProvider.h"
+#include "TraceServices/Utils.h"
+
+FChooserAnalyzer::FChooserAnalyzer(TraceServices::IAnalysisSession& InSession, FChooserProvider& InProvider)
+	: Session(InSession)
+	, Provider(InProvider)
+{
+}
+
+void FChooserAnalyzer::OnAnalysisBegin(const FOnAnalysisContext& Context)
+{
+	auto& Builder = Context.InterfaceBuilder;
+
+	Builder.RouteEvent(RouteId_ChooserEvaluation, "Chooser", "ChooserEvaluation");
+	Builder.RouteEvent(RouteId_ChooserValue, "Chooser", "ChooserValue");
+}
+
+bool FChooserAnalyzer::OnEvent(uint16 RouteId, EStyle Style, const FOnEventContext& Context)
+{
+	LLM_SCOPE_BYNAME(TEXT("Insights/FChooserAnalyzer"));
+
+	TraceServices::FAnalysisSessionEditScope _(Session);
+
+	const auto& EventData = Context.EventData;
+	switch (RouteId)
+	{
+		case RouteId_ChooserEvaluation:
+		{
+			uint64 ChooserId = EventData.GetValue<uint64>("ChooserId");
+			uint64 OwnerId = EventData.GetValue<uint64>("OwnerId");
+			int32 SelectedIndex = EventData.GetValue<uint64>("SelectedIndex");
+			double RecordingTime = EventData.GetValue<double>("RecordingTime");
+			Provider.AppendChooserEvaluation(ChooserId, OwnerId, SelectedIndex, RecordingTime);
+			break;
+		}
+		case RouteId_ChooserValue:
+		{
+			uint64 ChooserId = EventData.GetValue<uint64>("ChooserId");
+			uint64 OwnerId = EventData.GetValue<uint64>("OwnerId");
+			FString Key = TraceServices::FTraceAnalyzerUtils::LegacyAttachmentString<TCHAR>("Key", Context);
+			double RecordingTime = EventData.GetValue<double>("RecordingTime");
+			TArrayView<const uint8> SerializedData = EventData.GetArrayView<uint8>("Value");
+			Provider.AppendChooserValue(ChooserId, OwnerId, RecordingTime, FName(Key), SerializedData);
+			break;
+		}
+	}
+
+	return true;
+}
diff --git a/Engine/Plugins/Experimental/Chooser/Source/ChooserEditor/Private/ChooserAnalyzer.h b/Engine/Plugins/Experimental/Chooser/Source/ChooserEditor/Private/ChooserAnalyzer.h
new file mode 100644
--- /dev/null
+++ b/Engine/Plugins/Experimental/Chooser/Source/ChooserEditor/Private/ChooserAnalyzer.h
@@ -0,0 +1,28 @@
+// Copyright Epic Games, Inc. All Rights Reserved.
+
+#pragma once
+
+#include "Trace/Analyzer.h"
+
+class FChooserProvider;
+namespace TraceServices { class IAnalysisSession; }
+
+class FChooserAnalyzer : public UE::Trace::IAnalyzer
+{
+public:
+	FChooserAnalyzer(TraceServices::IAnalysisSession& InSession, FChooserProvider& InProvider);
+
+	virtual void OnAnalysisBegin(const FOnAnalysisContext& Context) override;
+	virtual void OnAnalysisEnd() override {}
+	virtual bool OnEvent(uint16 RouteId, EStyle Style, const FOnEventContext& Context) override;
+
+private:
+	enum : uint16
+	{
+		RouteId_ChooserEvaluation,
+		RouteId_ChooserValue,
+	};
+
+	TraceServices::IAnalysisSession& Session;
+	FChooserProvider& Provider;
+};
diff --git a/Engine/Plugins/Experimental/Chooser/Source/ChooserEditor/Private/ChooserEditorModule.cpp b/Engine/Plugins/Experimental/Chooser/Source/ChooserEditor/Private/ChooserEditorModule.cpp
--- a/Engine/Plugins/Experimental/Chooser/Source/ChooserEditor/Private/ChooserEditorModule.cpp
+++ b/Engine/Plugins/Experimental/Chooser/Source/ChooserEditor/Private/ChooserEditorModule.cpp
@@ -3,28 +3,31 @@
 #include "ChooserEditorModule.h"
 
 #include "AnimNode_ChooserPlayer.h"
-#include "IAssetTools.h"
-#include "ChooserTableEditor.h"
 #include "BoolColumnEditor.h"
+#include "ChooserPropertyAccess.h"
+#include "ChooserTableEditor.h"
+#include "ChooserTableEditorCommands.h"
+#include "CurveOverrideCustomization.h"
 #include "EnumColumnEditor.h"
+#include "Features/IModularFeatures.h"
 #include "FloatRangeColumnEditor.h"
-#include "OutputFloatColumnEditor.h"
+#include "FrameTimeCustomization.h"
 #include "GameplayTagColumnEditor.h"
+#include "IAssetTools.h"
 #include "ObjectColumnEditor.h"
-#include "ChooserTableEditorCommands.h"
-#include "ChooserPropertyAccess.h"
+#include "OutputFloatColumnEditor.h"
+#include "OutputStructColumnEditor.h"
+#include "PropertyAccessChainCustomization.h"
 #include "PropertyEditorModule.h"
 #include "RandomizeColumnEditor.h"
-#include "OutputStructColumnEditor.h"
-#include "CurveOverrideCustomization.h"
-#include "FrameTimeCustomization.h"
+#include "ChooserTrack.h"
 
 #define LOCTEXT_NAMESPACE "ChooserEditorModule"
 
 namespace UE::ChooserEditor
 {
 	
-
+FChoosersTrackCreator GChoosersTrackCreator;
 
 void FModule::StartupModule()
 {
@@ -51,10 +54,16 @@ void FModule::StartupModule()
 	PropertyModule.RegisterCustomPropertyTypeLayout(FChooserObjectPropertyBinding::StaticStruct()->GetFName(), FOnGetPropertyTypeCustomizationInstance::CreateLambda([] { return MakeShared<FPropertyAccessChainCustomization>(); }));
 	PropertyModule.RegisterCustomPropertyTypeLayout(FChooserStructPropertyBinding::StaticStruct()->GetFName(), FOnGetPropertyTypeCustomizationInstance::CreateLambda([] { return MakeShared<FPropertyAccessChainCustomization>(); }));
 
+	IModularFeatures::Get().RegisterModularFeature( IRewindDebuggerExtension::ModularFeatureName, &RewindDebuggerChooser);
+	IModularFeatures::Get().RegisterModularFeature(RewindDebugger::IRewindDebuggerTrackCreator::ModularFeatureName, &GChoosersTrackCreator);
+	IModularFeatures::Get().RegisterModularFeature(TraceServices::ModuleFeatureName, &ChooserTraceModule);
 }
 
 void FModule::ShutdownModule()
 {
+	IModularFeatures::Get().UnregisterModularFeature( IRewindDebuggerExtension::ModularFeatureName, &RewindDebuggerChooser);
+	IModularFeatures::Get().UnregisterModularFeature(RewindDebugger::IRewindDebuggerTrackCreator::ModularFeatureName, &GChoosersTrackCreator);
+	IModularFeatures::Get().UnregisterModularFeature(TraceServices::ModuleFeatureName, &ChooserTraceModule);
 	FChooserTableEditorCommands::Unregister();
 }
 
diff --git a/Engine/Plugins/Experimental/Chooser/Source/ChooserEditor/Private/ChooserEditorModule.h b/Engine/Plugins/Experimental/Chooser/Source/ChooserEditor/Private/ChooserEditorModule.h
--- a/Engine/Plugins/Experimental/Chooser/Source/ChooserEditor/Private/ChooserEditorModule.h
+++ b/Engine/Plugins/Experimental/Chooser/Source/ChooserEditor/Private/ChooserEditorModule.h
@@ -3,8 +3,9 @@
 #pragma once
 
 #include "CoreMinimal.h"
-#include "PropertyAccessChainCustomization.h"
+#include "ChooserTraceModule.h"
 #include "Modules/ModuleInterface.h"
+#include "RewindDebuggerChooser.h"
 
 namespace UE::ChooserEditor
 {
@@ -14,6 +15,9 @@ class FModule : public IModuleInterface
 private:
 	virtual void StartupModule() override;
 	virtual void ShutdownModule() override;
+	
+	FRewindDebuggerChooser RewindDebuggerChooser;
+	FChooserTraceModule ChooserTraceModule;
 };
 
 }
\ No newline at end of file
diff --git a/Engine/Plugins/Experimental/Chooser/Source/ChooserEditor/Private/ChooserProvider.cpp b/Engine/Plugins/Experimental/Chooser/Source/ChooserEditor/Private/ChooserProvider.cpp
new file mode 100644
--- /dev/null
+++ b/Engine/Plugins/Experimental/Chooser/Source/ChooserEditor/Private/ChooserProvider.cpp
@@ -0,0 +1,92 @@
+// Copyright Epic Games, Inc. All Rights Reserved.
+
+#include "ChooserProvider.h"
+#include "ObjectTrace.h"
+
+FName FChooserProvider::ProviderName("ChooserProvider");
+
+#define LOCTEXT_NAMESPACE "ChooserProvider"
+
+FChooserProvider::FChooserProvider(TraceServices::IAnalysisSession& InSession)
+	: Session(InSession)
+{
+}
+
+bool FChooserProvider::ReadChooserEvaluationTimeline(uint64 InObjectId, TFunctionRef<void(const ChooserEvaluationTimeline&)> Callback) const
+{
+	Session.ReadAccessCheck();
+
+	const uint32* IndexPtr = ObjectIdToChooserEvaluationTimelines.Find(InObjectId);
+	if(IndexPtr != nullptr)
+	{
+		if (*IndexPtr < uint32(ChooserEvaluationTimelines.Num()))
+		{
+			Callback(*ChooserEvaluationTimelines[*IndexPtr]);
+			return true;
+		}
+	}
+
+	return false;
+}
+
+bool FChooserProvider::ReadChooserValueTimeline(uint64 InObjectId, TFunctionRef<void(const ChooserValueTimeline&)> Callback) const
+{
+	Session.ReadAccessCheck();
+
+	const uint32* IndexPtr = ObjectIdToChooserEvaluationTimelines.Find(InObjectId);
+	if(IndexPtr != nullptr)
+	{
+		if (*IndexPtr < uint32(ChooserEvaluationTimelines.Num()))
+		{
+			Callback(*ChooserValueTimelines[*IndexPtr]);
+			return true;
+		}
+	}
+
+	return false;
+}
+
+void FChooserProvider::EnumerateChooserEvaluationTimelines(TFunctionRef<void(uint64 OwnerId, const ChooserEvaluationTimeline&)> Callback) const
+{
+	Session.ReadAccessCheck();
+
+	for (TTuple<uint64, uint32> Entry : ObjectIdToChooserEvaluationTimelines)
+	{
+		if (ChooserEvaluationTimelines.IsValidIndex(Entry.Value))
+		{
+			Callback(Entry.Key, *ChooserEvaluationTimelines[Entry.Value]);
+		}
+	}
+}
+
+uint32 FChooserProvider::GetTimelineIndex(uint64 InObjectId)
+{
+	TSharedPtr<TraceServices::TPointTimeline<FChooserEvaluationData>> Timeline;
+	uint32* IndexPtr = ObjectIdToChooserEvaluationTimelines.Find(InObjectId);
+	if(IndexPtr != nullptr)
+	{
+		return *IndexPtr;
+	}
+	else
+	{
+		uint32 Index = ChooserEvaluationTimelines.Num();
+		ObjectIdToChooserEvaluationTimelines.Add(InObjectId, ChooserEvaluationTimelines.Num());
+		ChooserEvaluationTimelines.Add( MakeShared<TraceServices::TPointTimeline<FChooserEvaluationData>>(Session.GetLinearAllocator()));
+		ChooserValueTimelines.Add(MakeShared<TraceServices::TPointTimeline<FChooserValueData>>(Session.GetLinearAllocator()));
+		return Index;
+	}
+}
+
+void  FChooserProvider::AppendChooserEvaluation(uint64 InChooserId, uint64 InObjectId, int32 InSelectedIndex, double InRecordingTime)
+{
+	Session.WriteAccessCheck();
+
+	TSharedPtr<TraceServices::TPointTimeline<FChooserEvaluationData>> Timeline = ChooserEvaluationTimelines[GetTimelineIndex(InObjectId)];
+	Timeline->AppendEvent(InRecordingTime, { InChooserId, InSelectedIndex });
+}
+
+
+
+
+
+#undef LOCTEXT_NAMESPACE
diff --git a/Engine/Plugins/Experimental/Chooser/Source/ChooserEditor/Private/ChooserProvider.h b/Engine/Plugins/Experimental/Chooser/Source/ChooserEditor/Private/ChooserProvider.h
new file mode 100644
--- /dev/null
+++ b/Engine/Plugins/Experimental/Chooser/Source/ChooserEditor/Private/ChooserProvider.h
@@ -0,0 +1,66 @@
+// Copyright Epic Games, Inc. All Rights Reserved.
+
+#pragma once
+
+#include "Engine/EngineBaseTypes.h"
+#include "Model/PointTimeline.h"
+#include "Containers/Map.h"
+#include "Templates/SharedPointer.h"
+#include "Model/IntervalTimeline.h"
+#include "TraceServices/Model/AnalysisSession.h"
+
+namespace TraceServices { class IAnalysisSession; }
+
+struct FChooserEvaluationData
+{
+	uint64 ChooserId;
+	int32 SelectedIndex;
+};
+
+struct FChooserValueData
+{
+	uint64 ChooserId;
+	FName Key;
+	// probably should get this using the linear allocator somehow:
+	TArray<uint8> Value;
+};
+
+class FChooserProvider : public TraceServices::IProvider
+{
+public:
+	typedef TraceServices::ITimeline<FChooserEvaluationData> ChooserEvaluationTimeline;
+	typedef TraceServices::ITimeline<FChooserValueData> ChooserValueTimeline;
+	
+	static FName ProviderName;
+
+	FChooserProvider(TraceServices::IAnalysisSession& InSession);
+
+	bool ReadChooserEvaluationTimeline(uint64 InObjectId, TFunctionRef<void(const ChooserEvaluationTimeline&)> Callback) const;
+	bool ReadChooserValueTimeline(uint64 InObjectId, TFunctionRef<void(const ChooserValueTimeline&)> Callback) const;
+	void EnumerateChooserEvaluationTimelines(TFunctionRef<void(uint64 OwnerId, const ChooserEvaluationTimeline&)> Callback) const;
+
+
+
+	void AppendChooserEvaluation(uint64 InObjectId, uint64 InChooser, int32 InResultId, double InRecordingTime);
+
+	void AppendChooserValue(uint64 InChooserId, uint64 InObjectId, double InRecordingTime, const FName& Key, TArrayView<const uint8> SerializedValue)
+    {
+    	Session.WriteAccessCheck();
+    
+    	TSharedPtr<TraceServices::TPointTimeline<FChooserValueData>> Timeline = ChooserValueTimelines[GetTimelineIndex(InObjectId)];
+		FChooserValueData EventData;
+		EventData.ChooserId = InChooserId;
+		EventData.Key = Key;
+		EventData.Value = SerializedValue;
+    	Timeline->AppendEvent(InRecordingTime, EventData);
+    }
+
+private:
+	uint32 GetTimelineIndex(uint64 ObjectId);
+	
+	TraceServices::IAnalysisSession& Session;
+
+	TMap<uint64, uint32> ObjectIdToChooserEvaluationTimelines;
+	TArray<TSharedRef<TraceServices::TPointTimeline<FChooserEvaluationData>>> ChooserEvaluationTimelines;
+	TArray<TSharedRef<TraceServices::TPointTimeline<FChooserValueData>>> ChooserValueTimelines;
+};
diff --git a/Engine/Plugins/Experimental/Chooser/Source/ChooserEditor/Private/ChooserTableEditor.cpp b/Engine/Plugins/Experimental/Chooser/Source/ChooserEditor/Private/ChooserTableEditor.cpp
--- a/Engine/Plugins/Experimental/Chooser/Source/ChooserEditor/Private/ChooserTableEditor.cpp
+++ b/Engine/Plugins/Experimental/Chooser/Source/ChooserEditor/Private/ChooserTableEditor.cpp
@@ -217,7 +217,7 @@ void FChooserTableEditor::RegisterToolbar()
 						{
 							if (Chooser->HasDebugTarget())
 							{
-								return  FText::FromString(Chooser->GetDebugTarget()->GetName());
+								return  FText::FromString(Chooser->GetDebugTargetName());
 							}
 							else
 							{
diff --git a/Engine/Plugins/Experimental/Chooser/Source/ChooserEditor/Private/ChooserTraceModule.cpp b/Engine/Plugins/Experimental/Chooser/Source/ChooserEditor/Private/ChooserTraceModule.cpp
new file mode 100644
--- /dev/null
+++ b/Engine/Plugins/Experimental/Chooser/Source/ChooserEditor/Private/ChooserTraceModule.cpp
@@ -0,0 +1,32 @@
+// Copyright Epic Games, Inc. All Rights Reserved.
+
+#include "ChooserTraceModule.h"
+#include "ChooserProvider.h"
+#include "ChooserAnalyzer.h"
+
+FName FChooserTraceModule::ModuleName("Chooser");
+
+void FChooserTraceModule::GetModuleInfo(TraceServices::FModuleInfo& OutModuleInfo)
+{
+	OutModuleInfo.Name = ModuleName;
+	OutModuleInfo.DisplayName = TEXT("Chooser");
+}
+
+void FChooserTraceModule::OnAnalysisBegin(TraceServices::IAnalysisSession& InSession)
+{
+	TSharedPtr<FChooserProvider> ChooserProvider = MakeShared<FChooserProvider>(InSession);
+	InSession.AddProvider(FChooserProvider::ProviderName, ChooserProvider);
+
+	InSession.AddAnalyzer(new FChooserAnalyzer(InSession, *ChooserProvider));
+}
+
+void FChooserTraceModule::GetLoggers(TArray<const TCHAR *>& OutLoggers)
+{
+	OutLoggers.Add(TEXT("Chooser"));
+}
+
+void FChooserTraceModule::GenerateReports(const TraceServices::IAnalysisSession& Session, const TCHAR* CmdLine, const TCHAR* OutputDirectory)
+{
+
+}
+
diff --git a/Engine/Plugins/Experimental/Chooser/Source/ChooserEditor/Private/ChooserTraceModule.h b/Engine/Plugins/Experimental/Chooser/Source/ChooserEditor/Private/ChooserTraceModule.h
new file mode 100644
--- /dev/null
+++ b/Engine/Plugins/Experimental/Chooser/Source/ChooserEditor/Private/ChooserTraceModule.h
@@ -0,0 +1,24 @@
+// Copyright Epic Games, Inc. All Rights Reserved.
+
+#pragma once
+
+#include "CoreMinimal.h"
+#include "TraceServices/ModuleService.h"
+
+class IAnimationProvider;
+namespace Insights { class ITimingViewSession; }
+
+class FChooserTraceModule : public TraceServices::IModule
+{
+public:
+	// TraceServices::IModule interface
+	virtual void GetModuleInfo(TraceServices::FModuleInfo& OutModuleInfo) override;
+	virtual void OnAnalysisBegin(TraceServices::IAnalysisSession& Session) override;
+	virtual void GetLoggers(TArray<const TCHAR *>& OutLoggers) override;
+	virtual void GenerateReports(const TraceServices::IAnalysisSession& Session, const TCHAR* CmdLine, const TCHAR* OutputDirectory) override;
+	virtual const TCHAR* GetCommandLineArgument() override { return TEXT("chooser"); }
+
+private:
+	static FName ModuleName;
+};
+
diff --git a/Engine/Plugins/Experimental/Chooser/Source/ChooserEditor/Private/ChooserTrack.cpp b/Engine/Plugins/Experimental/Chooser/Source/ChooserEditor/Private/ChooserTrack.cpp
new file mode 100644
--- /dev/null
+++ b/Engine/Plugins/Experimental/Chooser/Source/ChooserEditor/Private/ChooserTrack.cpp
@@ -0,0 +1,243 @@
+// Copyright Epic Games, Inc. All Rights Reserved.
+
+#include "ChooserTrack.h"
+
+#include "Chooser.h"
+#include "ChooserProvider.h"
+#include "Editor.h"
+#include "IGameplayProvider.h"
+#include "IRewindDebugger.h"
+#include "ObjectTrace.h"
+#include "Subsystems/AssetEditorSubsystem.h"
+
+#define LOCTEXT_NAMESPACE "ChooserTrack"
+
+namespace UE::ChooserEditor
+{
+
+FChooserTrack::FChooserTrack(uint64 InObjectId, uint64 InChooserId) :
+	ObjectId(InObjectId),
+	ChooserId(InChooserId)
+{
+	EventData = MakeShared<SEventTimelineView::FTimelineEventData>();
+	Icon = FSlateIcon("EditorStyle", "Sequencer.Tracks.Event", "Sequencer.Tracks.Event");
+	
+	ChooserTable = FObjectTrace::GetObjectFromId(InChooserId);
+	if (ChooserTable)
+	{
+		TrackName = FText::FromString(ChooserTable->GetName());
+	}
+}
+
+FChooserTrack::~FChooserTrack()
+{
+}
+
+TSharedPtr<SEventTimelineView::FTimelineEventData> FChooserTrack::GetEventData() const
+{
+	if (!EventData.IsValid())
+	{
+		EventData = MakeShared<SEventTimelineView::FTimelineEventData>();
+	}
+	
+	EventUpdateRequested++;
+	
+	return EventData;
+}
+	
+static FLinearColor MakeNotifyColor(uint32 InSeed, bool bInLine = false)
+{
+	FRandomStream Stream(InSeed);
+	const uint8 Hue = (uint8)(Stream.FRand() * 255.0f);
+	const uint8 SatVal = bInLine ? 196 : 128;
+	return FLinearColor::MakeFromHSV8(Hue, SatVal, SatVal);
+}
+
+TSharedPtr<SWidget> FChooserTrack::GetDetailsViewInternal() 
+{
+	return nullptr;
+}
+
+bool FChooserTrack::UpdateInternal()
+{
+	IRewindDebugger* RewindDebugger = IRewindDebugger::Instance();
+	
+	TRange<double> RecordingTimeRange = RewindDebugger->GetCurrentViewRange();
+	double StartTime = RecordingTimeRange.GetLowerBoundValue();
+	double EndTime = RecordingTimeRange.GetUpperBoundValue();
+	
+	const TraceServices::IAnalysisSession* AnalysisSession = RewindDebugger->GetAnalysisSession();
+	const FChooserProvider* ChooserProvider = AnalysisSession->ReadProvider<FChooserProvider>(FChooserProvider::ProviderName);
+	
+	if(EventUpdateRequested > 10 && ChooserProvider)
+	{
+		EventUpdateRequested = 0;
+		
+		EventData->Points.SetNum(0,false);
+		EventData->Windows.SetNum(0);
+
+		TraceServices::FAnalysisSessionReadScope SessionReadScope(*AnalysisSession);
+		
+		ChooserProvider->ReadChooserEvaluationTimeline(ObjectId, [this, StartTime, EndTime](const FChooserProvider::ChooserEvaluationTimeline& InTimeline)
+		{
+			InTimeline.EnumerateEvents(StartTime, EndTime, [this, StartTime, EndTime](double InStartTime, double InEndTime, uint32 InDepth, const FChooserEvaluationData& ChooserEvaluationData)
+			{
+				if (ChooserEvaluationData.ChooserId == ChooserId)
+				{
+					EventData->Points.Add({InStartTime, FText(), FText(), FLinearColor::White});
+				}
+				return TraceServices::EEventEnumerate::Continue;
+			});
+		});
+	}
+
+	EventUpdateRequested++;
+
+	bool bChanged = false;
+	return bChanged;
+	
+}
+
+bool FChooserTrack::HandleDoubleClickInternal()
+{
+	const IRewindDebugger* RewindDebugger = IRewindDebugger::Instance();
+	if (const TraceServices::IAnalysisSession* AnalysisSession = RewindDebugger->GetAnalysisSession())
+	{
+		TraceServices::FAnalysisSessionReadScope SessionReadScope(*AnalysisSession);
+
+		const IGameplayProvider* GameplayProvider = AnalysisSession->ReadProvider<IGameplayProvider>("GameplayProvider");
+		const FObjectInfo& AssetInfo = GameplayProvider->GetObjectInfo(GetChooserId());
+
+		if (UObject* OwnerObject = FObjectTrace::GetObjectFromId(ObjectId))
+		{
+			if (UChooserTable* Chooser = Cast<UChooserTable>(FObjectTrace::GetObjectFromId(GetChooserId())))
+			{
+				// attach chooser table editor debugging
+				Chooser->SetDebugTarget(OwnerObject);
+				Chooser->bEnableDebugTesting = true;
+			}
+		}
+
+		GEditor->GetEditorSubsystem<UAssetEditorSubsystem>()->OpenEditorForAsset(AssetInfo.PathName);
+
+		return true;
+	}
+	return false;
+}
+
+TSharedPtr<SWidget> FChoosersTrack::GetDetailsViewInternal() 
+{
+	return nullptr;
+}
+
+TSharedPtr<SWidget> FChooserTrack::GetTimelineViewInternal()
+{
+	return SNew(SEventTimelineView)
+		.ViewRange_Lambda([]() { return IRewindDebugger::Instance()->GetCurrentViewRange(); })
+		.EventData_Raw(this, &FChooserTrack::GetEventData);
+}
+
+static const FName ChoosersName("Choosers");
+
+ FName FChoosersTrackCreator::GetTargetTypeNameInternal() const
+ {
+ 	static const FName ObjectName("Object");
+	return ObjectName;
+}
+	
+FName FChoosersTrackCreator::GetNameInternal() const
+{
+	return ChoosersName;
+}
+
+TSharedPtr<RewindDebugger::FRewindDebuggerTrack> FChoosersTrackCreator::CreateTrackInternal(uint64 ObjectId) const
+{
+	return MakeShared<FChoosersTrack>(ObjectId);
+}
+
+		
+FChoosersTrack::FChoosersTrack(uint64 InObjectId) :
+	ObjectId(InObjectId)
+{
+	Icon = FSlateIcon("EditorStyle", "Sequencer.Tracks.Event", "Sequencer.Tracks.Event");
+}
+
+
+bool FChoosersTrack::UpdateInternal()
+{
+	IRewindDebugger* RewindDebugger = IRewindDebugger::Instance();
+	
+	TRange<double> RecordingTimeRange = RewindDebugger->GetCurrentViewRange();
+	double StartTime = RecordingTimeRange.GetLowerBoundValue();
+	double EndTime = RecordingTimeRange.GetUpperBoundValue();
+	
+	const TraceServices::IAnalysisSession* AnalysisSession = RewindDebugger->GetAnalysisSession();
+	const FChooserProvider* ChooserProvider = AnalysisSession->ReadProvider<FChooserProvider>(FChooserProvider::ProviderName);
+	
+	bool bChanged = false;
+	
+	if(ChooserProvider)
+	{
+		TArray<uint64> UniqueTrackIds;
+		
+		TraceServices::FAnalysisSessionReadScope SessionReadScope(*AnalysisSession);
+		
+		ChooserProvider->ReadChooserEvaluationTimeline(ObjectId, [this, StartTime, EndTime, &UniqueTrackIds](const FChooserProvider::ChooserEvaluationTimeline& InTimeline)
+		{
+			InTimeline.EnumerateEvents(StartTime, EndTime, [this, &UniqueTrackIds](double InStartTime, double InEndTime, uint32 InDepth, const FChooserEvaluationData& ChooserEvaluationData)
+			{
+				UniqueTrackIds.AddUnique(ChooserEvaluationData.ChooserId);
+		 		return TraceServices::EEventEnumerate::Continue;
+		 	});
+		});
+		
+		UniqueTrackIds.StableSort();
+		const int32 TrackCount = UniqueTrackIds.Num();
+		
+		if (Children.Num() != TrackCount)
+			bChanged = true;
+		
+		Children.SetNum(UniqueTrackIds.Num());
+		for(int i = 0; i < TrackCount; i++)
+		{
+			if (!Children[i].IsValid() || !(Children[i].Get()->GetChooserId() == UniqueTrackIds[i]))
+			{
+				Children[i] = MakeShared<FChooserTrack>(ObjectId, UniqueTrackIds[i]);
+				bChanged = true;
+			}
+		
+			bChanged = bChanged || Children[i]->Update();
+		}
+	}
+	
+	return bChanged;
+}
+
+void FChoosersTrack::IterateSubTracksInternal(TFunction<void(TSharedPtr<FRewindDebuggerTrack> SubTrack)> IteratorFunction)
+{
+	for(TSharedPtr<FChooserTrack>& Track : Children)
+	{
+		IteratorFunction(Track);
+	}
+};
+
+bool FChoosersTrackCreator::HasDebugInfoInternal(uint64 ObjectId) const
+{
+	const TraceServices::IAnalysisSession* AnalysisSession = IRewindDebugger::Instance()->GetAnalysisSession();
+	
+	TraceServices::FAnalysisSessionReadScope SessionReadScope(*AnalysisSession);
+	bool bHasData = false;
+	if (const FChooserProvider* ChooserProvider = AnalysisSession->ReadProvider<FChooserProvider>(FChooserProvider::ProviderName))
+	{
+		ChooserProvider->ReadChooserEvaluationTimeline(ObjectId, [&bHasData](const FChooserProvider::ChooserEvaluationTimeline& InTimeline)
+	 	{
+	 		bHasData = true;
+	 	});
+	}
+	return bHasData;
+}
+
+	
+}
+
+#undef LOCTEXT_NAMESPACE
\ No newline at end of file
diff --git a/Engine/Plugins/Experimental/Chooser/Source/ChooserEditor/Private/ChooserTrack.h b/Engine/Plugins/Experimental/Chooser/Source/ChooserEditor/Private/ChooserTrack.h
new file mode 100644
--- /dev/null
+++ b/Engine/Plugins/Experimental/Chooser/Source/ChooserEditor/Private/ChooserTrack.h
@@ -0,0 +1,75 @@
+// Copyright Epic Games, Inc. All Rights Reserved.
+#pragma once
+
+#include "IRewindDebuggerTrackCreator.h"
+#include "RewindDebuggerTrack.h"
+#include "Textures/SlateIcon.h"
+#include "SEventTimelineView.h"
+#include "IDetailsView.h"
+
+namespace UE::ChooserEditor
+{
+
+class FChooserTrack : public RewindDebugger::FRewindDebuggerTrack
+{
+public:
+	FChooserTrack(uint64 InObjectId, uint64 InChooserId);
+	virtual ~FChooserTrack();
+	
+	uint64 GetChooserId() { return ChooserId; }
+private:
+	virtual bool UpdateInternal() override;
+	TSharedPtr<SWidget> GetDetailsViewInternal() override;
+	virtual TSharedPtr<SWidget> GetTimelineViewInternal() override;
+	virtual FSlateIcon GetIconInternal() override { return Icon; }
+	virtual FText GetDisplayNameInternal() const override { return TrackName; }
+	virtual uint64 GetObjectIdInternal() const override { return ObjectId; }
+	// virtual FName GetNameInternal() const override { return Category; }
+	virtual bool HandleDoubleClickInternal() override;
+
+	TSharedPtr<SEventTimelineView::FTimelineEventData> GetEventData() const;
+	
+	mutable TSharedPtr<SEventTimelineView::FTimelineEventData> EventData;
+	mutable int EventUpdateRequested = 0;
+
+	UObject* ChooserTable;
+	uint64 ObjectId;
+	uint64 ChooserId;
+	FSlateIcon Icon;
+	FText TrackName;
+};
+
+
+
+// Parent track 
+class FChoosersTrack : public RewindDebugger::FRewindDebuggerTrack
+{
+public:
+	FChoosersTrack(uint64 InObjectId);
+private:
+	virtual bool UpdateInternal() override;
+	virtual TSharedPtr<SWidget> GetDetailsViewInternal() override;
+
+	virtual FSlateIcon GetIconInternal() override { return Icon; }
+	virtual FName GetNameInternal() const override { return "Choosers"; }
+	virtual FText GetDisplayNameInternal() const override { return NSLOCTEXT("RewindDebugger", "ChooserTrackName", "Chooser Evaluation"); }
+	virtual uint64 GetObjectIdInternal() const override { return ObjectId; }
+	virtual void IterateSubTracksInternal(TFunction<void(TSharedPtr<FRewindDebuggerTrack> SubTrack)> IteratorFunction) override;
+
+	FSlateIcon Icon;
+	uint64 ObjectId;
+
+	TArray<TSharedPtr<FChooserTrack>> Children;
+};
+
+
+class FChoosersTrackCreator : public RewindDebugger::IRewindDebuggerTrackCreator
+{
+private:
+	virtual FName GetTargetTypeNameInternal() const;
+	virtual FName GetNameInternal() const override;
+	virtual TSharedPtr<RewindDebugger::FRewindDebuggerTrack> CreateTrackInternal(uint64 ObjectId) const override;
+	virtual bool HasDebugInfoInternal(uint64 ObjectId) const override;
+};
+
+}
\ No newline at end of file
diff --git a/Engine/Plugins/Experimental/Chooser/Source/ChooserEditor/Private/RewindDebuggerChooser.cpp b/Engine/Plugins/Experimental/Chooser/Source/ChooserEditor/Private/RewindDebuggerChooser.cpp
new file mode 100644
--- /dev/null
+++ b/Engine/Plugins/Experimental/Chooser/Source/ChooserEditor/Private/RewindDebuggerChooser.cpp
@@ -0,0 +1,86 @@
+// Copyright Epic Games, Inc. All Rights Reserved.
+
+#include "RewindDebuggerChooser.h"
+
+#include "Chooser.h"
+#include "ChooserProvider.h"
+#include "ObjectTrace.h"
+
+FRewindDebuggerChooser::FRewindDebuggerChooser()
+{
+
+}
+
+void FRewindDebuggerChooser::Update(float DeltaTime, IRewindDebugger* RewindDebugger)
+{
+	const TraceServices::IAnalysisSession* AnalysisSession = RewindDebugger->GetAnalysisSession();
+	TraceServices::FAnalysisSessionReadScope SessionReadScope(*AnalysisSession);
+	
+	const FChooserProvider* ChooserProvider = AnalysisSession->ReadProvider<FChooserProvider>(FChooserProvider::ProviderName);
+	
+	const TraceServices::IFrameProvider& FrameProvider = TraceServices::ReadFrameProvider(*RewindDebugger->GetAnalysisSession());
+	TraceServices::FFrame Frame;
+	if (FrameProvider.GetFrameFromTime(ETraceFrameType::TraceFrameType_Game, RewindDebugger->CurrentTraceTime(), Frame))
+	{
+		
+		ChooserProvider->EnumerateChooserEvaluationTimelines([ChooserProvider, RewindDebugger, Frame](uint64 OwnerId, const FChooserProvider::ChooserEvaluationTimeline& ChooserEvaluationTimeline)
+		{
+			double StartTime = RewindDebugger->GetScrubTime();
+			double EndTime = StartTime + Frame.EndTime - Frame.StartTime;
+			
+			ChooserEvaluationTimeline.EnumerateEvents(StartTime, EndTime, [ChooserProvider, StartTime, EndTime, OwnerId](double InStartTime, double InEndTime, uint32 InDepth, const FChooserEvaluationData& ChooserEvaluationData)
+			{
+				if (UObject* ChooserObject = FObjectTrace::GetObjectFromId(ChooserEvaluationData.ChooserId))
+				{
+					if (UChooserTable* Chooser = Cast<UChooserTable>(ChooserObject))
+					{
+						if (Chooser->HasDebugTarget())
+						{
+							if (UObject* ContextObject = FObjectTrace::GetObjectFromId(OwnerId))
+							{
+								if (Chooser->GetDebugTarget() == ContextObject)
+								{
+									Chooser->SetDebugSelectedRow(ChooserEvaluationData.SelectedIndex);
+									Chooser->bDebugTestValuesValid = true;
+
+									ChooserProvider->ReadChooserValueTimeline(OwnerId, [StartTime, EndTime, Chooser, &ChooserEvaluationData](const FChooserProvider::ChooserValueTimeline& ValueTimeline)
+									{
+										ValueTimeline.EnumerateEvents(StartTime,EndTime,[Chooser, ChooserEvaluationData](double StartTime, double EndTime, uint32 Depth, const FChooserValueData& ValueData)
+										{
+											if (ChooserEvaluationData.ChooserId == ValueData.ChooserId)
+											{
+												for(FInstancedStruct& ColumnStruct : Chooser->ColumnsStructs)
+												{
+													FChooserColumnBase& Column = ColumnStruct.GetMutable<FChooserColumnBase>();
+													if(ValueData.Key == Column.GetInputValue()->GetDebugName())
+													{
+														Column.SetTestValue(ValueData.Value);
+													}
+												}
+											}
+											return TraceServices::EEventEnumerate::Continue;
+										});
+									});
+								}
+							}
+						}
+					}
+				}
+				
+				return TraceServices::EEventEnumerate::Continue;
+			});
+			
+
+		});
+	}
+}
+
+void FRewindDebuggerChooser::RecordingStarted(IRewindDebugger*)
+{
+	UE::Trace::ToggleChannel(TEXT("Chooser"), true);
+}
+
+void FRewindDebuggerChooser::RecordingStopped(IRewindDebugger*)
+{
+	UE::Trace::ToggleChannel(TEXT("Chooser"), false);
+}
\ No newline at end of file
diff --git a/Engine/Plugins/Experimental/Chooser/Source/ChooserEditor/Private/RewindDebuggerChooser.h b/Engine/Plugins/Experimental/Chooser/Source/ChooserEditor/Private/RewindDebuggerChooser.h
new file mode 100644
--- /dev/null
+++ b/Engine/Plugins/Experimental/Chooser/Source/ChooserEditor/Private/RewindDebuggerChooser.h
@@ -0,0 +1,19 @@
+// Copyright Epic Games, Inc. All Rights Reserved.
+#pragma once
+
+#include "CoreMinimal.h"
+#include "IRewindDebugger.h"
+#include "IRewindDebuggerExtension.h"
+
+// Rewind debugger extension for Chooser support
+
+class FRewindDebuggerChooser : public IRewindDebuggerExtension
+{
+public:
+	FRewindDebuggerChooser();
+	virtual ~FRewindDebuggerChooser() {};
+
+	virtual void Update(float DeltaTime, IRewindDebugger* RewindDebugger) override;
+	virtual void RecordingStarted(IRewindDebugger* RewindDebugger) override;
+	virtual void RecordingStopped(IRewindDebugger* RewindDebugger) override;
+};
diff --git a/Engine/Source/Editor/RewindDebuggerInterface/Public/IRewindDebugger.h b/Engine/Source/Editor/RewindDebuggerInterface/Public/IRewindDebugger.h
--- a/Engine/Source/Editor/RewindDebuggerInterface/Public/IRewindDebugger.h
+++ b/Engine/Source/Editor/RewindDebuggerInterface/Public/IRewindDebugger.h
@@ -56,6 +56,9 @@ public:
 
 	// get the time the debugger is scrubbed to, in seconds since the capture started (or the recording duration while the game is running)
 	virtual double CurrentTraceTime() const = 0;
+	
+	// get the time the debugger is scrubbed to, in seconds since the recording started
+	virtual double GetScrubTime() const = 0;
 
 	// current visible range in trace/profiler units (same units as CurrentTraceTime)
 	virtual const TRange<double>& GetCurrentTraceRange() const = 0;
