[Backout] - CL28812595
Original CL Desc
-----------------------------------------------------------------
Chooser Debugging with Rewind Debugger

diff --git a/Engine/Plugins/Animation/GameplayInsights/Source/RewindDebugger/Private/RewindDebugger.h b/Engine/Plugins/Animation/GameplayInsights/Source/RewindDebugger/Private/RewindDebugger.h
--- a/Engine/Plugins/Animation/GameplayInsights/Source/RewindDebugger/Private/RewindDebugger.h
+++ b/Engine/Plugins/Animation/GameplayInsights/Source/RewindDebugger/Private/RewindDebugger.h
@@ -33,7 +33,6 @@ public:
 
 	// IRewindDebugger interface
 	virtual double CurrentTraceTime() const override { return TraceTime.Get(); }
-	virtual double GetScrubTime() const override { return CurrentScrubTime; }
 	virtual const TRange<double>& GetCurrentTraceRange() const override { return CurrentTraceRange; }
 	virtual const TRange<double>& GetCurrentViewRange() const override { return CurrentViewRange; }
 	virtual const TraceServices::IAnalysisSession* GetAnalysisSession() const override;
@@ -94,6 +93,7 @@ public:
 
 	bool CanScrub() const;
 	void ScrubToTime(double ScrubTime, bool bIsScrubbing);
+	double GetScrubTime() { return CurrentScrubTime; }
 
 	// Tick function: While recording, update recording duration.  While paused, and we have recorded data, update skinned mesh poses for the current frame, and handle playback.
 	void Tick(float DeltaTime);
diff --git a/Engine/Plugins/Experimental/Chooser/Chooser.uplugin b/Engine/Plugins/Experimental/Chooser/Chooser.uplugin
--- a/Engine/Plugins/Experimental/Chooser/Chooser.uplugin
+++ b/Engine/Plugins/Experimental/Chooser/Chooser.uplugin
@@ -27,10 +27,6 @@
 		{
 			"Name": "BlendStack",
 			"Enabled": true
-		},
-		{
-			"Name": "GameplayInsights",
-			"Enabled": true
 		}
 	],
 	"Modules" :
diff --git a/Engine/Plugins/Experimental/Chooser/Source/Chooser/Chooser.Build.cs b/Engine/Plugins/Experimental/Chooser/Source/Chooser/Chooser.Build.cs
--- a/Engine/Plugins/Experimental/Chooser/Source/Chooser/Chooser.Build.cs
+++ b/Engine/Plugins/Experimental/Chooser/Source/Chooser/Chooser.Build.cs
@@ -23,8 +23,7 @@ namespace UnrealBuildTool.Rules
 					"GameplayTags",
 					"AnimationCore",
 					"AnimGraphRuntime",
-					"BlendStack",
-					"TraceLog"
+					"BlendStack"
 					// ... add private dependencies that you statically link with here ...
 				}
 			);
diff --git a/Engine/Plugins/Experimental/Chooser/Source/Chooser/Internal/BoolColumn.h b/Engine/Plugins/Experimental/Chooser/Source/Chooser/Internal/BoolColumn.h
--- a/Engine/Plugins/Experimental/Chooser/Source/Chooser/Internal/BoolColumn.h
+++ b/Engine/Plugins/Experimental/Chooser/Source/Chooser/Internal/BoolColumn.h
@@ -6,7 +6,6 @@
 #include "IChooserParameterBool.h"
 #include "ChooserPropertyAccess.h"
 #include "InstancedStruct.h"
-#include "Serialization/MemoryReader.h"
 #include "BoolColumn.generated.h"
 
 UENUM()
@@ -85,17 +84,11 @@ struct CHOOSER_API FBoolColumn : public FChooserColumnBase
 	virtual void Filter(FChooserEvaluationContext& Context, const TArray<uint32>& IndexListIn, TArray<uint32>& IndexListOut) const override;
 
 #if WITH_EDITOR
-	mutable bool TestValue = false;
+	mutable bool TestValue;
 	virtual bool EditorTestFilter(int32 RowIndex) const override
 	{
 		return RowValuesWithAny.IsValidIndex(RowIndex) && (RowValuesWithAny[RowIndex] == EBoolColumnCellValue::MatchAny || TestValue == static_cast<bool>(RowValuesWithAny[RowIndex]));
 	}
-	
-	virtual void SetTestValue(TArrayView<const uint8> Value) override
-	{
-		FMemoryReaderView Reader(Value);
-		Reader << TestValue;
-	}
 #endif
 
 	virtual void PostLoad() override
diff --git a/Engine/Plugins/Experimental/Chooser/Source/Chooser/Internal/Chooser.h b/Engine/Plugins/Experimental/Chooser/Source/Chooser/Internal/Chooser.h
--- a/Engine/Plugins/Experimental/Chooser/Source/Chooser/Internal/Chooser.h
+++ b/Engine/Plugins/Experimental/Chooser/Source/Chooser/Internal/Chooser.h
@@ -34,22 +34,10 @@ public:
 	
 	void SetDebugSelectedRow(int32 Index) const { DebugSelectedRow = Index; }
 	int32 GetDebugSelectedRow() const { return DebugSelectedRow; }
-	bool HasDebugTarget() const { return !DebugTargetName.IsEmpty(); }
+	bool HasDebugTarget() const { return DebugTarget != nullptr; }
 	const UObject* GetDebugTarget() const { return DebugTarget.Get(); }
-	const FString& GetDebugTargetName() const { return DebugTargetName; }
-	void SetDebugTarget(TWeakObjectPtr<const UObject> Target)
-	{
-		DebugTarget = Target;
-		if (const UObject* DebugObject = DebugTarget.Get())
-		{
-			DebugTargetName = DebugObject->GetName();
-		}
-	}
-	void ResetDebugTarget()
-	{
-		DebugTarget.Reset();
-		DebugTargetName = "";
-	}
+	void SetDebugTarget(TWeakObjectPtr<const UObject> Target) { DebugTarget = Target; }
+	void ResetDebugTarget() { DebugTarget.Reset(); }
 	void IterateRecentContextObjects(TFunction<void(const UObject*)> Callback) const;
 	void UpdateDebugging(FChooserEvaluationContext& Context) const;
 	
@@ -66,9 +54,7 @@ private:
 	mutable TSet<TWeakObjectPtr<const UObject>> RecentContextObjects;
 	mutable FCriticalSection DebugLock;
 	// reference to the UObject in PIE  which we want to get debug info for
-	mutable TWeakObjectPtr<const UObject> DebugTarget;
-	FString DebugTargetName;
-	
+	TWeakObjectPtr<const UObject> DebugTarget;
 	// Row which was selected last time this chooser was evaluated on DebugTarget
 	mutable int32 DebugSelectedRow = -1;
 
diff --git a/Engine/Plugins/Experimental/Chooser/Source/Chooser/Internal/ChooserTrace.h b/Engine/Plugins/Experimental/Chooser/Source/Chooser/Internal/ChooserTrace.h
deleted file mode 100644
diff --git a/Engine/Plugins/Experimental/Chooser/Source/Chooser/Internal/EnumColumn.h b/Engine/Plugins/Experimental/Chooser/Source/Chooser/Internal/EnumColumn.h
--- a/Engine/Plugins/Experimental/Chooser/Source/Chooser/Internal/EnumColumn.h
+++ b/Engine/Plugins/Experimental/Chooser/Source/Chooser/Internal/EnumColumn.h
@@ -6,7 +6,6 @@
 #include "IChooserParameterEnum.h"
 #include "InstancedStruct.h"
 #include "ChooserPropertyAccess.h"
-#include "Serialization/MemoryReader.h"
 #include "EnumColumn.generated.h"
 
 struct FBindingChainElement;
@@ -120,17 +119,11 @@ public:
 	virtual void Filter(FChooserEvaluationContext& Context, const TArray<uint32>& IndexListIn, TArray<uint32>& IndexListOut) const override;
 	
 #if WITH_EDITOR
-	mutable uint8 TestValue = 0;
+	mutable int32 TestValue;
 	virtual bool EditorTestFilter(int32 RowIndex) const override
 	{
 		return RowValues.IsValidIndex(RowIndex) && RowValues[RowIndex].Evaluate(TestValue);
 	}
-	
-	virtual void SetTestValue(TArrayView<const uint8> Value) override
-	{
-		FMemoryReaderView Reader(Value);
-		Reader << TestValue;
-	}
 #endif
 	
 	CHOOSER_COLUMN_BOILERPLATE(FChooserParameterEnumBase);
diff --git a/Engine/Plugins/Experimental/Chooser/Source/Chooser/Internal/FloatRangeColumn.h b/Engine/Plugins/Experimental/Chooser/Source/Chooser/Internal/FloatRangeColumn.h
--- a/Engine/Plugins/Experimental/Chooser/Source/Chooser/Internal/FloatRangeColumn.h
+++ b/Engine/Plugins/Experimental/Chooser/Source/Chooser/Internal/FloatRangeColumn.h
@@ -8,7 +8,6 @@
 #include "IChooserParameterFloat.h"
 #include "ChooserPropertyAccess.h"
 #include "InstancedStruct.h"
-#include "Serialization/MemoryReader.h"
 #include "FloatRangeColumn.generated.h"
 
 USTRUCT(DisplayName = "Float Property Binding")
@@ -90,17 +89,11 @@ struct CHOOSER_API FFloatRangeColumn : public FChooserColumnBase
 	virtual void Filter(FChooserEvaluationContext& Context, const TArray<uint32>& IndexListIn, TArray<uint32>& IndexListOut) const override;
 
 #if WITH_EDITOR
-	mutable double TestValue = 0.0;
+	mutable float TestValue;
 	virtual bool EditorTestFilter(int32 RowIndex) const override
 	{
 		return RowValues.IsValidIndex(RowIndex) && TestValue >= RowValues[RowIndex].Min && TestValue <= RowValues[RowIndex].Max;
 	}
-	
-	virtual void SetTestValue(TArrayView<const uint8> Value) override
-	{
-		FMemoryReaderView Reader(Value);
-		Reader << TestValue;
-	}
 #endif
 	
 	virtual void PostLoad() override
diff --git a/Engine/Plugins/Experimental/Chooser/Source/Chooser/Internal/GameplayTagColumn.h b/Engine/Plugins/Experimental/Chooser/Source/Chooser/Internal/GameplayTagColumn.h
--- a/Engine/Plugins/Experimental/Chooser/Source/Chooser/Internal/GameplayTagColumn.h
+++ b/Engine/Plugins/Experimental/Chooser/Source/Chooser/Internal/GameplayTagColumn.h
@@ -7,7 +7,6 @@
 #include "ChooserPropertyAccess.h"
 #include "GameplayTagContainer.h"
 #include "InstancedStruct.h"
-#include "Serialization/MemoryReader.h"
 #include "GameplayTagColumn.generated.h"
 
 struct FBindingChainElement;
@@ -98,14 +97,6 @@ struct CHOOSER_API FGameplayTagColumn : public FChooserColumnBase
 		}
 		return false;
 	}
-
-	virtual void SetTestValue(TArrayView<const uint8> Value) override
-    {
-    	FMemoryReaderView Reader(Value);
-		FString Tags;
-    	Reader << Tags;
-		TestValue.FromExportString(Tags);
-    }
 #endif
 
 	virtual void PostLoad() override
diff --git a/Engine/Plugins/Experimental/Chooser/Source/Chooser/Internal/ObjectColumn.h b/Engine/Plugins/Experimental/Chooser/Source/Chooser/Internal/ObjectColumn.h
--- a/Engine/Plugins/Experimental/Chooser/Source/Chooser/Internal/ObjectColumn.h
+++ b/Engine/Plugins/Experimental/Chooser/Source/Chooser/Internal/ObjectColumn.h
@@ -6,7 +6,6 @@
 #include "IChooserColumn.h"
 #include "IChooserParameterObject.h"
 #include "InstancedStruct.h"
-#include "Serialization/MemoryReader.h"
 #include "ObjectColumn.generated.h"
 
 struct FBindingChainElement;
@@ -91,13 +90,6 @@ struct CHOOSER_API FObjectColumn : public FChooserColumnBase
 	{
 		return RowValues.IsValidIndex(RowIndex) && RowValues[RowIndex].Evaluate(TestValue);
 	}
-	virtual void SetTestValue(TArrayView<const uint8> Value) override
-	{
-		FMemoryReaderView Reader(Value);
-		FString Path;
-		Reader << Path;
-		TestValue.SetPath(Path);
-	}
 #endif
 
 	CHOOSER_COLUMN_BOILERPLATE(FChooserParameterObjectBase);
diff --git a/Engine/Plugins/Experimental/Chooser/Source/Chooser/Private/BoolColumn.cpp b/Engine/Plugins/Experimental/Chooser/Source/Chooser/Private/BoolColumn.cpp
--- a/Engine/Plugins/Experimental/Chooser/Source/Chooser/Private/BoolColumn.cpp
+++ b/Engine/Plugins/Experimental/Chooser/Source/Chooser/Private/BoolColumn.cpp
@@ -2,7 +2,6 @@
 #include "BoolColumn.h"
 #include "ChooserPropertyAccess.h"
 #include "Chooser.h"
-#include "ChooserTrace.h"
 
 bool FBoolContextProperty::GetValue(FChooserEvaluationContext& Context, bool& OutResult) const
 {
@@ -26,8 +25,6 @@ void FBoolColumn::Filter(FChooserEvaluationContext& Context, const TArray<uint32
 		bool Result = false;
 		InputValue.Get<FChooserParameterBoolBase>().GetValue(Context,Result);
 
-		TRACE_CHOOSER_VALUE(Context, ToCStr(InputValue.Get<FChooserParameterBase>().GetDebugName()), Result);
-
 	#if WITH_EDITOR
 		if (Context.DebuggingInfo.bCurrentDebugTarget)
 		{
diff --git a/Engine/Plugins/Experimental/Chooser/Source/Chooser/Private/Chooser.cpp b/Engine/Plugins/Experimental/Chooser/Source/Chooser/Private/Chooser.cpp
--- a/Engine/Plugins/Experimental/Chooser/Source/Chooser/Private/Chooser.cpp
+++ b/Engine/Plugins/Experimental/Chooser/Source/Chooser/Private/Chooser.cpp
@@ -2,8 +2,6 @@
 #include "Chooser.h"
 #include "ChooserFunctionLibrary.h"
 #include "ChooserPropertyAccess.h"
-#include "ChooserTrace.h"
-#include "ObjectTrace.h"
 #include "Engine/UserDefinedStruct.h"
 #include "Engine/Blueprint.h"
 
@@ -183,30 +181,17 @@ void UChooserTable::IterateRecentContextObjects(TFunction<void(const UObject*)>
 void UChooserTable::UpdateDebugging(FChooserEvaluationContext& Context) const
 {
 	FScopeLock Lock(&DebugLock);
-	
 	for (const FInstancedStruct& Param : Context.Params)
 	{
 		if (const FChooserEvaluationInputObject* ObjectParam = Param.GetPtr<FChooserEvaluationInputObject>())
 		{
-			if (UObject* ContextObject = ObjectParam->Object)
+			UObject* ContextObject = ObjectParam->Object;
+			RecentContextObjects.Add(MakeWeakObjectPtr(ContextObject));
+			if (ContextObject == DebugTarget)
 			{
-				RecentContextObjects.Add(MakeWeakObjectPtr(ContextObject));
-				
-				if (DebugTarget == nullptr && !DebugTargetName.IsEmpty())
-				{
-					// if the DebugTargetName is set, but not the DebugTarget, it means that PIE has been restarted, so try matching by name
-					if (ContextObject->GetName() == DebugTargetName)
-					{
-						DebugTarget = ContextObject;
-					}
-				}
-				
-				if (ContextObject == DebugTarget)
-				{
-					bDebugTestValuesValid = true;
-					Context.DebuggingInfo.bCurrentDebugTarget = true;
-					return;
-				}
+				bDebugTestValuesValid = true;
+				Context.DebuggingInfo.bCurrentDebugTarget = true;
+				return;
 			}
 		}
 	}
@@ -232,11 +217,6 @@ FObjectChooserBase::EIteratorStatus UChooserTable::EvaluateChooser(FChooserEvalu
 		Chooser->UpdateDebugging(Context);
 	}
 #endif
-	
-#if CHOOSER_DEBUGGING_ENABLED
-	Context.DebuggingInfo.CurrentChooser = Chooser;
-#endif
-	
 
 	TArray<uint32> Indices1;
 	TArray<uint32> Indices2;
@@ -272,8 +252,7 @@ FObjectChooserBase::EIteratorStatus UChooserTable::EvaluateChooser(FChooserEvalu
 			Chooser->SetDebugSelectedRow(-1);
 		}
 		#endif
-		TRACE_CHOOSER_EVALUATION(Chooser, Context, -1);
-	
+		
 		if (Chooser->FallbackResult.IsValid())
 		{
 			const FObjectChooserBase& SelectedResult = Chooser->FallbackResult.Get<FObjectChooserBase>();
@@ -312,13 +291,12 @@ FObjectChooserBase::EIteratorStatus UChooserTable::EvaluateChooser(FChooserEvalu
 						const FChooserColumnBase& Column = ColumnData.Get<FChooserColumnBase>();
 						Column.SetOutputs(Context, SelectedIndex);
 					}
-					#if WITH_EDITOR
+	#if WITH_EDITOR
 					if (Context.DebuggingInfo.bCurrentDebugTarget)
 					{
 						Chooser->SetDebugSelectedRow(SelectedIndex);
 					}
-					#endif
-					TRACE_CHOOSER_EVALUATION(Chooser, Context, SelectedIndex);
+	#endif
 				}
 				if (Status == FObjectChooserBase::EIteratorStatus::Stop)
 				{
diff --git a/Engine/Plugins/Experimental/Chooser/Source/Chooser/Private/ChooserTrace.cpp b/Engine/Plugins/Experimental/Chooser/Source/Chooser/Private/ChooserTrace.cpp
deleted file mode 100644
diff --git a/Engine/Plugins/Experimental/Chooser/Source/Chooser/Private/EnumColumn.cpp b/Engine/Plugins/Experimental/Chooser/Source/Chooser/Private/EnumColumn.cpp
--- a/Engine/Plugins/Experimental/Chooser/Source/Chooser/Private/EnumColumn.cpp
+++ b/Engine/Plugins/Experimental/Chooser/Source/Chooser/Private/EnumColumn.cpp
@@ -1,7 +1,6 @@
 // Copyright Epic Games, Inc. All Rights Reserved.
 #include "EnumColumn.h"
 #include "ChooserPropertyAccess.h"
-#include "ChooserTrace.h"
 
 
 bool FEnumContextProperty::GetValue(FChooserEvaluationContext& Context, uint8& OutResult) const
@@ -74,8 +73,6 @@ void FEnumColumn::Filter(FChooserEvaluationContext& Context, const TArray<uint32
 			TestValue = Result;
 		}
 #endif
-
-		TRACE_CHOOSER_VALUE(Context, ToCStr(InputValue.Get<FChooserParameterBase>().GetDebugName()), Result);
 		
 		for (const uint32 Index : IndexListIn)
 		{
diff --git a/Engine/Plugins/Experimental/Chooser/Source/Chooser/Private/FloatRangeColumn.cpp b/Engine/Plugins/Experimental/Chooser/Source/Chooser/Private/FloatRangeColumn.cpp
--- a/Engine/Plugins/Experimental/Chooser/Source/Chooser/Private/FloatRangeColumn.cpp
+++ b/Engine/Plugins/Experimental/Chooser/Source/Chooser/Private/FloatRangeColumn.cpp
@@ -1,7 +1,6 @@
 // Copyright Epic Games, Inc. All Rights Reserved.
 #include "FloatRangeColumn.h"
 #include "ChooserPropertyAccess.h"
-#include "ChooserTrace.h"
 
 
 
@@ -26,8 +25,6 @@ void FFloatRangeColumn::Filter(FChooserEvaluationContext& Context, const TArray<
 	{
 		double Result = 0.0f;
 		InputValue.Get<FChooserParameterFloatBase>().GetValue(Context, Result);
-		
-		TRACE_CHOOSER_VALUE(Context, ToCStr(InputValue.Get<FChooserParameterBase>().GetDebugName()), Result);
 
 #if WITH_EDITOR
 		if (Context.DebuggingInfo.bCurrentDebugTarget)
diff --git a/Engine/Plugins/Experimental/Chooser/Source/Chooser/Private/GameplayTagColumn.cpp b/Engine/Plugins/Experimental/Chooser/Source/Chooser/Private/GameplayTagColumn.cpp
--- a/Engine/Plugins/Experimental/Chooser/Source/Chooser/Private/GameplayTagColumn.cpp
+++ b/Engine/Plugins/Experimental/Chooser/Source/Chooser/Private/GameplayTagColumn.cpp
@@ -1,7 +1,6 @@
 // Copyright Epic Games, Inc. All Rights Reserved.
 #include "GameplayTagColumn.h"
 #include "ChooserPropertyAccess.h"
-#include "ChooserTrace.h"
 
 bool FGameplayTagContextProperty::GetValue(FChooserEvaluationContext& Context, const FGameplayTagContainer*& OutResult) const
 {
@@ -18,8 +17,6 @@ void FGameplayTagColumn::Filter(FChooserEvaluationContext& Context, const TArray
 	const FGameplayTagContainer* Result = nullptr;
 	if (InputValue.IsValid() && InputValue.Get<FChooserParameterGameplayTagBase>().GetValue(Context,Result))
 	{
-		TRACE_CHOOSER_VALUE(Context, ToCStr(InputValue.Get<FChooserParameterBase>().GetDebugName()), Result->ToString());
-
 #if WITH_EDITOR
 		if (Context.DebuggingInfo.bCurrentDebugTarget)
 		{
diff --git a/Engine/Plugins/Experimental/Chooser/Source/Chooser/Private/ObjectColumn.cpp b/Engine/Plugins/Experimental/Chooser/Source/Chooser/Private/ObjectColumn.cpp
--- a/Engine/Plugins/Experimental/Chooser/Source/Chooser/Private/ObjectColumn.cpp
+++ b/Engine/Plugins/Experimental/Chooser/Source/Chooser/Private/ObjectColumn.cpp
@@ -1,7 +1,6 @@
 // Copyright Epic Games, Inc. All Rights Reserved.
 #include "ObjectColumn.h"
 #include "ChooserPropertyAccess.h"
-#include "ChooserTrace.h"
 
 #if WITH_EDITOR
 	#include "IPropertyAccessEditor.h"
@@ -63,8 +62,7 @@ void FObjectColumn::Filter(FChooserEvaluationContext& Context, const TArray<uint
 	if (InputValue.IsValid() &&
 		InputValue.Get<FChooserParameterObjectBase>().GetValue(Context, Result))
 	{
-		TRACE_CHOOSER_VALUE(Context, ToCStr(InputValue.Get<FChooserParameterBase>().GetDebugName()), Result.ToString());
-	
+		
 #if WITH_EDITOR
 		if (Context.DebuggingInfo.bCurrentDebugTarget)
 		{
diff --git a/Engine/Plugins/Experimental/Chooser/Source/Chooser/Public/IChooserColumn.h b/Engine/Plugins/Experimental/Chooser/Source/Chooser/Public/IChooserColumn.h
--- a/Engine/Plugins/Experimental/Chooser/Source/Chooser/Public/IChooserColumn.h
+++ b/Engine/Plugins/Experimental/Chooser/Source/Chooser/Public/IChooserColumn.h
@@ -24,7 +24,6 @@ public:
 };
 
 class UChooserTable;
-class FVariant;
 
 USTRUCT()
 struct CHOOSER_API FChooserColumnBase
@@ -49,7 +48,6 @@ public:
 		}
 	};
 
-	virtual void SetTestValue(TArrayView<const uint8> Value) { }
 	virtual FChooserParameterBase* GetInputValue() { return nullptr; };
 #if WITH_EDITOR
 	virtual FName RowValuesPropertyName() { return FName(); }
diff --git a/Engine/Plugins/Experimental/Chooser/Source/Chooser/Public/IChooserParameterBase.h b/Engine/Plugins/Experimental/Chooser/Source/Chooser/Public/IChooserParameterBase.h
--- a/Engine/Plugins/Experimental/Chooser/Source/Chooser/Public/IChooserParameterBase.h
+++ b/Engine/Plugins/Experimental/Chooser/Source/Chooser/Public/IChooserParameterBase.h
@@ -12,13 +12,6 @@ struct FChooserParameterBase
 {
 	GENERATED_BODY()
 
-	virtual FString GetDebugName() const
-	{
-		FText Name;
-		GetDisplayName(Name);
-		return Name.ToString();
-	}
-	
 	virtual void GetDisplayName(FText& OutName) const { }
 
 	virtual void PostLoad() {};
diff --git a/Engine/Plugins/Experimental/Chooser/Source/Chooser/Public/IObjectChooser.h b/Engine/Plugins/Experimental/Chooser/Source/Chooser/Public/IObjectChooser.h
--- a/Engine/Plugins/Experimental/Chooser/Source/Chooser/Public/IObjectChooser.h
+++ b/Engine/Plugins/Experimental/Chooser/Source/Chooser/Public/IObjectChooser.h
@@ -8,14 +8,6 @@
 #include "InstancedStruct.h"
 #include "IObjectChooser.generated.h"
 
-#if UE_TRACE_ENABLED && !IS_PROGRAM && !UE_BUILD_SHIPPING && !UE_BUILD_TEST
-#define CHOOSER_TRACE_ENABLED 1
-#else
-#define CHOOSER_TRACE_ENABLED 0
-#endif
-
-#define CHOOSER_DEBUGGING_ENABLED ((CHOOSER_TRACE_ENABLED) || (WITH_EDITOR))
-
 UINTERFACE(NotBlueprintType, meta = (CannotImplementInterfaceInBlueprint))
 class CHOOSER_API UObjectChooser : public UInterface
 {
@@ -29,10 +21,9 @@ public:
 	virtual void ConvertToInstancedStruct(FInstancedStruct& OutInstancedStruct) const { }
 };
 
-#if CHOOSER_DEBUGGING_ENABLED
+#if WITH_EDITOR
 struct CHOOSER_API FChooserDebuggingInfo
 {
-	const UObject* CurrentChooser = nullptr;
 	bool bCurrentDebugTarget = false;
 };
 #endif
@@ -60,7 +51,7 @@ struct CHOOSER_API FChooserEvaluationContext
 		Params.Last().GetMutable<FChooserEvaluationInputObject>().Object = ContextObject;
 	}
 
-	#if CHOOSER_DEBUGGING_ENABLED
+	#if WITH_EDITOR
     	FChooserDebuggingInfo DebuggingInfo;
     #endif
 	
diff --git a/Engine/Plugins/Experimental/Chooser/Source/ChooserEditor/ChooserEditor.Build.cs b/Engine/Plugins/Experimental/Chooser/Source/ChooserEditor/ChooserEditor.Build.cs
--- a/Engine/Plugins/Experimental/Chooser/Source/ChooserEditor/ChooserEditor.Build.cs
+++ b/Engine/Plugins/Experimental/Chooser/Source/ChooserEditor/ChooserEditor.Build.cs
@@ -32,12 +32,6 @@ namespace UnrealBuildTool.Rules
 					"StructUtils",
 					"KismetCompiler",
 					"BlendStack",
-					"TraceAnalysis",
-					"TraceLog",
-					"TraceServices",
-					"TraceInsights",
-					"RewindDebuggerInterface",
-					"GameplayInsights",
 					// ... add private dependencies that you statically link with here ...
 				}
 			);
diff --git a/Engine/Plugins/Experimental/Chooser/Source/ChooserEditor/Private/ChooserAnalyzer.cpp b/Engine/Plugins/Experimental/Chooser/Source/ChooserEditor/Private/ChooserAnalyzer.cpp
deleted file mode 100644
diff --git a/Engine/Plugins/Experimental/Chooser/Source/ChooserEditor/Private/ChooserAnalyzer.h b/Engine/Plugins/Experimental/Chooser/Source/ChooserEditor/Private/ChooserAnalyzer.h
deleted file mode 100644
diff --git a/Engine/Plugins/Experimental/Chooser/Source/ChooserEditor/Private/ChooserEditorModule.cpp b/Engine/Plugins/Experimental/Chooser/Source/ChooserEditor/Private/ChooserEditorModule.cpp
--- a/Engine/Plugins/Experimental/Chooser/Source/ChooserEditor/Private/ChooserEditorModule.cpp
+++ b/Engine/Plugins/Experimental/Chooser/Source/ChooserEditor/Private/ChooserEditorModule.cpp
@@ -3,31 +3,28 @@
 #include "ChooserEditorModule.h"
 
 #include "AnimNode_ChooserPlayer.h"
-#include "BoolColumnEditor.h"
-#include "ChooserPropertyAccess.h"
+#include "IAssetTools.h"
 #include "ChooserTableEditor.h"
-#include "ChooserTableEditorCommands.h"
-#include "CurveOverrideCustomization.h"
+#include "BoolColumnEditor.h"
 #include "EnumColumnEditor.h"
-#include "Features/IModularFeatures.h"
 #include "FloatRangeColumnEditor.h"
-#include "FrameTimeCustomization.h"
+#include "OutputFloatColumnEditor.h"
 #include "GameplayTagColumnEditor.h"
-#include "IAssetTools.h"
 #include "ObjectColumnEditor.h"
-#include "OutputFloatColumnEditor.h"
-#include "OutputStructColumnEditor.h"
-#include "PropertyAccessChainCustomization.h"
+#include "ChooserTableEditorCommands.h"
+#include "ChooserPropertyAccess.h"
 #include "PropertyEditorModule.h"
 #include "RandomizeColumnEditor.h"
-#include "ChooserTrack.h"
+#include "OutputStructColumnEditor.h"
+#include "CurveOverrideCustomization.h"
+#include "FrameTimeCustomization.h"
 
 #define LOCTEXT_NAMESPACE "ChooserEditorModule"
 
 namespace UE::ChooserEditor
 {
 	
-FChoosersTrackCreator GChoosersTrackCreator;
+
 
 void FModule::StartupModule()
 {
@@ -54,16 +51,10 @@ void FModule::StartupModule()
 	PropertyModule.RegisterCustomPropertyTypeLayout(FChooserObjectPropertyBinding::StaticStruct()->GetFName(), FOnGetPropertyTypeCustomizationInstance::CreateLambda([] { return MakeShared<FPropertyAccessChainCustomization>(); }));
 	PropertyModule.RegisterCustomPropertyTypeLayout(FChooserStructPropertyBinding::StaticStruct()->GetFName(), FOnGetPropertyTypeCustomizationInstance::CreateLambda([] { return MakeShared<FPropertyAccessChainCustomization>(); }));
 
-	IModularFeatures::Get().RegisterModularFeature( IRewindDebuggerExtension::ModularFeatureName, &RewindDebuggerChooser);
-	IModularFeatures::Get().RegisterModularFeature(RewindDebugger::IRewindDebuggerTrackCreator::ModularFeatureName, &GChoosersTrackCreator);
-	IModularFeatures::Get().RegisterModularFeature(TraceServices::ModuleFeatureName, &ChooserTraceModule);
 }
 
 void FModule::ShutdownModule()
 {
-	IModularFeatures::Get().UnregisterModularFeature( IRewindDebuggerExtension::ModularFeatureName, &RewindDebuggerChooser);
-	IModularFeatures::Get().UnregisterModularFeature(RewindDebugger::IRewindDebuggerTrackCreator::ModularFeatureName, &GChoosersTrackCreator);
-	IModularFeatures::Get().UnregisterModularFeature(TraceServices::ModuleFeatureName, &ChooserTraceModule);
 	FChooserTableEditorCommands::Unregister();
 }
 
diff --git a/Engine/Plugins/Experimental/Chooser/Source/ChooserEditor/Private/ChooserEditorModule.h b/Engine/Plugins/Experimental/Chooser/Source/ChooserEditor/Private/ChooserEditorModule.h
--- a/Engine/Plugins/Experimental/Chooser/Source/ChooserEditor/Private/ChooserEditorModule.h
+++ b/Engine/Plugins/Experimental/Chooser/Source/ChooserEditor/Private/ChooserEditorModule.h
@@ -3,9 +3,8 @@
 #pragma once
 
 #include "CoreMinimal.h"
-#include "ChooserTraceModule.h"
+#include "PropertyAccessChainCustomization.h"
 #include "Modules/ModuleInterface.h"
-#include "RewindDebuggerChooser.h"
 
 namespace UE::ChooserEditor
 {
@@ -15,9 +14,6 @@ class FModule : public IModuleInterface
 private:
 	virtual void StartupModule() override;
 	virtual void ShutdownModule() override;
-	
-	FRewindDebuggerChooser RewindDebuggerChooser;
-	FChooserTraceModule ChooserTraceModule;
 };
 
 }
\ No newline at end of file
diff --git a/Engine/Plugins/Experimental/Chooser/Source/ChooserEditor/Private/ChooserProvider.cpp b/Engine/Plugins/Experimental/Chooser/Source/ChooserEditor/Private/ChooserProvider.cpp
deleted file mode 100644
diff --git a/Engine/Plugins/Experimental/Chooser/Source/ChooserEditor/Private/ChooserProvider.h b/Engine/Plugins/Experimental/Chooser/Source/ChooserEditor/Private/ChooserProvider.h
deleted file mode 100644
diff --git a/Engine/Plugins/Experimental/Chooser/Source/ChooserEditor/Private/ChooserTableEditor.cpp b/Engine/Plugins/Experimental/Chooser/Source/ChooserEditor/Private/ChooserTableEditor.cpp
--- a/Engine/Plugins/Experimental/Chooser/Source/ChooserEditor/Private/ChooserTableEditor.cpp
+++ b/Engine/Plugins/Experimental/Chooser/Source/ChooserEditor/Private/ChooserTableEditor.cpp
@@ -217,7 +217,7 @@ void FChooserTableEditor::RegisterToolbar()
 						{
 							if (Chooser->HasDebugTarget())
 							{
-								return  FText::FromString(Chooser->GetDebugTargetName());
+								return  FText::FromString(Chooser->GetDebugTarget()->GetName());
 							}
 							else
 							{
diff --git a/Engine/Plugins/Experimental/Chooser/Source/ChooserEditor/Private/ChooserTraceModule.cpp b/Engine/Plugins/Experimental/Chooser/Source/ChooserEditor/Private/ChooserTraceModule.cpp
deleted file mode 100644
diff --git a/Engine/Plugins/Experimental/Chooser/Source/ChooserEditor/Private/ChooserTraceModule.h b/Engine/Plugins/Experimental/Chooser/Source/ChooserEditor/Private/ChooserTraceModule.h
deleted file mode 100644
diff --git a/Engine/Plugins/Experimental/Chooser/Source/ChooserEditor/Private/ChooserTrack.cpp b/Engine/Plugins/Experimental/Chooser/Source/ChooserEditor/Private/ChooserTrack.cpp
deleted file mode 100644
diff --git a/Engine/Plugins/Experimental/Chooser/Source/ChooserEditor/Private/ChooserTrack.h b/Engine/Plugins/Experimental/Chooser/Source/ChooserEditor/Private/ChooserTrack.h
deleted file mode 100644
diff --git a/Engine/Plugins/Experimental/Chooser/Source/ChooserEditor/Private/RewindDebuggerChooser.cpp b/Engine/Plugins/Experimental/Chooser/Source/ChooserEditor/Private/RewindDebuggerChooser.cpp
deleted file mode 100644
diff --git a/Engine/Plugins/Experimental/Chooser/Source/ChooserEditor/Private/RewindDebuggerChooser.h b/Engine/Plugins/Experimental/Chooser/Source/ChooserEditor/Private/RewindDebuggerChooser.h
deleted file mode 100644
diff --git a/Engine/Source/Editor/RewindDebuggerInterface/Public/IRewindDebugger.h b/Engine/Source/Editor/RewindDebuggerInterface/Public/IRewindDebugger.h
--- a/Engine/Source/Editor/RewindDebuggerInterface/Public/IRewindDebugger.h
+++ b/Engine/Source/Editor/RewindDebuggerInterface/Public/IRewindDebugger.h
@@ -56,9 +56,6 @@ public:
 
 	// get the time the debugger is scrubbed to, in seconds since the capture started (or the recording duration while the game is running)
 	virtual double CurrentTraceTime() const = 0;
-	
-	// get the time the debugger is scrubbed to, in seconds since the recording started
-	virtual double GetScrubTime() const = 0;
 
 	// current visible range in trace/profiler units (same units as CurrentTraceTime)
 	virtual const TRange<double>& GetCurrentTraceRange() const = 0;
