Allocate temporary memory for chooser evaluation on the stack

diff --git a/Engine/Plugins/Experimental/Chooser/Source/Chooser/Internal/BoolColumn.h b/Engine/Plugins/Experimental/Chooser/Source/Chooser/Internal/BoolColumn.h
--- a/Engine/Plugins/Experimental/Chooser/Source/Chooser/Internal/BoolColumn.h
+++ b/Engine/Plugins/Experimental/Chooser/Source/Chooser/Internal/BoolColumn.h
@@ -82,7 +82,7 @@ struct CHOOSER_API FBoolColumn : public FChooserColumnBase
 	UPROPERTY(EditAnywhere, Category= "Data", DisplayName="RowValues");
 	TArray<EBoolColumnCellValue> RowValuesWithAny; 
 	
-	virtual void Filter(FChooserEvaluationContext& Context, const TArray<uint32>& IndexListIn, TArray<uint32>& IndexListOut) const override;
+	virtual void Filter(FChooserEvaluationContext& Context, const FChooserIndexArray& IndexListIn, FChooserIndexArray& IndexListOut) const override;
 
 #if WITH_EDITOR
 	mutable bool TestValue = false;
diff --git a/Engine/Plugins/Experimental/Chooser/Source/Chooser/Internal/EnumColumn.h b/Engine/Plugins/Experimental/Chooser/Source/Chooser/Internal/EnumColumn.h
--- a/Engine/Plugins/Experimental/Chooser/Source/Chooser/Internal/EnumColumn.h
+++ b/Engine/Plugins/Experimental/Chooser/Source/Chooser/Internal/EnumColumn.h
@@ -117,7 +117,7 @@ public:
 	// should match the length of the Results array
 	TArray<FChooserEnumRowData> RowValues;
 
-	virtual void Filter(FChooserEvaluationContext& Context, const TArray<uint32>& IndexListIn, TArray<uint32>& IndexListOut) const override;
+	virtual void Filter(FChooserEvaluationContext& Context, const FChooserIndexArray& IndexListIn, FChooserIndexArray& IndexListOut) const override;
 	
 #if WITH_EDITOR
 	mutable uint8 TestValue = 0;
diff --git a/Engine/Plugins/Experimental/Chooser/Source/Chooser/Internal/FloatRangeColumn.h b/Engine/Plugins/Experimental/Chooser/Source/Chooser/Internal/FloatRangeColumn.h
--- a/Engine/Plugins/Experimental/Chooser/Source/Chooser/Internal/FloatRangeColumn.h
+++ b/Engine/Plugins/Experimental/Chooser/Source/Chooser/Internal/FloatRangeColumn.h
@@ -87,7 +87,7 @@ struct CHOOSER_API FFloatRangeColumn : public FChooserColumnBase
 	// should match the length of the Results array 
 	TArray<FChooserFloatRangeRowData> RowValues;
 	
-	virtual void Filter(FChooserEvaluationContext& Context, const TArray<uint32>& IndexListIn, TArray<uint32>& IndexListOut) const override;
+	virtual void Filter(FChooserEvaluationContext& Context, const FChooserIndexArray& IndexListIn, FChooserIndexArray& IndexListOut) const override;
 
 #if WITH_EDITOR
 	mutable double TestValue = 0.0;
diff --git a/Engine/Plugins/Experimental/Chooser/Source/Chooser/Internal/GameplayTagColumn.h b/Engine/Plugins/Experimental/Chooser/Source/Chooser/Internal/GameplayTagColumn.h
--- a/Engine/Plugins/Experimental/Chooser/Source/Chooser/Internal/GameplayTagColumn.h
+++ b/Engine/Plugins/Experimental/Chooser/Source/Chooser/Internal/GameplayTagColumn.h
@@ -79,7 +79,7 @@ struct CHOOSER_API FGameplayTagColumn : public FChooserColumnBase
 	// should match the length of the Results array 
 	TArray<FGameplayTagContainer> RowValues;
 	
-	virtual void Filter(FChooserEvaluationContext& Context, const TArray<uint32>& IndexListIn, TArray<uint32>& IndexListOut) const override;
+	virtual void Filter(FChooserEvaluationContext& Context, const FChooserIndexArray& IndexListIn, FChooserIndexArray& IndexListOut) const override;
 
 #if WITH_EDITOR
 	mutable FGameplayTagContainer TestValue;
diff --git a/Engine/Plugins/Experimental/Chooser/Source/Chooser/Internal/ObjectColumn.h b/Engine/Plugins/Experimental/Chooser/Source/Chooser/Internal/ObjectColumn.h
--- a/Engine/Plugins/Experimental/Chooser/Source/Chooser/Internal/ObjectColumn.h
+++ b/Engine/Plugins/Experimental/Chooser/Source/Chooser/Internal/ObjectColumn.h
@@ -83,7 +83,7 @@ struct CHOOSER_API FObjectColumn : public FChooserColumnBase
 	// should match the length of the Results array
 	TArray<FChooserObjectRowData> RowValues;
 
-	virtual void Filter(FChooserEvaluationContext& Context, const TArray<uint32>& IndexListIn, TArray<uint32>& IndexListOut) const override;
+	virtual void Filter(FChooserEvaluationContext& Context, const FChooserIndexArray& IndexListIn, FChooserIndexArray& IndexListOut) const override;
 
 #if WITH_EDITOR
 	mutable FSoftObjectPath TestValue;
diff --git a/Engine/Plugins/Experimental/Chooser/Source/Chooser/Internal/RandomizeColumn.h b/Engine/Plugins/Experimental/Chooser/Source/Chooser/Internal/RandomizeColumn.h
--- a/Engine/Plugins/Experimental/Chooser/Source/Chooser/Internal/RandomizeColumn.h
+++ b/Engine/Plugins/Experimental/Chooser/Source/Chooser/Internal/RandomizeColumn.h
@@ -61,7 +61,7 @@ struct CHOOSER_API FRandomizeColumn : public FChooserColumnBase
 	UPROPERTY(EditAnywhere, Category= "Data", DisplayName="RowValues");
 	TArray<float> RowValues; 
 	
-	virtual void Filter(FChooserEvaluationContext& Context, const TArray<uint32>& IndexListIn, TArray<uint32>& IndexListOut) const override;
+	virtual void Filter(FChooserEvaluationContext& Context, const FChooserIndexArray& IndexListIn, FChooserIndexArray& IndexListOut) const override;
 	virtual void SetOutputs(FChooserEvaluationContext& Context, int RowIndex) const override;
 	
 	virtual bool HasFilters() const override { return true; }
diff --git a/Engine/Plugins/Experimental/Chooser/Source/Chooser/Private/BoolColumn.cpp b/Engine/Plugins/Experimental/Chooser/Source/Chooser/Private/BoolColumn.cpp
--- a/Engine/Plugins/Experimental/Chooser/Source/Chooser/Private/BoolColumn.cpp
+++ b/Engine/Plugins/Experimental/Chooser/Source/Chooser/Private/BoolColumn.cpp
@@ -1,5 +1,6 @@
 // Copyright Epic Games, Inc. All Rights Reserved.
 #include "BoolColumn.h"
+#include "ChooserIndexArray.h"
 #include "ChooserPropertyAccess.h"
 #include "Chooser.h"
 #include "ChooserTrace.h"
@@ -19,7 +20,7 @@ FBoolColumn::FBoolColumn()
 	InputValue.InitializeAs(FBoolContextProperty::StaticStruct());
 }
 
-void FBoolColumn::Filter(FChooserEvaluationContext& Context, const TArray<uint32>& IndexListIn, TArray<uint32>& IndexListOut) const
+void FBoolColumn::Filter(FChooserEvaluationContext& Context, const FChooserIndexArray& IndexListIn, FChooserIndexArray& IndexListOut) const
 {
 	if (InputValue.IsValid())
 	{
diff --git a/Engine/Plugins/Experimental/Chooser/Source/Chooser/Private/Chooser.cpp b/Engine/Plugins/Experimental/Chooser/Source/Chooser/Private/Chooser.cpp
--- a/Engine/Plugins/Experimental/Chooser/Source/Chooser/Private/Chooser.cpp
+++ b/Engine/Plugins/Experimental/Chooser/Source/Chooser/Private/Chooser.cpp
@@ -6,6 +6,7 @@
 #include "ObjectTrace.h"
 #include "Engine/UserDefinedStruct.h"
 #include "Engine/Blueprint.h"
+#include "ChooserIndexArray.h"
 
 #include UE_INLINE_GENERATED_CPP_BY_NAME(Chooser)
 
@@ -236,10 +237,12 @@ FObjectChooserBase::EIteratorStatus UChooserTable::EvaluateChooser(FChooserEvalu
 #if CHOOSER_DEBUGGING_ENABLED
 	Context.DebuggingInfo.CurrentChooser = Chooser;
 #endif
-	
 
-	TArray<uint32> Indices1;
-	TArray<uint32> Indices2;
+	uint32 Count = Chooser->ResultsStructs.Num();
+	uint32 BufferSize = Count * sizeof(uint32);
+
+	FChooserIndexArray Indices1(static_cast<uint32*>(FMemory_Alloca(BufferSize)), Count);
+	FChooserIndexArray Indices2(static_cast<uint32*>(FMemory_Alloca(BufferSize)), Count);
 
 	int RowCount = Chooser->ResultsStructs.Num();
 	Indices1.SetNum(RowCount);
@@ -247,8 +250,8 @@ FObjectChooserBase::EIteratorStatus UChooserTable::EvaluateChooser(FChooserEvalu
 	{
 		Indices1[i]=i;
 	}
-	TArray<uint32>* IndicesOut = &Indices1;
-	TArray<uint32>* IndicesIn = &Indices2;
+	FChooserIndexArray* IndicesOut = &Indices1;
+	FChooserIndexArray* IndicesIn = &Indices2;
 
 	for (const FInstancedStruct& ColumnData : Chooser->ColumnsStructs)
 	{
@@ -256,7 +259,7 @@ FObjectChooserBase::EIteratorStatus UChooserTable::EvaluateChooser(FChooserEvalu
 		if (Column.HasFilters())
 		{
 			Swap(IndicesIn, IndicesOut);
-			IndicesOut->SetNum(0, false);
+			IndicesOut->SetNum(0);
 			Column.Filter(Context, *IndicesIn, *IndicesOut);
 		}
 	}
diff --git a/Engine/Plugins/Experimental/Chooser/Source/Chooser/Private/EnumColumn.cpp b/Engine/Plugins/Experimental/Chooser/Source/Chooser/Private/EnumColumn.cpp
--- a/Engine/Plugins/Experimental/Chooser/Source/Chooser/Private/EnumColumn.cpp
+++ b/Engine/Plugins/Experimental/Chooser/Source/Chooser/Private/EnumColumn.cpp
@@ -1,5 +1,6 @@
 // Copyright Epic Games, Inc. All Rights Reserved.
 #include "EnumColumn.h"
+#include "ChooserIndexArray.h"
 #include "ChooserPropertyAccess.h"
 #include "ChooserTrace.h"
 
@@ -62,7 +63,7 @@ bool FChooserEnumRowData::Evaluate(const uint8 LeftHandSide) const
 	}
 }
 
-void FEnumColumn::Filter(FChooserEvaluationContext& Context, const TArray<uint32>& IndexListIn, TArray<uint32>& IndexListOut) const
+void FEnumColumn::Filter(FChooserEvaluationContext& Context, const FChooserIndexArray& IndexListIn, FChooserIndexArray& IndexListOut) const
 {
 	uint8 Result = 0;
 	if (InputValue.IsValid() &&
@@ -84,7 +85,7 @@ void FEnumColumn::Filter(FChooserEvaluationContext& Context, const TArray<uint32
 				const FChooserEnumRowData& RowValue = RowValues[Index];
 				if (RowValue.Evaluate(Result))
 				{
-					IndexListOut.Emplace(Index);
+					IndexListOut.Push(Index);
 				}
 			}
 		}
diff --git a/Engine/Plugins/Experimental/Chooser/Source/Chooser/Private/FloatRangeColumn.cpp b/Engine/Plugins/Experimental/Chooser/Source/Chooser/Private/FloatRangeColumn.cpp
--- a/Engine/Plugins/Experimental/Chooser/Source/Chooser/Private/FloatRangeColumn.cpp
+++ b/Engine/Plugins/Experimental/Chooser/Source/Chooser/Private/FloatRangeColumn.cpp
@@ -1,5 +1,6 @@
 // Copyright Epic Games, Inc. All Rights Reserved.
 #include "FloatRangeColumn.h"
+#include "ChooserIndexArray.h"
 #include "ChooserPropertyAccess.h"
 #include "ChooserTrace.h"
 
@@ -20,7 +21,7 @@ FFloatRangeColumn::FFloatRangeColumn()
 	InputValue.InitializeAs(FFloatContextProperty::StaticStruct());
 }
 
-void FFloatRangeColumn::Filter(FChooserEvaluationContext& Context, const TArray<uint32>& IndexListIn, TArray<uint32>& IndexListOut) const
+void FFloatRangeColumn::Filter(FChooserEvaluationContext& Context, const FChooserIndexArray& IndexListIn, FChooserIndexArray& IndexListOut) const
 {
 	if (InputValue.IsValid())
 	{
diff --git a/Engine/Plugins/Experimental/Chooser/Source/Chooser/Private/GameplayTagColumn.cpp b/Engine/Plugins/Experimental/Chooser/Source/Chooser/Private/GameplayTagColumn.cpp
--- a/Engine/Plugins/Experimental/Chooser/Source/Chooser/Private/GameplayTagColumn.cpp
+++ b/Engine/Plugins/Experimental/Chooser/Source/Chooser/Private/GameplayTagColumn.cpp
@@ -1,5 +1,6 @@
 // Copyright Epic Games, Inc. All Rights Reserved.
 #include "GameplayTagColumn.h"
+#include "ChooserIndexArray.h"
 #include "ChooserPropertyAccess.h"
 #include "ChooserTrace.h"
 
@@ -13,7 +14,7 @@ FGameplayTagColumn::FGameplayTagColumn()
 	InputValue.InitializeAs(FGameplayTagContextProperty::StaticStruct());
 }
 
-void FGameplayTagColumn::Filter(FChooserEvaluationContext& Context, const TArray<uint32>& IndexListIn, TArray<uint32>& IndexListOut) const
+void FGameplayTagColumn::Filter(FChooserEvaluationContext& Context, const FChooserIndexArray& IndexListIn, FChooserIndexArray& IndexListOut) const
 {
 	const FGameplayTagContainer* Result = nullptr;
 	if (InputValue.IsValid() && InputValue.Get<FChooserParameterGameplayTagBase>().GetValue(Context,Result))
diff --git a/Engine/Plugins/Experimental/Chooser/Source/Chooser/Private/ObjectColumn.cpp b/Engine/Plugins/Experimental/Chooser/Source/Chooser/Private/ObjectColumn.cpp
--- a/Engine/Plugins/Experimental/Chooser/Source/Chooser/Private/ObjectColumn.cpp
+++ b/Engine/Plugins/Experimental/Chooser/Source/Chooser/Private/ObjectColumn.cpp
@@ -1,5 +1,6 @@
 // Copyright Epic Games, Inc. All Rights Reserved.
 #include "ObjectColumn.h"
+#include "ChooserIndexArray.h"
 #include "ChooserPropertyAccess.h"
 #include "ChooserTrace.h"
 
@@ -57,7 +58,7 @@ bool FChooserObjectRowData::Evaluate(const FSoftObjectPath& LeftHandSide) const
 	}
 }
 
-void FObjectColumn::Filter(FChooserEvaluationContext& Context, const TArray<uint32>& IndexListIn, TArray<uint32>& IndexListOut) const
+void FObjectColumn::Filter(FChooserEvaluationContext& Context, const FChooserIndexArray& IndexListIn, FChooserIndexArray& IndexListOut) const
 {
 	FSoftObjectPath Result;
 	if (InputValue.IsValid() &&
@@ -79,7 +80,7 @@ void FObjectColumn::Filter(FChooserEvaluationContext& Context, const TArray<uint
 				const FChooserObjectRowData& RowValue = RowValues[Index];
 				if (RowValue.Evaluate(Result))
 				{
-					IndexListOut.Emplace(Index);
+					IndexListOut.Push(Index);
 				}
 			}
 		}
diff --git a/Engine/Plugins/Experimental/Chooser/Source/Chooser/Private/RandomizeColumn.cpp b/Engine/Plugins/Experimental/Chooser/Source/Chooser/Private/RandomizeColumn.cpp
--- a/Engine/Plugins/Experimental/Chooser/Source/Chooser/Private/RandomizeColumn.cpp
+++ b/Engine/Plugins/Experimental/Chooser/Source/Chooser/Private/RandomizeColumn.cpp
@@ -1,5 +1,6 @@
 // Copyright Epic Games, Inc. All Rights Reserved.
 #include "RandomizeColumn.h"
+#include "ChooserIndexArray.h"
 #include "ChooserPropertyAccess.h"
 
 bool FRandomizeContextProperty::GetValue(FChooserEvaluationContext& Context, const FChooserRandomizationContext*& OutResult) const
@@ -12,7 +13,7 @@ FRandomizeColumn::FRandomizeColumn()
 	InputValue.InitializeAs(FRandomizeContextProperty::StaticStruct());
 }
 
-void FRandomizeColumn::Filter(FChooserEvaluationContext& Context, const TArray<uint32>& IndexListIn, TArray<uint32>& IndexListOut) const
+void FRandomizeColumn::Filter(FChooserEvaluationContext& Context, const FChooserIndexArray& IndexListIn, FChooserIndexArray& IndexListOut) const
 {
 	int Count = IndexListIn.Num();
 	int Selection = 0;
@@ -82,7 +83,7 @@ void FRandomizeColumn::Filter(FChooserEvaluationContext& Context, const TArray<u
 
 	if (Selection < Count)
 	{
-		IndexListOut.Add(IndexListIn[Selection]);
+		IndexListOut.Push(IndexListIn[Selection]);
 	}
 }
 
diff --git a/Engine/Plugins/Experimental/Chooser/Source/Chooser/Public/ChooserIndexArray.h b/Engine/Plugins/Experimental/Chooser/Source/Chooser/Public/ChooserIndexArray.h
new file mode 100644
--- /dev/null
+++ b/Engine/Plugins/Experimental/Chooser/Source/Chooser/Public/ChooserIndexArray.h
@@ -0,0 +1,66 @@
+// Copyright Epic Games, Inc. All Rights Reserved.
+
+#pragma once
+
+#include "CoreMinimal.h"
+
+class FChooserIndexArray
+{
+public:
+	FChooserIndexArray(uint32* InData, uint32 InMaxSize): Data(InData), MaxSize(InMaxSize), Size(0) {}
+	
+	typedef uint32* iterator;
+	typedef const uint32* const_iterator;
+	
+	iterator begin() { return Data; }
+	const_iterator begin() const { return Data; }
+	iterator end() { return Data + Size; }
+	const_iterator end() const { return Data + Size; }
+	
+	
+	void Push(uint32 Value)
+	{
+		check(Size < MaxSize);
+		Data[Size++] = Value;
+	}
+
+	bool IsEmpty() const
+	{
+		return Size == 0;
+	}
+	
+	uint32 Num() const
+	{
+		return Size;
+	}
+	
+	void SetNum(uint32 Num)
+	{
+		check(Num <= MaxSize);
+		Size = Num;		
+	}
+
+	uint32& operator [] (uint32 Index)
+	{
+		check(Index < MaxSize);
+		return Data[Index];
+	}
+	
+	uint32 operator [] (uint32 Index) const
+	{
+		check(Index < MaxSize);
+		return Data[Index];
+	}
+	
+	void operator = (const FChooserIndexArray& Other)
+	{
+		check(MaxSize >= Other.Size);
+		Size = Other.Size;
+		FMemory::Memcpy(Data,Other.Data, Size*sizeof(uint32));
+	}
+
+private:
+	uint32* Data;
+	uint32 MaxSize;
+	uint32 Size;
+};
diff --git a/Engine/Plugins/Experimental/Chooser/Source/Chooser/Public/IChooserColumn.h b/Engine/Plugins/Experimental/Chooser/Source/Chooser/Public/IChooserColumn.h
--- a/Engine/Plugins/Experimental/Chooser/Source/Chooser/Public/IChooserColumn.h
+++ b/Engine/Plugins/Experimental/Chooser/Source/Chooser/Public/IChooserColumn.h
@@ -26,6 +26,8 @@ public:
 class UChooserTable;
 class FVariant;
 
+class FChooserIndexArray;
+
 USTRUCT()
 struct CHOOSER_API FChooserColumnBase
 {
@@ -34,7 +36,7 @@ struct CHOOSER_API FChooserColumnBase
 public:
 	virtual ~FChooserColumnBase() {}
 	virtual void PostLoad() {};
-	virtual void Filter(FChooserEvaluationContext& Context, const TArray<uint32>& IndexListIn, TArray<uint32>& IndexListOut) const {}
+	virtual void Filter(FChooserEvaluationContext& Context, const FChooserIndexArray& IndexListIn, FChooserIndexArray& IndexListOut) const {}
 
 	virtual bool HasFilters() const { return true; }
 	virtual bool HasOutputs() const { return false; }
