Refactor chooser evaluation context to avoid memory allocations at runtime

diff --git a/Engine/Plugins/Experimental/Chooser/Source/Chooser/Internal/AnimNode_ChooserPlayer.h b/Engine/Plugins/Experimental/Chooser/Source/Chooser/Internal/AnimNode_ChooserPlayer.h
--- a/Engine/Plugins/Experimental/Chooser/Source/Chooser/Internal/AnimNode_ChooserPlayer.h
+++ b/Engine/Plugins/Experimental/Chooser/Source/Chooser/Internal/AnimNode_ChooserPlayer.h
@@ -155,11 +155,13 @@ public:
 	
 private:
 	UAnimationAsset* ChooseAsset(const FAnimationUpdateContext& Context);
+	FChooserPlayerSettings Settings;
 	FChooserEvaluationContext ChooserContext;
 	UAnimationAsset* CurrentAsset = nullptr;
 	float CurrentStartTime = 0;
 	bool CurrentMirror = false;
 	uint32 CurrentCurveOverridesHash = 0;
+	
 
 	// Update Counter for detecting being relevant
 	FGraphTraversalCounter UpdateCounter;
diff --git a/Engine/Plugins/Experimental/Chooser/Source/Chooser/Private/AnimNode_ChooserPlayer.cpp b/Engine/Plugins/Experimental/Chooser/Source/Chooser/Private/AnimNode_ChooserPlayer.cpp
--- a/Engine/Plugins/Experimental/Chooser/Source/Chooser/Private/AnimNode_ChooserPlayer.cpp
+++ b/Engine/Plugins/Experimental/Chooser/Source/Chooser/Private/AnimNode_ChooserPlayer.cpp
@@ -27,8 +27,6 @@ UAnimationAsset* FAnimNode_ChooserPlayer::ChooseAsset(const FAnimationUpdateCont
 	if (Chooser.IsValid())
 	{
 		// reset settings to default
-		FChooserPlayerSettings& Settings = ChooserContext.Params[1].GetMutable<FChooserPlayerSettings>();
-		
 		Settings = DefaultSettings;
 		
 		const FObjectChooserBase& ChooserBase = Chooser.Get<FObjectChooserBase>();
@@ -64,12 +62,12 @@ void FAnimNode_ChooserPlayer::Initialize_AnyThread(const FAnimationInitializeCon
 {
 	DECLARE_SCOPE_HIERARCHICAL_COUNTER_ANIMNODE(Initialize_AnyThread)
 	FAnimNode_BlendStack_Standalone::Initialize_AnyThread(Context);
-
-	FInstancedStruct This;
-	This.InitializeAs(FChooserEvaluationInputObject::StaticStruct());
-	This.GetMutable<FChooserEvaluationInputObject>().Object = Context.AnimInstanceProxy->GetAnimInstanceObject();
-	ChooserContext.Params.Add(This);
-	ChooserContext.Params.Add(FInstancedStruct(FChooserPlayerSettings::StaticStruct()));
+	
+	if(ChooserContext.Params.IsEmpty())
+	{
+		ChooserContext.AddObjectParam(Context.GetAnimInstanceObject());
+		ChooserContext.AddStructParam(Settings);
+	}
 }
 
 
@@ -109,8 +107,6 @@ void FAnimNode_ChooserPlayer::UpdateAssetPlayer(const FAnimationUpdateContext& C
 		NewAsset = ChooseAsset(Context);
 	}
 
-	const FChooserPlayerSettings& Settings = ChooserContext.Params[1].Get<FChooserPlayerSettings>();
-
 	// Restart the animation:
 	// - if this node just became relevant
 	// - if we chose a new animation
@@ -194,8 +190,6 @@ void FAnimNode_ChooserPlayer::Evaluate_AnyThread(FPoseContext& Output)
 	ANIM_MT_SCOPE_CYCLE_COUNTER_VERBOSE(MotionMatching, !IsInGameThread());
 
 	FAnimNode_BlendStack_Standalone::Evaluate_AnyThread(Output);
-
-	const FChooserPlayerSettings& Settings = ChooserContext.Params[1].Get<FChooserPlayerSettings>();
 }
 
 FName FAnimNode_ChooserPlayer::GetGroupName() const
diff --git a/Engine/Plugins/Experimental/Chooser/Source/Chooser/Private/Chooser.cpp b/Engine/Plugins/Experimental/Chooser/Source/Chooser/Private/Chooser.cpp
--- a/Engine/Plugins/Experimental/Chooser/Source/Chooser/Private/Chooser.cpp
+++ b/Engine/Plugins/Experimental/Chooser/Source/Chooser/Private/Chooser.cpp
@@ -185,9 +185,9 @@ void UChooserTable::UpdateDebugging(FChooserEvaluationContext& Context) const
 {
 	FScopeLock Lock(&DebugLock);
 	
-	for (const FInstancedStruct& Param : Context.Params)
+	for (const FStructView& Param : Context.Params)
 	{
-		if (const FChooserEvaluationInputObject* ObjectParam = Param.GetPtr<FChooserEvaluationInputObject>())
+		if (const FChooserEvaluationInputObject* ObjectParam = Param.GetPtr<const FChooserEvaluationInputObject>())
 		{
 			if (UObject* ContextObject = ObjectParam->Object)
 			{
@@ -202,7 +202,7 @@ void UChooserTable::UpdateDebugging(FChooserEvaluationContext& Context) const
 					}
 				}
 				
-				if (ContextObject == DebugTarget)
+				if (ContextObject == DebugTarget) 
 				{
 					bDebugTestValuesValid = true;
 					Context.DebuggingInfo.bCurrentDebugTarget = true;
diff --git a/Engine/Plugins/Experimental/Chooser/Source/Chooser/Private/ChooserFunctionLibrary.cpp b/Engine/Plugins/Experimental/Chooser/Source/Chooser/Private/ChooserFunctionLibrary.cpp
--- a/Engine/Plugins/Experimental/Chooser/Source/Chooser/Private/ChooserFunctionLibrary.cpp
+++ b/Engine/Plugins/Experimental/Chooser/Source/Chooser/Private/ChooserFunctionLibrary.cpp
@@ -12,10 +12,7 @@ UChooserFunctionLibrary::UChooserFunctionLibrary(const FObjectInitializer& Objec
 static FObjectChooserBase::EIteratorStatus StaticEvaluateChooser(const UObject* ContextObject, const UChooserTable* Chooser, FObjectChooserBase::FObjectChooserIteratorCallback Callback)
 {
 	// Fallback single context object version
-	FChooserEvaluationContext Context;
-	Context.Params.AddDefaulted();
-	Context.Params.Last().InitializeAs(FChooserEvaluationInputObject::StaticStruct());
-	Context.Params.Last().GetMutable<FChooserEvaluationInputObject>().Object = const_cast<UObject*>(ContextObject);
+	FChooserEvaluationContext Context(const_cast<UObject*>(ContextObject));
 	return UChooserTable::EvaluateChooser(Context, Chooser, Callback);
 }
 
@@ -37,9 +34,7 @@ UObject* UChooserFunctionLibrary::EvaluateChooser(const UObject* ContextObject,
 
 void UChooserFunctionLibrary::AddChooserObjectInput(FChooserEvaluationContext& Context, UObject* Object)
 {
-	Context.Params.AddDefaulted();
-	Context.Params.Last().InitializeAs(FChooserEvaluationInputObject::StaticStruct());
-	Context.Params.Last().GetMutable<FChooserEvaluationInputObject>().Object = Object;
+	Context.AddObjectParam(Object);
 }
 
 void UChooserFunctionLibrary::AddChooserStructInput(FChooserEvaluationContext& Inputs, int32 Value)
@@ -57,7 +52,7 @@ DEFINE_FUNCTION(UChooserFunctionLibrary::execAddChooserStructInput)
 	
 	Stack.StepCompiledIn<FStructProperty>(nullptr);
 	const FStructProperty* ValueProp = CastField<FStructProperty>(Stack.MostRecentProperty);
-	const void* ValuePtr = Stack.MostRecentPropertyAddress;
+	uint8* ValuePtr = Stack.MostRecentPropertyAddress;
 	P_FINISH;
 	
 	if (!ValueProp || !ValuePtr)
@@ -72,8 +67,7 @@ DEFINE_FUNCTION(UChooserFunctionLibrary::execAddChooserStructInput)
 	else
 	{
 		P_NATIVE_BEGIN;
-		ChooserContext.Params.AddDefaulted();
-		ChooserContext.Params.Last().InitializeAs(ValueProp->Struct, static_cast<const uint8*>(ValuePtr));
+		ChooserContext.Params.Add(FStructView(ValueProp->Struct, ValuePtr));
 		P_NATIVE_END;
 	}
 }
diff --git a/Engine/Plugins/Experimental/Chooser/Source/Chooser/Private/ChooserPropertyAccess.cpp b/Engine/Plugins/Experimental/Chooser/Source/Chooser/Private/ChooserPropertyAccess.cpp
--- a/Engine/Plugins/Experimental/Chooser/Source/Chooser/Private/ChooserPropertyAccess.cpp
+++ b/Engine/Plugins/Experimental/Chooser/Source/Chooser/Private/ChooserPropertyAccess.cpp
@@ -314,7 +314,7 @@ namespace UE::Chooser
 			return false;
 		}
 		
-		if (FChooserEvaluationInputObject* ObjectInput = Context.Params[CompiledBinding.ContextIndex].GetMutablePtr<FChooserEvaluationInputObject>())
+		if (FChooserEvaluationInputObject* ObjectInput = Context.Params[CompiledBinding.ContextIndex].GetPtr<FChooserEvaluationInputObject>())
 		{
 			UObject* Object = ObjectInput->Object.Get();
 			Container = reinterpret_cast<uint8*>(Object);
@@ -325,7 +325,7 @@ namespace UE::Chooser
 		}
 		else
 		{
-			Container = Context.Params[CompiledBinding.ContextIndex].GetMutableMemory();
+			Container = Context.Params[CompiledBinding.ContextIndex].GetMemory();
 			InputType = Context.Params[CompiledBinding.ContextIndex].GetScriptStruct();
 		}
 
@@ -535,7 +535,7 @@ namespace UE::Chooser
 		{
 			uint8* Container = nullptr;
 			const UStruct* StructType = nullptr;
-			if (FChooserEvaluationInputObject* ObjectParam = Context.Params[PropertyBinding.ContextIndex].GetMutablePtr<FChooserEvaluationInputObject>())
+			if (FChooserEvaluationInputObject* ObjectParam = Context.Params[PropertyBinding.ContextIndex].GetPtr<FChooserEvaluationInputObject>())
 			{
 				Container = static_cast<uint8*>(ObjectParam->Object);
 				if (Container)
@@ -545,7 +545,7 @@ namespace UE::Chooser
 			}
 			else
 			{
-				Container = Context.Params[PropertyBinding.ContextIndex].GetMutableMemory();
+				Container = Context.Params[PropertyBinding.ContextIndex].GetMemory();
 				StructType = Context.Params[PropertyBinding.ContextIndex].GetScriptStruct();
 			}
 
diff --git a/Engine/Plugins/Experimental/Chooser/Source/Chooser/Public/IObjectChooser.h b/Engine/Plugins/Experimental/Chooser/Source/Chooser/Public/IObjectChooser.h
--- a/Engine/Plugins/Experimental/Chooser/Source/Chooser/Public/IObjectChooser.h
+++ b/Engine/Plugins/Experimental/Chooser/Source/Chooser/Public/IObjectChooser.h
@@ -6,6 +6,7 @@
 #include "UObject/Object.h"
 #include "UObject/Interface.h"
 #include "InstancedStruct.h"
+#include "StructView.h"
 #include "IObjectChooser.generated.h"
 
 #if UE_TRACE_ENABLED && !IS_PROGRAM && !UE_BUILD_SHIPPING && !UE_BUILD_TEST
@@ -42,6 +43,8 @@ USTRUCT()
 struct CHOOSER_API FChooserEvaluationInputObject
 {
 	GENERATED_BODY()
+	FChooserEvaluationInputObject() {}
+	FChooserEvaluationInputObject (UObject* InObject) : Object(InObject) {}
 	TObjectPtr<UObject> Object;
 };
 
@@ -51,13 +54,22 @@ struct CHOOSER_API FChooserEvaluationContext
 	FChooserEvaluationContext() {}
 	FChooserEvaluationContext(UObject* ContextObject)
 	{
-		AddParam(ContextObject);
+		AddObjectParam(ContextObject);
 	}
-	void AddParam(UObject* ContextObject)
+
+	// Add a UObject Parameter to the context
+	void AddObjectParam(UObject* Param)
 	{
-		Params.AddDefaulted();
-		Params.Last().InitializeAs(FChooserEvaluationInputObject::StaticStruct());
-		Params.Last().GetMutable<FChooserEvaluationInputObject>().Object = ContextObject;
+		ObjectParams.Add(Param);
+		AddStructParam(ObjectParams.Last());
+	}
+
+	// Add a struct Parameter to the Context
+	// the struct will be referred to by reference, and so must have a lifetime that is longer than this context
+	template <class T>
+	void AddStructParam(T& Param)
+	{
+		Params.Add(FStructView::Make(Param));
 	}
 
 	#if CHOOSER_DEBUGGING_ENABLED
@@ -65,7 +77,10 @@ struct CHOOSER_API FChooserEvaluationContext
     #endif
 	
 	GENERATED_BODY()
-	TArray<FInstancedStruct, TInlineAllocator<4>> Params;
+	TArray<FStructView, TInlineAllocator<4>> Params;
+
+	// storage for Object Params, call AddObjectParam to allocate one FChooserEvaluationInputObject in this array and then add a StructView of it to the Params array
+	TArray<FChooserEvaluationInputObject, TFixedAllocator<4>> ObjectParams;
 };
 
 USTRUCT()
diff --git a/Engine/Plugins/Experimental/Chooser/Source/ProxyTable/Private/ProxyTableFunctionLibrary.cpp b/Engine/Plugins/Experimental/Chooser/Source/ProxyTable/Private/ProxyTableFunctionLibrary.cpp
--- a/Engine/Plugins/Experimental/Chooser/Source/ProxyTable/Private/ProxyTableFunctionLibrary.cpp
+++ b/Engine/Plugins/Experimental/Chooser/Source/ProxyTable/Private/ProxyTableFunctionLibrary.cpp
@@ -12,10 +12,7 @@ UObject* UProxyTableFunctionLibrary::EvaluateProxyAsset(const UObject* ContextOb
 	UObject* Result = nullptr;
 	if (Proxy)
 	{
-		FChooserEvaluationContext Context;
-		Context.Params.AddDefaulted();
-		Context.Params.Last().InitializeAs(FChooserEvaluationInputObject::StaticStruct());
-		Context.Params.Last().GetMutable<FChooserEvaluationInputObject>().Object = const_cast<UObject*>(ContextObject);
+		FChooserEvaluationContext Context(const_cast<UObject*>(ContextObject));
 		
 		Result = Proxy->FindProxyObject(Context);
 		if (ObjectClass && Result && !Result->IsA(ObjectClass))
@@ -34,10 +31,7 @@ UObject* UProxyTableFunctionLibrary::EvaluateProxyTable(const UObject* ContextOb
 	{
 		FGuid Guid;
 		Guid.A = GetTypeHash(Key);
-		FChooserEvaluationContext Context;
-		Context.Params.AddDefaulted();
-		Context.Params.Last().InitializeAs(FChooserEvaluationInputObject::StaticStruct());
-		Context.Params.Last().GetMutable<FChooserEvaluationInputObject>().Object = const_cast<UObject*>(ContextObject);
+		FChooserEvaluationContext Context(const_cast<UObject*>(ContextObject));
 		if (UObject* Value = ProxyTable->FindProxyObject(Guid, Context))
 		{
 			return Value;
