Compile property access chains for:
- Proxy Table results
- Proxy Table struct outputs
- ChooserPlayer Chooser property

diff --git a/Engine/Plugins/Experimental/Chooser/Source/Chooser/Internal/AnimNode_ChooserPlayer.h b/Engine/Plugins/Experimental/Chooser/Source/Chooser/Internal/AnimNode_ChooserPlayer.h
--- a/Engine/Plugins/Experimental/Chooser/Source/Chooser/Internal/AnimNode_ChooserPlayer.h
+++ b/Engine/Plugins/Experimental/Chooser/Source/Chooser/Internal/AnimNode_ChooserPlayer.h
@@ -6,6 +6,7 @@
 #include "Animation/AnimNode_RelevantAssetPlayerBase.h"
 #include "Animation/AnimationAsset.h"
 #include "CoreMinimal.h"
+#include "IHasContext.h"
 #include "InstancedStruct.h"
 #include "IObjectChooser.h"
 #include "BlendStack/AnimNode_BlendStack.h"
@@ -94,13 +95,15 @@ struct FChooserPlayerSettings
 };
 
 USTRUCT(BlueprintInternalUseOnly)
-struct FAnimNode_ChooserPlayer : public FAnimNode_BlendStack_Standalone
+struct FAnimNode_ChooserPlayer : public FAnimNode_BlendStack_Standalone, public IHasContextClass
 {
 	GENERATED_BODY()
 
 	CHOOSER_API FAnimNode_ChooserPlayer();
 
 public:
+	virtual TConstArrayView<FInstancedStruct> GetContextData() const override { return ChooserContextDefinition; }
+	
 	// How often the chooser should be evaluated
 	UPROPERTY(EditAnywhere, Category = "Chooser")
 	EChooserEvaluationFrequency EvaluationFrequency = EChooserEvaluationFrequency::OnBecomeRelevant;
@@ -160,6 +163,7 @@ private:
 	UAnimationAsset* CurrentAsset = nullptr;
 	float CurrentStartTime = 0;
 	bool CurrentMirror = false;
+	bool bInitialized = false;
 	uint32 CurrentCurveOverridesHash = 0;
 	
 
diff --git a/Engine/Plugins/Experimental/Chooser/Source/Chooser/Private/AnimNode_ChooserPlayer.cpp b/Engine/Plugins/Experimental/Chooser/Source/Chooser/Private/AnimNode_ChooserPlayer.cpp
--- a/Engine/Plugins/Experimental/Chooser/Source/Chooser/Private/AnimNode_ChooserPlayer.cpp
+++ b/Engine/Plugins/Experimental/Chooser/Source/Chooser/Private/AnimNode_ChooserPlayer.cpp
@@ -2,6 +2,7 @@
 
 #include "AnimNode_ChooserPlayer.h"
 
+#include "IHasContext.h"
 #include "Animation/AnimInstanceProxy.h"
 #include "Animation/AnimPoseSearchProvider.h"
 #include "Animation/AnimStats.h"
@@ -63,10 +64,17 @@ void FAnimNode_ChooserPlayer::Initialize_AnyThread(const FAnimationInitializeCon
 	DECLARE_SCOPE_HIERARCHICAL_COUNTER_ANIMNODE(Initialize_AnyThread)
 	FAnimNode_BlendStack_Standalone::Initialize_AnyThread(Context);
 	
-	if(ChooserContext.Params.IsEmpty())
+	if(!bInitialized)
 	{
 		ChooserContext.AddObjectParam(Context.GetAnimInstanceObject());
 		ChooserContext.AddStructParam(Settings);
+
+		if (FObjectChooserBase* ObjectChooser = Chooser.GetMutablePtr<FObjectChooserBase>())
+		{
+			ObjectChooser->Compile(this,false);
+		}
+		
+		bInitialized = true;
 	}
 }
 
diff --git a/Engine/Plugins/Experimental/Chooser/Source/ProxyTable/Private/ProxyAsset.cpp b/Engine/Plugins/Experimental/Chooser/Source/ProxyTable/Private/ProxyAsset.cpp
--- a/Engine/Plugins/Experimental/Chooser/Source/ProxyTable/Private/ProxyAsset.cpp
+++ b/Engine/Plugins/Experimental/Chooser/Source/ProxyTable/Private/ProxyAsset.cpp
@@ -83,6 +83,14 @@ void UProxyAsset::PostLoad()
 		Guid.B = GetTypeHash(GetPackage()->GetPathName());
 	}
 
+	if (ProxyTable.IsValid())
+	{
+		// compile property access for Proxy Table fallback codepath
+		if (FChooserParameterProxyTableBase* ProxyReference = ProxyTable.GetMutablePtr<FChooserParameterProxyTableBase>())
+		{
+			ProxyReference->Compile(this, false);
+		}
+	}
 }
 
 void UProxyAsset::PostDuplicate(EDuplicateMode::Type DuplicateMode)
diff --git a/Engine/Plugins/Experimental/Chooser/Source/ProxyTable/Private/ProxyTable.cpp b/Engine/Plugins/Experimental/Chooser/Source/ProxyTable/Private/ProxyTable.cpp
--- a/Engine/Plugins/Experimental/Chooser/Source/ProxyTable/Private/ProxyTable.cpp
+++ b/Engine/Plugins/Experimental/Chooser/Source/ProxyTable/Private/ProxyTable.cpp
@@ -128,7 +128,7 @@ void UProxyTable::BuildRuntimeData()
 	for(const FProxyEntry& Entry : RuntimeEntries) 
 	{
 		Keys.Add(Entry.GetGuid());
-   		RuntimeValues.Add({Entry.ValueStruct, Entry.OutputStructData});
+   		RuntimeValues.Add({Entry.Proxy, Entry.ValueStruct, Entry.OutputStructData});
 	}
 	
 	// register callbacks on updated dependencies
@@ -147,13 +147,6 @@ void UProxyTable::BuildRuntimeData()
 	}
 }
 
-
-void UProxyTable::PostLoad()
-{
-	Super::PostLoad();
-	BuildRuntimeData();
-}
-
 void UProxyTable::PostTransacted(const FTransactionObjectEvent& TransactionEvent)
 {
 	UObject::PostTransacted(TransactionEvent);
@@ -163,6 +156,31 @@ void UProxyTable::PostTransacted(const FTransactionObjectEvent& TransactionEvent
 
 #endif
 
+void UProxyTable::PostLoad()
+{
+	Super::PostLoad();
+#if WITH_EDITORONLY_DATA
+	BuildRuntimeData();
+#endif
+
+	// compilation for property accesses
+	for(FRuntimeProxyValue& Entry : RuntimeValues) 
+	{
+		if (FObjectChooserBase* Result = Entry.Value.GetMutablePtr<FObjectChooserBase>())
+		{
+			// need to compile results in case one is a LookupProxy
+			Result->Compile(Entry.ProxyAsset, false);
+		}
+
+		/// compile any struct output property references
+		for(FProxyStructOutput& StructOutput : Entry.OutputStructData)
+		{
+			StructOutput.Binding.Compile(Entry.ProxyAsset);
+		}
+	}
+}
+
+
 void UProxyTable::BeginDestroy()
 {
 	RuntimeValues.Empty();
diff --git a/Engine/Plugins/Experimental/Chooser/Source/ProxyTable/Public/ProxyTable.h b/Engine/Plugins/Experimental/Chooser/Source/ProxyTable/Public/ProxyTable.h
--- a/Engine/Plugins/Experimental/Chooser/Source/ProxyTable/Public/ProxyTable.h
+++ b/Engine/Plugins/Experimental/Chooser/Source/ProxyTable/Public/ProxyTable.h
@@ -56,6 +56,9 @@ struct PROXYTABLE_API FRuntimeProxyValue
 {
 	GENERATED_BODY()
 
+	UPROPERTY()
+	TObjectPtr<UProxyAsset> ProxyAsset;
+
 	UPROPERTY()
 	FInstancedStruct Value;
 
@@ -80,6 +83,7 @@ public:
 	FObjectChooserBase::EIteratorStatus FindProxyObjectMulti(const FGuid& Key, FChooserEvaluationContext &Context, FObjectChooserBase::FObjectChooserIteratorCallback Callback) const;
 	UObject* FindProxyObject(const FGuid& Key, FChooserEvaluationContext& Context) const;
 
+	virtual void PostLoad() override;
 #if WITH_EDITORONLY_DATA
 public:
 	FProxyTableChanged OnProxyTableChanged;
@@ -90,7 +94,6 @@ public:
 	UPROPERTY(EditAnywhere, Category = "Inheritance")
 	TArray<TObjectPtr<UProxyTable>> InheritEntriesFrom;
 
-	virtual void PostLoad() override;
 	virtual void PostTransacted(const FTransactionObjectEvent& TransactionEvent) override;
 private:
 	void BuildRuntimeData();
