Find feature for ChooserTable property references

diff --git a/Engine/Plugins/Experimental/Chooser/Source/Chooser/Internal/BoolColumn.h b/Engine/Plugins/Experimental/Chooser/Source/Chooser/Internal/BoolColumn.h
--- a/Engine/Plugins/Experimental/Chooser/Source/Chooser/Internal/BoolColumn.h
+++ b/Engine/Plugins/Experimental/Chooser/Source/Chooser/Internal/BoolColumn.h
@@ -17,6 +17,7 @@ enum class EBoolColumnCellValue
 	MatchAny = 2,
 };
 
+
 USTRUCT(DisplayName = "Bool Property Binding")
 struct CHOOSER_API FBoolContextProperty :  public FChooserParameterBoolBase
 {
@@ -41,24 +42,8 @@ public:
 		}
 	}
 	
-	virtual void Compile(IHasContextClass* Owner, bool bForce) override
-	{
-		Binding.Compile(Owner, bForce);
-	};
-
-#if WITH_EDITOR
-	virtual void GetDisplayName(FText& OutName) const override
-	{
-		if (!Binding.DisplayName.IsEmpty())
-		{
-			OutName = FText::FromString(Binding.DisplayName);
-		} 
-		else if (!Binding.PropertyBindingChain.IsEmpty())
-		{
-			OutName = FText::FromName(Binding.PropertyBindingChain.Last());
-		}
-	}
-#endif
+	CHOOSER_PARAMETER_BOILERPLATE();
+	
 };
 
 USTRUCT()
diff --git a/Engine/Plugins/Experimental/Chooser/Source/Chooser/Internal/Chooser.h b/Engine/Plugins/Experimental/Chooser/Source/Chooser/Internal/Chooser.h
--- a/Engine/Plugins/Experimental/Chooser/Source/Chooser/Internal/Chooser.h
+++ b/Engine/Plugins/Experimental/Chooser/Source/Chooser/Internal/Chooser.h
@@ -13,6 +13,7 @@
 
 #include "Chooser.generated.h"
 
+
 UCLASS(BlueprintType)
 class CHOOSER_API UChooserTable : public UObject, public IHasContextClass
 {
@@ -24,6 +25,8 @@ public:
 	virtual void PostLoad() override;
 	virtual void Compile(bool bForce = false) override;
 #if WITH_EDITOR
+	virtual void GetAssetRegistryTags(TArray<FAssetRegistryTag>& OutTags) const override;
+	
 	void OnDependentStructChanged(UUserDefinedStruct* Blueprint) { Compile(true); }
 	void OnDependencyCompiled(UBlueprint* Blueprint) { Compile(true); }
 	virtual void AddCompileDependency(const UStruct* Struct) override;
@@ -57,6 +60,9 @@ public:
 	bool bEnableDebugTesting = false;
 	mutable bool bDebugTestValuesValid = false;
 
+	static const FName PropertyNamesTag;
+	static const FString PropertyTagDelimiter;
+
 private: 
 	// caching the OutputObjectType and ContextObjectType so that on Undo, we can tell if we should fire the changed delegate
 	UClass* CachedPreviousOutputObjectType = nullptr;
diff --git a/Engine/Plugins/Experimental/Chooser/Source/Chooser/Internal/EnumColumn.h b/Engine/Plugins/Experimental/Chooser/Source/Chooser/Internal/EnumColumn.h
--- a/Engine/Plugins/Experimental/Chooser/Source/Chooser/Internal/EnumColumn.h
+++ b/Engine/Plugins/Experimental/Chooser/Source/Chooser/Internal/EnumColumn.h
@@ -48,26 +48,10 @@ struct CHOOSER_API FEnumContextProperty : public FChooserParameterEnumBase
 		}
 	}
 
-	virtual void Compile(IHasContextClass* Owner, bool bForce) override
-    {
-    	Binding.Compile(Owner, bForce);
-    };
+	CHOOSER_PARAMETER_BOILERPLATE();
 
 #if WITH_EDITOR
-	virtual void GetDisplayName(FText& OutName) const override
-	{
-		if (!Binding.DisplayName.IsEmpty())
-		{
-			OutName = FText::FromString(Binding.DisplayName);
-		} 
-		else if (!Binding.PropertyBindingChain.IsEmpty())
-		{
-			OutName = FText::FromName(Binding.PropertyBindingChain.Last());
-		}
-	}
-
 	virtual const UEnum* GetEnum() const override { return Binding.Enum; }
-
 #endif
 	
 private:
diff --git a/Engine/Plugins/Experimental/Chooser/Source/Chooser/Internal/FloatRangeColumn.h b/Engine/Plugins/Experimental/Chooser/Source/Chooser/Internal/FloatRangeColumn.h
--- a/Engine/Plugins/Experimental/Chooser/Source/Chooser/Internal/FloatRangeColumn.h
+++ b/Engine/Plugins/Experimental/Chooser/Source/Chooser/Internal/FloatRangeColumn.h
@@ -34,24 +34,7 @@ struct CHOOSER_API FFloatContextProperty :  public FChooserParameterFloatBase
 		}
 	}
 	
-	virtual void Compile(IHasContextClass* Owner, bool bForce) override
-	{
-		Binding.Compile(Owner, bForce);
-	};
-
-#if WITH_EDITOR
-	virtual void GetDisplayName(FText& OutName) const override
-	{
-		if (!Binding.DisplayName.IsEmpty())
-		{
-			OutName = FText::FromString(Binding.DisplayName);
-		} 
-		else if (!Binding.PropertyBindingChain.IsEmpty())
-		{
-			OutName = FText::FromName(Binding.PropertyBindingChain.Last());
-		}
-	}
-#endif
+	CHOOSER_PARAMETER_BOILERPLATE();
 };
 
 USTRUCT()
diff --git a/Engine/Plugins/Experimental/Chooser/Source/Chooser/Internal/GameplayTagColumn.h b/Engine/Plugins/Experimental/Chooser/Source/Chooser/Internal/GameplayTagColumn.h
--- a/Engine/Plugins/Experimental/Chooser/Source/Chooser/Internal/GameplayTagColumn.h
+++ b/Engine/Plugins/Experimental/Chooser/Source/Chooser/Internal/GameplayTagColumn.h
@@ -34,25 +34,8 @@ public:
 			PropertyBindingChain_DEPRECATED.SetNum(0);
 		}
 	}
-	
-	virtual void Compile(IHasContextClass* Owner, bool bForce) override
-	{
-		Binding.Compile(Owner, bForce);
-	};
 
-#if WITH_EDITOR
-	virtual void GetDisplayName(FText& OutName) const override
-	{
-		if (!Binding.DisplayName.IsEmpty())
-		{
-			OutName = FText::FromString(Binding.DisplayName);
-		} 
-		else if (!Binding.PropertyBindingChain.IsEmpty())
-		{
-			OutName = FText::FromName(Binding.PropertyBindingChain.Last());
-		}
-	}
-#endif
+	CHOOSER_PARAMETER_BOILERPLATE();
 };
 
 USTRUCT()
diff --git a/Engine/Plugins/Experimental/Chooser/Source/Chooser/Internal/ObjectColumn.h b/Engine/Plugins/Experimental/Chooser/Source/Chooser/Internal/ObjectColumn.h
--- a/Engine/Plugins/Experimental/Chooser/Source/Chooser/Internal/ObjectColumn.h
+++ b/Engine/Plugins/Experimental/Chooser/Source/Chooser/Internal/ObjectColumn.h
@@ -20,21 +20,10 @@ struct CHOOSER_API FObjectContextProperty : public FChooserParameterObjectBase
 	FChooserObjectPropertyBinding Binding;
 
 	virtual bool GetValue(FChooserEvaluationContext& Context, FSoftObjectPath& OutResult) const override;
-	
-	virtual void Compile(IHasContextClass* Owner, bool bForce) override
-	{
-		Binding.Compile(Owner, bForce);
-	};
 
-#if WITH_EDITOR
-	virtual void GetDisplayName(FText& OutName) const override
-	{
-		if (!Binding.PropertyBindingChain.IsEmpty())
-		{
-			OutName = FText::FromName(Binding.PropertyBindingChain.Last());
-		}
-	}
+	CHOOSER_PARAMETER_BOILERPLATE();
 
+#if WITH_EDITOR
 	virtual UClass* GetAllowedClass() const override { return Binding.AllowedClass; }
 #endif
 };
diff --git a/Engine/Plugins/Experimental/Chooser/Source/Chooser/Internal/OutputStructColumn.h b/Engine/Plugins/Experimental/Chooser/Source/Chooser/Internal/OutputStructColumn.h
--- a/Engine/Plugins/Experimental/Chooser/Source/Chooser/Internal/OutputStructColumn.h
+++ b/Engine/Plugins/Experimental/Chooser/Source/Chooser/Internal/OutputStructColumn.h
@@ -20,24 +20,9 @@ struct CHOOSER_API FStructContextProperty : public FChooserParameterStructBase
 
 	virtual bool SetValue(FChooserEvaluationContext& Context, const FInstancedStruct &Value) const override;
 
-	virtual void Compile(IHasContextClass* Owner, bool bForce) override
- 	{
- 		Binding.Compile(Owner, bForce);
- 	};
+	CHOOSER_PARAMETER_BOILERPLATE();
 
 #if WITH_EDITOR
-	virtual void GetDisplayName(FText& OutName) const override
-	{
-		if (!Binding.DisplayName.IsEmpty())
-		{
-			OutName = FText::FromString(Binding.DisplayName);
-		} 
-		else if (!Binding.PropertyBindingChain.IsEmpty())
-		{
-			OutName = FText::FromName(Binding.PropertyBindingChain.Last());
-		}
-	}
-
 	virtual UScriptStruct* GetStructType() const override { return Binding.StructType; }
 #endif
 };
diff --git a/Engine/Plugins/Experimental/Chooser/Source/Chooser/Internal/RandomizeColumn.h b/Engine/Plugins/Experimental/Chooser/Source/Chooser/Internal/RandomizeColumn.h
--- a/Engine/Plugins/Experimental/Chooser/Source/Chooser/Internal/RandomizeColumn.h
+++ b/Engine/Plugins/Experimental/Chooser/Source/Chooser/Internal/RandomizeColumn.h
@@ -18,25 +18,8 @@ public:
 	FChooserPropertyBinding Binding;
 	
 	virtual bool GetValue(FChooserEvaluationContext& Context, const FChooserRandomizationContext*& OutResult) const override;
-	
-	virtual void Compile(IHasContextClass* Owner, bool bForce) override
-	{
-		Binding.Compile(Owner, bForce);
-	};
 
-#if WITH_EDITOR
-	virtual void GetDisplayName(FText& OutName) const override
-	{
-		if (!Binding.DisplayName.IsEmpty())
-		{
-			OutName = FText::FromString(Binding.DisplayName);
-		} 
-		else if (!Binding.PropertyBindingChain.IsEmpty())
-		{
-			OutName = FText::FromName(Binding.PropertyBindingChain.Last());
-		}
-	}
-#endif
+	CHOOSER_PARAMETER_BOILERPLATE();
 };
 
 USTRUCT()
diff --git a/Engine/Plugins/Experimental/Chooser/Source/Chooser/Private/Chooser.cpp b/Engine/Plugins/Experimental/Chooser/Source/Chooser/Private/Chooser.cpp
--- a/Engine/Plugins/Experimental/Chooser/Source/Chooser/Private/Chooser.cpp
+++ b/Engine/Plugins/Experimental/Chooser/Source/Chooser/Private/Chooser.cpp
@@ -12,6 +12,12 @@
 
 DEFINE_LOG_CATEGORY(LogChooser)
 
+
+#if WITH_EDITOR
+const FName UChooserTable::PropertyNamesTag = "ChooserPropertyNames";
+const FString UChooserTable::PropertyTagDelimiter = TEXT(";");
+#endif
+
 UChooserTable::UChooserTable(const FObjectInitializer& Initializer)
 	:Super(Initializer)
 {
@@ -83,6 +89,28 @@ void UChooserTable::BeginDestroy()
 
 #if WITH_EDITOR
 
+void UChooserTable::GetAssetRegistryTags(TArray<FAssetRegistryTag>& OutTags) const
+{
+	Super::GetAssetRegistryTags(OutTags);
+
+	// Output property names we use
+	TStringBuilder<256> PropertyNamesBuilder;
+	PropertyNamesBuilder.Append(PropertyTagDelimiter);
+
+	for(const FInstancedStruct& Column : ColumnsStructs)
+	{
+		if (const FChooserColumnBase* ColumnBase = Column.GetPtr<FChooserColumnBase>())
+		{
+			if (FChooserParameterBase* Parameter = const_cast<FChooserColumnBase*>(ColumnBase)->GetInputValue())
+			{
+				Parameter->AddSearchNames(PropertyNamesBuilder);
+			}
+		}
+	}
+
+	OutTags.Add(FAssetRegistryTag(PropertyNamesTag, PropertyNamesBuilder.ToString(), FAssetRegistryTag::TT_Hidden));
+}
+
 void UChooserTable::AddCompileDependency(const UStruct* InStructType)
 {
 	UStruct* StructType = const_cast<UStruct*>(InStructType);
diff --git a/Engine/Plugins/Experimental/Chooser/Source/Chooser/Public/ChooserPropertyAccess.h b/Engine/Plugins/Experimental/Chooser/Source/Chooser/Public/ChooserPropertyAccess.h
--- a/Engine/Plugins/Experimental/Chooser/Source/Chooser/Public/ChooserPropertyAccess.h
+++ b/Engine/Plugins/Experimental/Chooser/Source/Chooser/Public/ChooserPropertyAccess.h
@@ -373,3 +373,63 @@ bool FChooserPropertyBinding::SetValue(FChooserEvaluationContext& Context, const
 	}
 	return false;
 }
+
+#if WITH_EDITOR
+#define CHOOSER_PARAMETER_BOILERPLATE() \
+	virtual void Compile(IHasContextClass* Owner, bool bForce) override\
+	{\
+		Binding.Compile(Owner, bForce);\
+	};\
+	\
+	virtual void AddSearchNames(FStringBuilderBase& Builder) const override\
+	{\
+		for (const FName& Entry : Binding.PropertyBindingChain)\
+		{\
+			Builder.Append(Entry.ToString());\
+			Builder.Append(";");\
+		}\
+	}\
+	\
+	virtual void GetDisplayName(FText& OutName) const override\
+	{\
+		if (!Binding.DisplayName.IsEmpty())\
+		{\
+			OutName = FText::FromString(Binding.DisplayName);\
+		} \
+		else if (!Binding.PropertyBindingChain.IsEmpty())\
+		{\
+			OutName = FText::FromName(Binding.PropertyBindingChain.Last());\
+		}\
+	}\
+	\
+	virtual void ReplaceString(FStringView FindString, ESearchCase::Type SearchCase, bool bFindWholeWord, FStringView ReplaceString) override\
+	{\
+		for (FName& Entry : Binding.PropertyBindingChain)\
+		{\
+			if(bFindWholeWord)\
+			{\
+				if(Entry.ToString().Compare(FString(FindString), SearchCase) == 0)\
+				{\
+					Entry = FName(ReplaceString);\
+					Binding.DisplayName = "";\
+				}\
+			}\
+			else\
+			{\
+				if(UE::String::FindFirst(Entry.ToString(), FindString, SearchCase) != INDEX_NONE)\
+				{\
+					FString NewString = Entry.ToString().Replace(FindString.begin(), ReplaceString.begin(), SearchCase);\
+					Entry = FName(NewString);\
+					Binding.DisplayName = "";\
+				}\
+			}\
+		}\
+	}
+#else
+#define CHOOSER_PARAMETER_BOILERPLATE() \
+	virtual void Compile(IHasContextClass* Owner, bool bForce) override\
+	{\
+		Binding.Compile(Owner, bForce);\
+	}
+#endif
+
diff --git a/Engine/Plugins/Experimental/Chooser/Source/Chooser/Public/IChooserParameterBase.h b/Engine/Plugins/Experimental/Chooser/Source/Chooser/Public/IChooserParameterBase.h
--- a/Engine/Plugins/Experimental/Chooser/Source/Chooser/Public/IChooserParameterBase.h
+++ b/Engine/Plugins/Experimental/Chooser/Source/Chooser/Public/IChooserParameterBase.h
@@ -7,6 +7,8 @@
 #include "IObjectChooser.h"
 #include "IChooserParameterBase.generated.h"
 
+struct FAssetRegistryTag;
+
 USTRUCT()
 struct FChooserParameterBase
 {
@@ -20,6 +22,8 @@ struct FChooserParameterBase
 	}
 	
 	virtual void GetDisplayName(FText& OutName) const { }
+	virtual void AddSearchNames(FStringBuilderBase& Builder) const { }
+	virtual void ReplaceString(FStringView FindString, ESearchCase::Type, bool MatchWholeWord, FStringView ReplaceString) { }
 
 	virtual void PostLoad() {};
 	virtual void Compile(IHasContextClass* Owner, bool bForce) {};
diff --git a/Engine/Plugins/Experimental/Chooser/Source/ChooserEditor/ChooserEditor.Build.cs b/Engine/Plugins/Experimental/Chooser/Source/ChooserEditor/ChooserEditor.Build.cs
--- a/Engine/Plugins/Experimental/Chooser/Source/ChooserEditor/ChooserEditor.Build.cs
+++ b/Engine/Plugins/Experimental/Chooser/Source/ChooserEditor/ChooserEditor.Build.cs
@@ -38,6 +38,7 @@ namespace UnrealBuildTool.Rules
 					"TraceInsights",
 					"RewindDebuggerInterface",
 					"GameplayInsights",
+					"Persona",
 					// ... add private dependencies that you statically link with here ...
 				}
 			);
diff --git a/Engine/Plugins/Experimental/Chooser/Source/ChooserEditor/Private/ChooserFindProperties.cpp b/Engine/Plugins/Experimental/Chooser/Source/ChooserEditor/Private/ChooserFindProperties.cpp
new file mode 100644
--- /dev/null
+++ b/Engine/Plugins/Experimental/Chooser/Source/ChooserEditor/Private/ChooserFindProperties.cpp
@@ -0,0 +1,170 @@
+﻿// Copyright Epic Games, Inc. All Rights Reserved.
+
+#include "ChooserFindProperties.h"
+
+#include "Chooser.h"
+#include "AssetRegistry/AssetRegistryModule.h"
+#include "String/ParseTokens.h"
+
+#define LOCTEXT_NAMESPACE "ChooserFindProperties"
+
+FString UChooserFindProperties::GetFindResultStringFromAssetData(const FAssetData& InAssetData) const
+{
+	if(GetFindString().IsEmpty())
+	{
+		return FString();
+	}
+
+	auto GetMatchingPropertyNamesForAsset = [this](const FAssetData& InAssetData, TArray<FString>& OutPropertyNames)
+	{
+		const FString TagValue = InAssetData.GetTagValueRef<FString>(UChooserTable::PropertyNamesTag);
+		if (!TagValue.IsEmpty())
+		{
+			UE::String::ParseTokens(TagValue, *UChooserTable::PropertyTagDelimiter, [this, &OutPropertyNames](FStringView InToken)
+			{
+				if(GetFindWholeWord())
+				{
+					if(InToken.Compare(GetFindString(), GetSearchCase()) == 0)
+					{
+						OutPropertyNames.Add(FString(InToken));
+					}
+				}
+				else
+				{
+					if(UE::String::FindFirst(InToken, GetFindString(), GetSearchCase()) != INDEX_NONE)
+					{
+						OutPropertyNames.Add(FString(InToken));
+					}
+				}
+			}, UE::String::EParseTokensOptions::SkipEmpty);
+		}
+	};
+
+	TStringBuilder<128> Builder;
+	TArray<FString> PropertyNames;
+	GetMatchingPropertyNamesForAsset(InAssetData, PropertyNames);
+	if(PropertyNames.Num() > 0)
+	{
+		for(int32 NameIndex = 0; NameIndex < PropertyNames.Num(); ++NameIndex)
+		{
+			Builder.Append(PropertyNames[NameIndex]);
+			if(NameIndex != PropertyNames.Num() - 1)
+			{
+				Builder.Append(TEXT(", "));
+			}
+		}
+	}
+	return FString(Builder.ToString());
+}
+
+TConstArrayView<UClass*> UChooserFindProperties::GetSupportedAssetTypes() const
+{
+	static UClass* Types[] = { UChooserTable::StaticClass() };
+	return Types;
+}
+
+bool UChooserFindProperties::ShouldFilterOutAsset(const FAssetData& InAssetData, bool& bOutIsOldAsset) const
+{
+	FString TagValue;
+	if(InAssetData.GetTagValue<FString>(UChooserTable::PropertyNamesTag, TagValue))
+	{
+		bOutIsOldAsset = false;
+
+		if(GetFindString().IsEmpty())
+		{
+			return true;
+		}
+		
+		bool bFoundMatch = false;
+		UE::String::ParseTokens(TagValue, *UChooserTable::PropertyTagDelimiter, [this, &bFoundMatch](FStringView InToken)
+		{
+			if(NameMatches(InToken))
+			{
+				bFoundMatch = true;
+			}
+		}, UE::String::EParseTokensOptions::SkipEmpty);
+
+		if(bFoundMatch)
+		{
+			return false;
+		}
+	}
+	
+	bOutIsOldAsset = false;
+	return true;
+}
+
+void UChooserFindProperties::ReplaceInAsset(const FAssetData& InAssetData) const
+{
+	if(UObject* Asset = InAssetData.GetAsset())
+	{
+		if(UChooserTable* ChooserTable = Cast<UChooserTable>(Asset))
+		{
+			Asset->MarkPackageDirty();
+
+			for(FInstancedStruct& Column : ChooserTable->ColumnsStructs)
+			{
+				if (FChooserColumnBase* ColumnData = Column.GetMutablePtr<FChooserColumnBase>())
+				{
+					FStringBuilderBase StringBuilder;
+					ColumnData->GetInputValue()->ReplaceString(GetFindString(),GetSearchCase(),GetFindWholeWord(), GetReplaceString());
+				}
+			}
+	
+		}
+	}
+}
+
+void UChooserFindProperties::RemoveInAsset(const FAssetData& InAssetData) const
+{
+	if(UObject* Asset = InAssetData.GetAsset())
+	{
+		if(UChooserTable* ChooserTable = Cast<UChooserTable>(Asset))
+		{
+			Asset->MarkPackageDirty();
+
+			for(FInstancedStruct& Column : ChooserTable->ColumnsStructs)
+			{
+				if (FChooserColumnBase* ColumnData = Column.GetMutablePtr<FChooserColumnBase>())
+				{
+					FStringBuilderBase StringBuilder;
+					ColumnData->GetInputValue()->AddSearchNames(StringBuilder);
+					
+					bool bMatched = false;
+					
+					UE::String::ParseTokens(StringBuilder, *UChooserTable::PropertyTagDelimiter, [this, &bMatched](FStringView InToken)
+					{
+						if (NameMatches(InToken))
+						{
+							bMatched = true;
+						}
+					}, UE::String::EParseTokensOptions::SkipEmpty);
+
+					if (bMatched)
+					{
+						// if the name matched, reset the input parameter
+						ColumnData->SetInputType(ColumnData->GetInputType());
+					}
+				}
+			}
+	
+		}
+	}
+}
+
+void UChooserFindProperties::GetAutoCompleteNames(TArrayView<FAssetData> InAssetDatas, TSet<FString>& OutUniqueNames) const
+{
+	for (const FAssetData& AssetData : InAssetDatas)
+	{
+		const FString TagValue = AssetData.GetTagValueRef<FString>(UChooserTable::PropertyNamesTag);
+		if (!TagValue.IsEmpty())
+		{
+			UE::String::ParseTokens(TagValue, *UChooserTable::PropertyTagDelimiter, [&OutUniqueNames](FStringView InToken)
+			{
+				OutUniqueNames.Add(FString(InToken));
+			}, UE::String::EParseTokensOptions::SkipEmpty);
+		}
+	}
+}
+
+#undef LOCTEXT_NAMESPACE
\ No newline at end of file
diff --git a/Engine/Plugins/Experimental/Chooser/Source/ChooserEditor/Private/ChooserFindProperties.h b/Engine/Plugins/Experimental/Chooser/Source/ChooserEditor/Private/ChooserFindProperties.h
new file mode 100644
--- /dev/null
+++ b/Engine/Plugins/Experimental/Chooser/Source/ChooserEditor/Private/ChooserFindProperties.h
@@ -0,0 +1,24 @@
+﻿// Copyright Epic Games, Inc. All Rights Reserved.
+
+#pragma once
+
+#include "CoreMinimal.h"
+#include "AnimAssetFindReplace.h"
+#include "ChooserFindProperties.generated.h"
+
+/** Find, replace and remove curves across assets */
+UCLASS(MinimalAPI, DisplayName="ChooserProperties")
+class UChooserFindProperties : public UAnimAssetFindReplaceProcessor_StringBase
+{
+	GENERATED_BODY()
+
+	// UAnimAssetFindReplaceProcessor interface
+	virtual FString GetFindResultStringFromAssetData(const FAssetData& InAssetData) const override;
+	virtual TConstArrayView<UClass*> GetSupportedAssetTypes() const override;
+	virtual bool ShouldFilterOutAsset(const FAssetData& InAssetData, bool& bOutIsOldAsset) const override;
+	virtual void ReplaceInAsset(const FAssetData& InAssetData) const override;
+	virtual void RemoveInAsset(const FAssetData& InAssetData) const override;
+
+	// UAnimAssetFindReplaceProcessor_StringBase interface
+	virtual void GetAutoCompleteNames(TArrayView<FAssetData> InAssetDatas, TSet<FString>& OutUniqueNames) const override;
+};
diff --git a/Engine/Plugins/Experimental/Chooser/Source/ChooserEditor/Private/ChooserTableEditor.cpp b/Engine/Plugins/Experimental/Chooser/Source/ChooserEditor/Private/ChooserTableEditor.cpp
--- a/Engine/Plugins/Experimental/Chooser/Source/ChooserEditor/Private/ChooserTableEditor.cpp
+++ b/Engine/Plugins/Experimental/Chooser/Source/ChooserEditor/Private/ChooserTableEditor.cpp
@@ -18,6 +18,7 @@
 #include "StructViewerModule.h"
 #include "SourceCodeNavigation.h"
 #include "Chooser.h"
+#include "ChooserFindProperties.h"
 #include "DetailCategoryBuilder.h"
 #include "DetailLayoutBuilder.h"
 #include "ClassViewerFilter.h"
@@ -31,6 +32,7 @@
 #include "ObjectChooserWidgetFactories.h"
 #include "GraphEditorSettings.h"
 #include "IDetailCustomization.h"
+#include "PersonaModule.h"
 #include "PropertyCustomizationHelpers.h"
 #include "ScopedTransaction.h"
 #include "Widgets/Layout/SSeparator.h"
@@ -46,6 +48,7 @@ namespace UE::ChooserEditor
 
 const FName FChooserTableEditor::ToolkitFName( TEXT( "ChooserTableEditor" ) );
 const FName FChooserTableEditor::PropertiesTabId( TEXT( "ChooserEditor_Properties" ) );
+const FName FChooserTableEditor::FindReplaceTabId( TEXT( "ChooserEditor_FindReplace" ) );
 const FName FChooserTableEditor::TableTabId( TEXT( "ChooserEditor_Table" ) );
 	
 void FChooserTableEditor::RegisterTabSpawners(const TSharedRef<class FTabManager>& InTabManager)
@@ -63,6 +66,11 @@ void FChooserTableEditor::RegisterTabSpawners(const TSharedRef<class FTabManager
 		.SetDisplayName( LOCTEXT("TableTab", "Chooser Table") )
 		.SetGroup(WorkspaceMenuCategory.ToSharedRef())
 		.SetIcon(FSlateIcon("EditorStyle", "LevelEditor.Tabs.Details"));
+
+	InTabManager->RegisterTabSpawner( FindReplaceTabId, FOnSpawnTab::CreateSP(this, &FChooserTableEditor::SpawnFindReplaceTab) )
+		.SetDisplayName( LOCTEXT("FindReplaceTab", "Find/Replace") )
+		.SetGroup(WorkspaceMenuCategory.ToSharedRef())
+		.SetIcon(FSlateIcon(FAppStyle::GetAppStyleSetName(), "Icons.Find"));
 }
 	
 void FChooserTableEditor::UnregisterTabSpawners(const TSharedRef<class FTabManager>& InTabManager)
@@ -71,6 +79,7 @@ void FChooserTableEditor::UnregisterTabSpawners(const TSharedRef<class FTabManag
 
 	InTabManager->UnregisterTabSpawner( TableTabId );
 	InTabManager->UnregisterTabSpawner( PropertiesTabId );
+	InTabManager->UnregisterTabSpawner( FindReplaceTabId );
 }
 
 const FName FChooserTableEditor::ChooserEditorAppIdentifier( TEXT( "ChooserEditorApp" ) );
@@ -292,6 +301,12 @@ void FChooserTableEditor::InitEditor( const EToolkitMode::Type Mode, const TShar
 	RegisterToolbar();
 
 	SelectRootProperties();
+
+		
+	FAnimAssetFindReplaceConfig FindReplaceConfig;
+	FindReplaceConfig.InitialProcessorClass = UChooserFindProperties::StaticClass();
+
+	FPersonaModule& PersonaModule = FModuleManager::LoadModuleChecked<FPersonaModule>("Persona");
 }
 
 FName FChooserTableEditor::GetToolkitFName() const
@@ -465,7 +480,21 @@ TSharedRef<SDockTab> FChooserTableEditor::SpawnPropertiesTab( const FSpawnTabArg
 			DetailsView.ToSharedRef()
 		];
 }
+	
+TSharedRef<SDockTab> FChooserTableEditor::SpawnFindReplaceTab( const FSpawnTabArgs& Args )
+{
+	check( Args.GetTabId() == FindReplaceTabId );
 
+	FPersonaModule& PersonaModule = FModuleManager::LoadModuleChecked<FPersonaModule>("Persona");
+	FAnimAssetFindReplaceConfig Config;
+	Config.InitialProcessorClass = UChooserFindProperties::StaticClass();
+	return SNew(SDockTab)
+		.Label( LOCTEXT("FindReplaceTitle", "Find/Replace") )
+		.TabColorScale( GetTabColorScale() )
+	[
+		PersonaModule.CreateFindReplaceWidget(Config)
+	];
+}
 
 class FChooserRowDragDropOp : public FDecoratedDragDropOp
 {
diff --git a/Engine/Plugins/Experimental/Chooser/Source/ChooserEditor/Private/ChooserTableEditor.h b/Engine/Plugins/Experimental/Chooser/Source/ChooserEditor/Private/ChooserTableEditor.h
--- a/Engine/Plugins/Experimental/Chooser/Source/ChooserEditor/Private/ChooserTableEditor.h
+++ b/Engine/Plugins/Experimental/Chooser/Source/ChooserEditor/Private/ChooserTableEditor.h
@@ -117,6 +117,8 @@ namespace UE::ChooserEditor
 		TSharedRef<SDockTab> SpawnPropertiesTab( const FSpawnTabArgs& Args );
 		/** Create the table tab and its content */
 		TSharedRef<SDockTab> SpawnTableTab( const FSpawnTabArgs& Args );
+		/** Create the find/replace tab and its content */
+		TSharedRef<SDockTab> SpawnFindReplaceTab( const FSpawnTabArgs& Args );
 	
 		TSharedRef<ITableRow> GenerateTableRow(TSharedPtr<FChooserTableRow> InItem, const TSharedRef<STableViewBase>& OwnerTable);
 
@@ -131,6 +133,7 @@ namespace UE::ChooserEditor
 
 		/**	The tab ids for all the tabs used */
 		static const FName PropertiesTabId;
+		static const FName FindReplaceTabId;
 		static const FName TableTabId;
 
 		/** The objects open within this editor */
diff --git a/Engine/Plugins/Experimental/Chooser/Source/ProxyTable/Internal/LookupProxy.h b/Engine/Plugins/Experimental/Chooser/Source/ProxyTable/Internal/LookupProxy.h
--- a/Engine/Plugins/Experimental/Chooser/Source/ProxyTable/Internal/LookupProxy.h
+++ b/Engine/Plugins/Experimental/Chooser/Source/ProxyTable/Internal/LookupProxy.h
@@ -21,11 +21,8 @@ public:
 	FChooserPropertyBinding Binding;
 
 	virtual bool GetValue(FChooserEvaluationContext& Context, const UProxyTable*& OutResult) const override;
-	
-	virtual void Compile(IHasContextClass* Owner, bool bForce) override
-	{
-		Binding.Compile(Owner, bForce);
-	};
+
+	CHOOSER_PARAMETER_BOILERPLATE();
 };
 
 USTRUCT()
diff --git a/Engine/Source/Editor/Persona/Private/PersonaModule.cpp b/Engine/Source/Editor/Persona/Private/PersonaModule.cpp
--- a/Engine/Source/Editor/Persona/Private/PersonaModule.cpp
+++ b/Engine/Source/Editor/Persona/Private/PersonaModule.cpp
@@ -386,6 +386,12 @@ TSharedRef<FWorkflowTabFactory> FPersonaModule::CreateAnimAssetFindReplaceTabFac
 	return MakeShared<FAnimAssetFindReplaceSummoner>(InHostingApp, InConfig);
 }
 
+TSharedRef<SWidget> FPersonaModule::CreateFindReplaceWidget(const FAnimAssetFindReplaceConfig& InConfig) const
+{
+	return SNew(SAnimAssetFindReplace)
+		.Config(InConfig);
+}
+
 TSharedRef<SWidget> FPersonaModule::CreateEditorWidgetForAnimDocument(const TSharedRef<IAnimationEditor>& InHostingApp, UObject* InAnimAsset, const FAnimDocumentArgs& InArgs, FString& OutDocumentLink)
 {
 	TSharedPtr<SWidget> Result = SNullWidget::NullWidget;
diff --git a/Engine/Source/Editor/Persona/Private/SAnimAssetFindReplace.cpp b/Engine/Source/Editor/Persona/Private/SAnimAssetFindReplace.cpp
--- a/Engine/Source/Editor/Persona/Private/SAnimAssetFindReplace.cpp
+++ b/Engine/Source/Editor/Persona/Private/SAnimAssetFindReplace.cpp
@@ -456,7 +456,8 @@ void SAnimAssetFindReplace::RefreshSearchResults()
 {
 	bFoundAssets = false;
 	OldAssets.Empty();
-	SetARFilterDelegate.ExecuteIfBound(MakeARFilter());
+	AssetPickerConfig.Filter = MakeARFilter();
+	SetARFilterDelegate.ExecuteIfBound(AssetPickerConfig.Filter);
 	RefreshAssetViewDelegate.ExecuteIfBound(true);
 }
 
diff --git a/Engine/Source/Editor/Persona/Public/AnimAssetFindReplace.h b/Engine/Source/Editor/Persona/Public/AnimAssetFindReplace.h
--- a/Engine/Source/Editor/Persona/Public/AnimAssetFindReplace.h
+++ b/Engine/Source/Editor/Persona/Public/AnimAssetFindReplace.h
@@ -32,8 +32,8 @@ enum class EAnimAssetFindReplaceMode : int32
 };
 
 /** Processor base class to allow systems to add their own find/replace functionality */
-UCLASS(MinimalAPI, Abstract)
-class UAnimAssetFindReplaceProcessor : public UObject
+UCLASS(Abstract)
+class PERSONA_API UAnimAssetFindReplaceProcessor : public UObject
 {
 	GENERATED_BODY()
 
@@ -99,8 +99,8 @@ private:
 };
 
 /** Processor for string-based animation asset operations */
-UCLASS(MinimalAPI, Abstract)
-class UAnimAssetFindReplaceProcessor_StringBase : public UAnimAssetFindReplaceProcessor
+UCLASS(Abstract)
+class PERSONA_API UAnimAssetFindReplaceProcessor_StringBase : public UAnimAssetFindReplaceProcessor
 {
 	GENERATED_BODY()
 
@@ -110,7 +110,7 @@ public:
 	{}
 
 	/** Set the current find string */
-	PERSONA_API void SetFindString(const FString& InString);
+	void SetFindString(const FString& InString);
 
 	/** Get the current find string */
 	FStringView GetFindString() const { return FindString; }
@@ -119,7 +119,7 @@ public:
 	const FString& GetFindStringRef() const { return FindString; }
 
 	/** Set the current replace string */
-	PERSONA_API void SetReplaceString(const FString& InString);
+	void SetReplaceString(const FString& InString);
 
 	/** Get the current replace string */
 	FStringView GetReplaceString() const { return ReplaceString; }
@@ -128,19 +128,19 @@ public:
 	const FString& GetReplaceStringRef() const { return ReplaceString; }
 	
 	/** Set the current skeleton filter */
-	PERSONA_API void SetSkeletonFilter(const FAssetData& InSkeletonFilter);
+	void SetSkeletonFilter(const FAssetData& InSkeletonFilter);
 
 	/** Get the current skeleton filter **/
 	const FAssetData& GetSkeletonFilter() const { return SkeletonFilter; }
 
 	/** Sets whether to match whole words when searching strings */
-	PERSONA_API void SetFindWholeWord(bool bInFindWholeWord);
+	void SetFindWholeWord(bool bInFindWholeWord);
 
 	/** Gets whether to match whole words when searching strings */
 	bool GetFindWholeWord() const { return bFindWholeWord; }
 
 	/** Sets case comparison to use when searching */
-	PERSONA_API void SetSearchCase(ESearchCase::Type InSearchCase);
+	void SetSearchCase(ESearchCase::Type InSearchCase);
 
 	/** Gets case comparison to use when searching */
 	ESearchCase::Type GetSearchCase() const { return SearchCase; }
diff --git a/Engine/Source/Editor/Persona/Public/PersonaModule.h b/Engine/Source/Editor/Persona/Public/PersonaModule.h
--- a/Engine/Source/Editor/Persona/Public/PersonaModule.h
+++ b/Engine/Source/Editor/Persona/Public/PersonaModule.h
@@ -353,6 +353,9 @@ public:
 	/** Create a tab factory for finding and replacing in anim data */
 	virtual TSharedRef<FWorkflowTabFactory> CreateAnimAssetFindReplaceTabFactory(const TSharedRef<FWorkflowCentricApplication>& InHostingApp, const FAnimAssetFindReplaceConfig& InConfig) const;
 	
+	/** Create a widget for finding and replacing in anim data */
+	virtual TSharedRef<SWidget> CreateFindReplaceWidget(const FAnimAssetFindReplaceConfig& InConfig) const;
+	
 	/** Create a widget that acts as a document for an animation asset */
 	virtual TSharedRef<SWidget> CreateEditorWidgetForAnimDocument(const TSharedRef<IAnimationEditor>& InHostingApp, UObject* InAnimAsset, const FAnimDocumentArgs& InArgs, FString& OutDocumentLink);
 
