Support for nested Chooser Tables in the same file, and Editor support for a Fallback result in choosers

diff --git a/Engine/Plugins/Experimental/Chooser/Chooser.uplugin b/Engine/Plugins/Experimental/Chooser/Chooser.uplugin
--- a/Engine/Plugins/Experimental/Chooser/Chooser.uplugin
+++ b/Engine/Plugins/Experimental/Chooser/Chooser.uplugin
@@ -11,7 +11,7 @@
 	"MarketplaceURL" : "",
 	"SupportURL" : "",
 	"EnabledByDefault" : false,
-	"CanContainContent" : false,
+	"CanContainContent" : true,
 	"IsBetaVersion" : false,
 	"IsExperimentalVersion" : true,
 	"Installed" : false,
diff --git a/Engine/Plugins/Experimental/Chooser/Source/Chooser/Internal/Chooser.h b/Engine/Plugins/Experimental/Chooser/Source/Chooser/Internal/Chooser.h
--- a/Engine/Plugins/Experimental/Chooser/Source/Chooser/Internal/Chooser.h
+++ b/Engine/Plugins/Experimental/Chooser/Source/Chooser/Internal/Chooser.h
@@ -96,6 +96,11 @@ public:
 	TArray<TScriptInterface<IChooserColumn>> Columns_DEPRECATED;
 #endif
 
+	UChooserTable* GetContextOwner() { return ParentTable ? ParentTable.Get() : this; }
+	const UChooserTable* GetContextOwner() const { return ParentTable ? ParentTable.Get() : this; }
+
+	UPROPERTY()
+	TObjectPtr<UChooserTable> ParentTable;
 	
 	// FallbackResult will be used as the Result if there are no rows in the chooser which pass all filters.  If FallbackResult is not assigned, then the Chooser will return null in that case.
 	UPROPERTY(EditAnywhere, Meta = (ExcludeBaseStruct, BaseStruct = "/Script/Chooser.ObjectChooserBase"), Category = "Fallback")
@@ -122,6 +127,22 @@ public:
 	static FObjectChooserBase::EIteratorStatus EvaluateChooser(FChooserEvaluationContext& Context, const UChooserTable* Chooser, FObjectChooserBase::FObjectChooserIteratorCallback Callback);
 };
 
+USTRUCT(BlueprintType, DisplayName = "Nested Chooser")
+struct CHOOSER_API FNestedChooser : public FObjectChooserBase
+{
+	GENERATED_BODY()
+
+	virtual UObject* ChooseObject(FChooserEvaluationContext& Context) const final override;
+	virtual EIteratorStatus ChooseMulti(FChooserEvaluationContext &Context, FObjectChooserIteratorCallback Callback) const final override;
+	virtual void GetDebugName(FString& OutDebugName) const override;
+	
+	public:
+
+	FNestedChooser();
+	
+	UPROPERTY()
+	TObjectPtr<UChooserTable> Chooser;
+};
 
 USTRUCT(BlueprintType, DisplayName = "Evaluate Chooser")
 struct CHOOSER_API FEvaluateChooser : public FObjectChooserBase
diff --git a/Engine/Plugins/Experimental/Chooser/Source/Chooser/Private/Chooser.cpp b/Engine/Plugins/Experimental/Chooser/Source/Chooser/Private/Chooser.cpp
--- a/Engine/Plugins/Experimental/Chooser/Source/Chooser/Private/Chooser.cpp
+++ b/Engine/Plugins/Experimental/Chooser/Source/Chooser/Private/Chooser.cpp
@@ -7,6 +7,7 @@
 #include "Engine/UserDefinedStruct.h"
 #include "Engine/Blueprint.h"
 #include "ChooserIndexArray.h"
+#include "IChooserParameterGameplayTag.h"
 
 #include UE_INLINE_GENERATED_CPP_BY_NAME(Chooser)
 
@@ -136,12 +137,14 @@ void UChooserTable::AddCompileDependency(const UStruct* InStructType)
 
 void UChooserTable::Compile(bool bForce)
 {
+	IHasContextClass* ContextOwner = GetContextOwner();
+	
 	for (FInstancedStruct& ColumnData : ColumnsStructs)
 	{
 		if (ColumnData.IsValid())
 		{
 			FChooserColumnBase& Column = ColumnData.GetMutable<FChooserColumnBase>();
-			Column.Compile(this, bForce);
+			Column.Compile(ContextOwner, bForce);
 		}
 	}
 
@@ -150,7 +153,7 @@ void UChooserTable::Compile(bool bForce)
 		if (ResultData.IsValid())
 		{
 			FObjectChooserBase& Result = ResultData.GetMutable<FObjectChooserBase>();
-			Result.Compile(this, bForce);
+			Result.Compile(ContextOwner, bForce);
 		}
 	}
 }
@@ -212,6 +215,8 @@ void UChooserTable::IterateRecentContextObjects(TFunction<void(const UObject*)>
 void UChooserTable::UpdateDebugging(FChooserEvaluationContext& Context) const
 {
 	FScopeLock Lock(&DebugLock);
+
+	const UChooserTable* ContextOwner = GetContextOwner();
 	
 	for (const FStructView& Param : Context.Params)
 	{
@@ -221,16 +226,16 @@ void UChooserTable::UpdateDebugging(FChooserEvaluationContext& Context) const
 			{
 				RecentContextObjects.Add(MakeWeakObjectPtr(ContextObject));
 				
-				if (DebugTarget == nullptr && !DebugTargetName.IsEmpty())
+				if (ContextOwner->DebugTarget == nullptr && !DebugTargetName.IsEmpty())
 				{
 					// if the DebugTargetName is set, but not the DebugTarget, it means that PIE has been restarted, so try matching by name
-					if (ContextObject->GetName() == DebugTargetName)
+					if (ContextObject->GetName() == ContextOwner->DebugTargetName)
 					{
-						DebugTarget = ContextObject;
+						ContextOwner->DebugTarget = ContextObject;
 					}
 				}
 				
-				if (ContextObject == DebugTarget) 
+				if (ContextObject == ContextOwner->DebugTarget) 
 				{
 					bDebugTestValuesValid = true;
 					Context.DebuggingInfo.bCurrentDebugTarget = true;
@@ -386,3 +391,30 @@ void FEvaluateChooser::GetDebugName(FString& OutDebugName) const
 		OutDebugName = Chooser.GetName();
 	}
 }
+
+
+FNestedChooser::FNestedChooser()
+{
+}
+
+UObject* FNestedChooser::ChooseObject(FChooserEvaluationContext& Context) const
+{
+	UObject* Result = nullptr;
+	UChooserTable::EvaluateChooser(Context, Chooser, FObjectChooserIteratorCallback::CreateLambda([&Result](UObject* InResult)
+	{
+		Result = InResult;
+		return FObjectChooserBase::EIteratorStatus::Stop;
+	}));
+
+	return Result;
+}
+
+FObjectChooserBase::EIteratorStatus FNestedChooser::ChooseMulti(FChooserEvaluationContext& Context, FObjectChooserIteratorCallback Callback) const
+{
+	return UChooserTable::EvaluateChooser(Context, Chooser, Callback);
+}
+
+void FNestedChooser::GetDebugName(FString& OutDebugName) const
+{
+	OutDebugName  = GetNameSafe(Chooser);
+}
\ No newline at end of file
diff --git a/Engine/Plugins/Experimental/Chooser/Source/ChooserEditor/Private/AssetDefinition_ChooserTable.cpp b/Engine/Plugins/Experimental/Chooser/Source/ChooserEditor/Private/AssetDefinition_ChooserTable.cpp
--- a/Engine/Plugins/Experimental/Chooser/Source/ChooserEditor/Private/AssetDefinition_ChooserTable.cpp
+++ b/Engine/Plugins/Experimental/Chooser/Source/ChooserEditor/Private/AssetDefinition_ChooserTable.cpp
@@ -2,9 +2,21 @@
 
 #include "AssetDefinition_ChooserTable.h"
 #include "ChooserTableEditor.h"
+#include "Styling/SlateStyleRegistry.h"
+#include "ChooserEditorStyle.h"
 
 #define LOCTEXT_NAMESPACE "UAssetDefinition_ChooserTable" 
 
+const FSlateBrush* UAssetDefinition_ChooserTable::GetThumbnailBrush(const FAssetData& InAssetData, const FName InClassName) const
+{
+	return UE::ChooserEditor::FChooserEditorStyle::Get().GetBrush("ChooserEditor.ChooserTableIconLarge");
+}
+
+const FSlateBrush* UAssetDefinition_ChooserTable::GetIconBrush(const FAssetData& InAssetData, const FName InClassName) const
+{
+	return UE::ChooserEditor::FChooserEditorStyle::Get().GetBrush("ChooserEditor.ChooserTableIconSmall");
+}
+
 EAssetCommandResult UAssetDefinition_ChooserTable::OpenAssets(const FAssetOpenArgs& OpenArgs) const
 {
 	TArray<UObject*> Objects = OpenArgs.LoadObjects<UObject>();
diff --git a/Engine/Plugins/Experimental/Chooser/Source/ChooserEditor/Private/AssetDefinition_ChooserTable.h b/Engine/Plugins/Experimental/Chooser/Source/ChooserEditor/Private/AssetDefinition_ChooserTable.h
--- a/Engine/Plugins/Experimental/Chooser/Source/ChooserEditor/Private/AssetDefinition_ChooserTable.h
+++ b/Engine/Plugins/Experimental/Chooser/Source/ChooserEditor/Private/AssetDefinition_ChooserTable.h
@@ -16,6 +16,8 @@ public:
 	// UAssetDefinition Begin
 	virtual FText GetAssetDisplayName() const override { return NSLOCTEXT("AssetTypeActions", "AssetTypeActions_ChooserTable", "Chooser Table"); }
 	virtual FLinearColor GetAssetColor() const override { return FLinearColor(FColor(128,128,64)); }
+	virtual const FSlateBrush* GetIconBrush(const FAssetData& InAssetData, const FName InClassName) const override;
+	virtual const FSlateBrush* GetThumbnailBrush(const FAssetData& InAssetData, const FName InClassName) const override;
 	virtual TSoftClassPtr<UObject> GetAssetClass() const override { return UChooserTable::StaticClass(); }
 	virtual TConstArrayView<FAssetCategoryPath> GetAssetCategories() const override
 	{
diff --git a/Engine/Plugins/Experimental/Chooser/Source/ChooserEditor/Private/BoolColumnEditor.cpp b/Engine/Plugins/Experimental/Chooser/Source/ChooserEditor/Private/BoolColumnEditor.cpp
--- a/Engine/Plugins/Experimental/Chooser/Source/ChooserEditor/Private/BoolColumnEditor.cpp
+++ b/Engine/Plugins/Experimental/Chooser/Source/ChooserEditor/Private/BoolColumnEditor.cpp
@@ -21,7 +21,11 @@ TSharedRef<SWidget> CreateBoolColumnWidget(UChooserTable* Chooser, FChooserColum
 {
 	FBoolColumn* BoolColumn = static_cast<FBoolColumn*>(Column);
 	
-	if (Row < 0)
+	if (Row == ColumnWidget_SpecialIndex_Fallback)
+	{
+		return SNullWidget::NullWidget;
+	}
+	if (Row == ColumnWidget_SpecialIndex_Header)
 	{
 		// create column header widget
 		TSharedPtr<SWidget> InputValueWidget = nullptr;
@@ -111,7 +115,7 @@ TSharedRef<SWidget> CreateOutputBoolColumnWidget(UChooserTable* Chooser, FChoose
 {
 	FOutputBoolColumn* BoolColumn = static_cast<FOutputBoolColumn*>(Column);
 
-	if (Row < 0)
+	if (Row == ColumnWidget_SpecialIndex_Header)
 	{
 		// create column header widget
 		TSharedPtr<SWidget> InputValueWidget = nullptr;
@@ -159,6 +163,19 @@ TSharedRef<SWidget> CreateOutputBoolColumnWidget(UChooserTable* Chooser, FChoose
 
 		return ColumnHeaderWidget;
 	}
+	else if (Row == ColumnWidget_SpecialIndex_Fallback) 
+	{
+		// create fallback value widget
+		return SNew(SHorizontalBox)
+		+ SHorizontalBox::Slot().FillWidth(1)
+		+ SHorizontalBox::Slot().AutoWidth().VAlign(VAlign_Center)
+		[
+			SNew(SCheckBox)
+			.IsChecked_Lambda([BoolColumn]() { return BoolColumn->bFallbackValue ? ECheckBoxState::Checked : ECheckBoxState::Unchecked; })
+			.OnCheckStateChanged_Lambda([BoolColumn](ECheckBoxState State) { BoolColumn->bFallbackValue = State == ECheckBoxState::Checked; })
+		]
+		+ SHorizontalBox::Slot().FillWidth(1);
+	}
 
 	return
 	SNew(SHorizontalBox)
diff --git a/Engine/Plugins/Experimental/Chooser/Source/ChooserEditor/Private/ChooserEditorModule.cpp b/Engine/Plugins/Experimental/Chooser/Source/ChooserEditor/Private/ChooserEditorModule.cpp
--- a/Engine/Plugins/Experimental/Chooser/Source/ChooserEditor/Private/ChooserEditorModule.cpp
+++ b/Engine/Plugins/Experimental/Chooser/Source/ChooserEditor/Private/ChooserEditorModule.cpp
@@ -4,6 +4,7 @@
 
 #include "AnimNode_ChooserPlayer.h"
 #include "BoolColumnEditor.h"
+#include "ChooserEditorStyle.h"
 #include "ChooserPropertyAccess.h"
 #include "ChooserTableEditor.h"
 #include "ChooserTableEditorCommands.h"
@@ -31,6 +32,8 @@ FChoosersTrackCreator GChoosersTrackCreator;
 
 void FModule::StartupModule()
 {
+	FChooserEditorStyle::Initialize();
+	
 	FChooserTableEditor::RegisterWidgets();
 	RegisterGameplayTagWidgets();
 	RegisterFloatRangeWidgets();
@@ -45,7 +48,6 @@ void FModule::StartupModule()
 
 	FPropertyEditorModule& PropertyModule = FModuleManager::LoadModuleChecked<FPropertyEditorModule>("PropertyEditor");
 	
-	
 	PropertyModule.RegisterCustomPropertyTypeLayout("FloatProperty", FOnGetPropertyTypeCustomizationInstance::CreateLambda([] { return MakeShared<FFrameTimeCustomization>(); }), MakeShared<FFrameTimePropertyTypeIdentifier>());
 	PropertyModule.RegisterCustomPropertyTypeLayout(FAnimCurveOverride::StaticStruct()->GetFName(), FOnGetPropertyTypeCustomizationInstance::CreateLambda([] { return MakeShared<FCurveOverrideCustomization>(); }));
 	PropertyModule.RegisterCustomPropertyTypeLayout(FAnimCurveOverrideList::StaticStruct()->GetFName(), FOnGetPropertyTypeCustomizationInstance::CreateLambda([] { return MakeShared<FCurveOverrideListCustomization>(); }));
@@ -65,6 +67,8 @@ void FModule::ShutdownModule()
 	IModularFeatures::Get().UnregisterModularFeature(RewindDebugger::IRewindDebuggerTrackCreator::ModularFeatureName, &GChoosersTrackCreator);
 	IModularFeatures::Get().UnregisterModularFeature(TraceServices::ModuleFeatureName, &ChooserTraceModule);
 	FChooserTableEditorCommands::Unregister();
+	
+	FChooserEditorStyle::Shutdown();
 }
 
 }
diff --git a/Engine/Plugins/Experimental/Chooser/Source/ChooserEditor/Private/ChooserEditorStyle.cpp b/Engine/Plugins/Experimental/Chooser/Source/ChooserEditor/Private/ChooserEditorStyle.cpp
new file mode 100644
--- /dev/null
+++ b/Engine/Plugins/Experimental/Chooser/Source/ChooserEditor/Private/ChooserEditorStyle.cpp
@@ -0,0 +1,52 @@
+// Copyright Epic Games, Inc. All Rights Reserved.
+
+#include "ChooserEditorStyle.h"
+#include "Styling/CoreStyle.h"
+#include "Styling/SlateStyle.h"
+#include "Styling/SlateStyleMacros.h"
+#include "Styling/SlateStyleRegistry.h"
+
+namespace UE::ChooserEditor
+{
+	
+TSharedPtr< FChooserEditorStyle > FChooserEditorStyle::StyleInstance = nullptr;
+
+
+FChooserEditorStyle::FChooserEditorStyle() :
+	FSlateStyleSet("ChooserEditorStyle")
+{
+	const FVector2D Icon16x16(16.0f, 16.0f);
+	const FVector2D Icon24x24(24.0f, 24.0f);
+
+	SetContentRoot(FPaths::EnginePluginsDir() / TEXT("Experimental/Chooser/Content"));
+	SetCoreContentRoot(FPaths::EngineContentDir() / TEXT(""));
+
+	// content browser icon
+	Set("ChooserEditor.ChooserTableIconLarge", new IMAGE_BRUSH_SVG("ChooserIcon_24", Icon24x24));
+
+	// tab icon
+	Set("ChooserEditor.ChooserTableIconSmall", new IMAGE_BRUSH_SVG("ChooserIcon_16", Icon16x16));
+}
+
+void FChooserEditorStyle::Initialize()
+{
+	if (!StyleInstance.IsValid())
+	{
+		StyleInstance = MakeShared<FChooserEditorStyle>();
+		FSlateStyleRegistry::RegisterSlateStyle(*StyleInstance);
+	}
+}
+
+void FChooserEditorStyle::Shutdown()
+{
+	FSlateStyleRegistry::UnRegisterSlateStyle(*StyleInstance);
+	ensure(StyleInstance.IsUnique());
+	StyleInstance.Reset();
+}
+
+const ISlateStyle& FChooserEditorStyle::Get()
+{
+	return *StyleInstance;
+}
+
+}
\ No newline at end of file
diff --git a/Engine/Plugins/Experimental/Chooser/Source/ChooserEditor/Private/ChooserEditorStyle.h b/Engine/Plugins/Experimental/Chooser/Source/ChooserEditor/Private/ChooserEditorStyle.h
new file mode 100644
--- /dev/null
+++ b/Engine/Plugins/Experimental/Chooser/Source/ChooserEditor/Private/ChooserEditorStyle.h
@@ -0,0 +1,24 @@
+// Copyright Epic Games, Inc. All Rights Reserved.
+#pragma once
+
+#include "CoreMinimal.h"
+#include "Styling/SlateStyle.h"
+
+namespace UE::ChooserEditor
+{
+	
+class FChooserEditorStyle : FSlateStyleSet
+{
+public:
+    FChooserEditorStyle();
+	static void Initialize();
+	static void Shutdown();
+
+	static const ISlateStyle& Get();
+
+private:
+
+	static TSharedPtr< FChooserEditorStyle > StyleInstance;
+};
+
+}
diff --git a/Engine/Plugins/Experimental/Chooser/Source/ChooserEditor/Private/ChooserFactory.h b/Engine/Plugins/Experimental/Chooser/Source/ChooserEditor/Private/ChooserFactory.h
--- a/Engine/Plugins/Experimental/Chooser/Source/ChooserEditor/Private/ChooserFactory.h
+++ b/Engine/Plugins/Experimental/Chooser/Source/ChooserEditor/Private/ChooserFactory.h
@@ -19,4 +19,5 @@ class UChooserTableFactory : public UFactory
 	{
 		return FactoryCreateNew(Class, InParent, Name, Flags, Context, Warn, NAME_None);
 	}
+	virtual FName GetNewAssetThumbnailOverride() const override { return "ChooserEditor.ChooserTableIconLarge"; }
 };
\ No newline at end of file
diff --git a/Engine/Plugins/Experimental/Chooser/Source/ChooserEditor/Private/ChooserTableEditor.cpp b/Engine/Plugins/Experimental/Chooser/Source/ChooserEditor/Private/ChooserTableEditor.cpp
--- a/Engine/Plugins/Experimental/Chooser/Source/ChooserEditor/Private/ChooserTableEditor.cpp
+++ b/Engine/Plugins/Experimental/Chooser/Source/ChooserEditor/Private/ChooserTableEditor.cpp
@@ -18,6 +18,7 @@
 #include "StructViewerModule.h"
 #include "SourceCodeNavigation.h"
 #include "Chooser.h"
+#include "ChooserEditorStyle.h"
 #include "ChooserFindProperties.h"
 #include "DetailCategoryBuilder.h"
 #include "DetailLayoutBuilder.h"
@@ -39,6 +40,7 @@
 #include "ChooserTableEditorCommands.h"
 #include "Animation/AnimBlueprintGeneratedClass.h"
 #include "Animation/AnimInstance.h"
+#include "Widgets/Input/STextEntryPopup.h"
 #include "Widgets/Layout/SScrollBox.h"
 
 #define LOCTEXT_NAMESPACE "ChooserEditor"
@@ -50,6 +52,21 @@ const FName FChooserTableEditor::ToolkitFName( TEXT( "ChooserTableEditor" ) );
 const FName FChooserTableEditor::PropertiesTabId( TEXT( "ChooserEditor_Properties" ) );
 const FName FChooserTableEditor::FindReplaceTabId( TEXT( "ChooserEditor_FindReplace" ) );
 const FName FChooserTableEditor::TableTabId( TEXT( "ChooserEditor_Table" ) );
+
+void FChooserTableEditor::PushChooserTableToEdit(UChooserTable* Chooser)
+{
+	BreadcrumbTrail->PushCrumb(FText::FromString(Chooser->GetName()), Chooser);
+	RefreshAll();
+}
+	
+void FChooserTableEditor::PopChooserTableToEdit()
+{
+	if (BreadcrumbTrail->HasCrumbs())
+	{
+		BreadcrumbTrail->PopCrumb();
+		RefreshAll();
+	}
+}
 	
 void FChooserTableEditor::RegisterTabSpawners(const TSharedRef<class FTabManager>& InTabManager)
 {
@@ -65,7 +82,7 @@ void FChooserTableEditor::RegisterTabSpawners(const TSharedRef<class FTabManager
 	InTabManager->RegisterTabSpawner( TableTabId, FOnSpawnTab::CreateSP(this, &FChooserTableEditor::SpawnTableTab) )
 		.SetDisplayName( LOCTEXT("TableTab", "Chooser Table") )
 		.SetGroup(WorkspaceMenuCategory.ToSharedRef())
-		.SetIcon(FSlateIcon("EditorStyle", "LevelEditor.Tabs.Details"));
+		.SetIcon(FSlateIcon("ChooserEditorStyle", "ChooserEditor.ChooserTableIconSmall"));
 
 	InTabManager->RegisterTabSpawner( FindReplaceTabId, FOnSpawnTab::CreateSP(this, &FChooserTableEditor::SpawnFindReplaceTab) )
 		.SetDisplayName( LOCTEXT("FindReplaceTab", "Find/Replace") )
@@ -99,7 +116,15 @@ FChooserTableEditor::~FChooserTableEditor()
 
 	DetailsView.Reset();
 }
+
+
+FName FChooserTableEditor::EditorName = "ChooserTableEditor";
 	
+FName FChooserTableEditor::GetEditorName() const
+{
+	return EditorName;
+}
+
 void FChooserTableEditor::MakeDebugTargetMenu(UToolMenu* InToolMenu) 
 {
 	static FName SectionName = "Select Debug Target";
@@ -114,7 +139,7 @@ void FChooserTableEditor::MakeDebugTargetMenu(UToolMenu* InToolMenu)
 				FUIAction(
 					FExecuteAction::CreateLambda([this]()
 					{
-						UChooserTable* Chooser = GetChooser();
+						UChooserTable* Chooser = GetRootChooser();
 						Chooser->ResetDebugTarget();
 						if (Chooser->bEnableDebugTesting)
 						{
@@ -137,7 +162,7 @@ void FChooserTableEditor::MakeDebugTargetMenu(UToolMenu* InToolMenu)
 				FUIAction(
 					FExecuteAction::CreateLambda([this]()
 					{
-						UChooserTable* Chooser = GetChooser();
+						UChooserTable* Chooser = GetRootChooser();
 						Chooser->ResetDebugTarget();
 						if (!Chooser->bEnableDebugTesting)
 						{
@@ -166,7 +191,7 @@ void FChooserTableEditor::MakeDebugTargetMenu(UToolMenu* InToolMenu)
 								{
 									if(ObjectPtr.IsValid())
 									{
-										UChooserTable* Chooser = GetChooser();
+										UChooserTable* Chooser = GetRootChooser();
 										Chooser->SetDebugTarget(ObjectPtr);
 										Chooser->bDebugTestValuesValid = false;
 										if (!Chooser->bEnableDebugTesting)
@@ -222,7 +247,7 @@ void FChooserTableEditor::RegisterToolbar()
 					InSection.AddEntry(FToolMenuEntry::InitComboButton( "SelectDebugTarget",
 						FToolUIActionChoice(),
 					  FNewToolMenuDelegate::CreateSP(ChooserEditor.Get(), &FChooserTableEditor::MakeDebugTargetMenu),
-						TAttribute<FText>::CreateLambda([Chooser = ChooserEditor->GetChooser()]
+						TAttribute<FText>::CreateLambda([Chooser = ChooserEditor->GetRootChooser() ]
 						{
 							if (Chooser->HasDebugTarget())
 							{
@@ -237,7 +262,6 @@ void FChooserTableEditor::RegisterToolbar()
 				}
 			}
 		}));
-
 	}
 
 }
@@ -263,6 +287,25 @@ void FChooserTableEditor::BindCommands()
 void FChooserTableEditor::InitEditor( const EToolkitMode::Type Mode, const TSharedPtr< class IToolkitHost >& InitToolkitHost, const TArray<UObject*>& ObjectsToEdit, FGetDetailsViewObjects GetDetailsViewObjects )
 {
 	EditingObjects = ObjectsToEdit;
+
+	BreadcrumbTrail = SNew(SBreadcrumbTrail<UChooserTable*>)
+		.ButtonStyle(FAppStyle::Get(), "GraphBreadcrumbButton")
+		.TextStyle(FAppStyle::Get(), "GraphBreadcrumbButtonText")
+		.ButtonContentPadding( FMargin(4.f, 2.f) )
+		.DelimiterImage( FAppStyle::GetBrush("BreadcrumbTrail.Delimiter") )
+		.OnCrumbPushed_Lambda([this](UChooserTable* Table)
+		{
+			RefreshAll();
+		})
+		.OnCrumbPopped_Lambda([this](UChooserTable* Table)
+		{
+			RefreshAll();
+		});
+		
+	UChooserTable* RootTable = GetRootChooser();
+	BreadcrumbTrail->PushCrumb(FText::FromString(RootTable->GetName()), RootTable);
+	
+	
 	FCoreUObjectDelegates::OnObjectsReplaced.AddSP(this, &FChooserTableEditor::OnObjectsReplaced);
 
 	FPropertyEditorModule& PropertyEditorModule = FModuleManager::GetModuleChecked<FPropertyEditorModule>( "PropertyEditor" );
@@ -319,20 +362,25 @@ FText FChooserTableEditor::GetBaseToolkitName() const
 	return LOCTEXT("AppLabel", "Chooser Table Editor");
 }
 
+void FChooserTableEditor::RefreshAll()
+{
+	if (HeaderRow)
+	{
+		UpdateTableColumns();
+		UpdateTableRows();
+		ClearSelectedColumn();
+		SelectRootProperties();
+	}
+}
+
 void FChooserTableEditor::PostUndo(bool bSuccess)
 {
-	UpdateTableColumns();
-	UpdateTableRows();
-	ClearSelectedColumn();
-	SelectRootProperties();
+	RefreshAll();
 }
 
 void FChooserTableEditor::PostRedo(bool bSuccess)
 {
-	UpdateTableColumns();
-	UpdateTableRows();
-	ClearSelectedColumn();
-	SelectRootProperties();
+	RefreshAll();
 }
 
 
@@ -559,6 +607,20 @@ public:
 			]
 		];
 	}
+	
+	virtual FReply OnMouseButtonDoubleClick(const FGeometry& InMyGeometry, const FPointerEvent& InMouseEvent) override
+	{
+		UChooserTable* Chooser = ChooserEditor->GetChooser();
+
+		if (const FNestedChooser* NestedChooserResult = Chooser->ResultsStructs[RowIndex].GetPtr<FNestedChooser>())
+		{
+			if (NestedChooserResult->Chooser)
+			{
+				ChooserEditor->PushChooserTableToEdit(NestedChooserResult->Chooser);
+			}
+		}
+		return FReply::Handled();
+	}
 
 	FReply OnMouseButtonDown(const FGeometry& MyGeometry, const FPointerEvent& MouseEvent) override
 	{
@@ -591,6 +653,7 @@ public:
 	SLATE_ARGUMENT(FChooserTableEditor*, Editor)
 	SLATE_END_ARGS()
 
+	enum { SpecialIndex_AddRow = -1, SpecialIndex_Fallback = -2 };
 
 	void Construct(const FArguments& Args, const TSharedRef<STableViewBase>& OwnerTableView)
 	{
@@ -610,9 +673,9 @@ public:
 		static FName Result = "Result";
 		static FName Handles = "Handles";
 	
-		if (RowIndex->RowIndex < Chooser->ResultsStructs.Num())
+		if (Chooser->ResultsStructs.IsValidIndex(RowIndex->RowIndex))
 		{
-			if (ColumnName == Handles && RowIndex->RowIndex < Chooser->ResultsStructs.Num())
+			if (ColumnName == Handles)
 			{
 				// row drag handle
 			
@@ -620,13 +683,13 @@ public:
 			}
 			else if (ColumnName == Result) 
 			{
-				TSharedPtr<SWidget> ResultWidget = FObjectChooserWidgetFactories::CreateWidget(false, Chooser, FObjectChooserBase::StaticStruct(), Chooser->ResultsStructs[RowIndex->RowIndex].GetMutableMemory(), Chooser->ResultsStructs[RowIndex->RowIndex].GetScriptStruct(), Chooser->OutputObjectType,
+				TSharedPtr<SWidget> ResultWidget = FObjectChooserWidgetFactories::CreateWidget(false, Chooser->GetContextOwner(), FObjectChooserBase::StaticStruct(), Chooser->ResultsStructs[RowIndex->RowIndex].GetMutableMemory(), Chooser->ResultsStructs[RowIndex->RowIndex].GetScriptStruct(), Chooser->OutputObjectType,
 				FOnStructPicked::CreateLambda([this, RowIndex=RowIndex->RowIndex](const UScriptStruct* ChosenStruct)
 				{
 					const FScopedTransaction Transaction(LOCTEXT("Change Row Result Type", "Change Row Result Type"));
 					Chooser->Modify(true);
 					Chooser->ResultsStructs[RowIndex].InitializeAs(ChosenStruct);
-					FObjectChooserWidgetFactories::CreateWidget(false, Chooser, FObjectChooserBase::StaticStruct(), Chooser->ResultsStructs[RowIndex].GetMutableMemory(), ChosenStruct, Chooser->OutputObjectType, FOnStructPicked(), &CacheBorder);
+					FObjectChooserWidgetFactories::CreateWidget(false, Chooser->GetContextOwner(), FObjectChooserBase::StaticStruct(), Chooser->ResultsStructs[RowIndex].GetMutableMemory(), ChosenStruct, Chooser->OutputObjectType, FOnStructPicked(), &CacheBorder);
 				}),
 				&CacheBorder
 				);
@@ -655,7 +718,7 @@ public:
 					FChooserColumnBase* Column = &Chooser->ColumnsStructs[ColumnIndex].GetMutable<FChooserColumnBase>();
 					const UStruct * ColumnStruct = Chooser->ColumnsStructs[ColumnIndex].GetScriptStruct();
 
-					TSharedPtr<SWidget> ColumnWidget = FObjectChooserWidgetFactories::CreateColumnWidget(Column, ColumnStruct, Chooser, RowIndex->RowIndex);
+					TSharedPtr<SWidget> ColumnWidget = FObjectChooserWidgetFactories::CreateColumnWidget(Column, ColumnStruct, Chooser->GetContextOwner(), RowIndex->RowIndex);
 				
 					if (ColumnWidget.IsValid())
 					{
@@ -687,7 +750,47 @@ public:
 				}
 			}
 		}
-		else if (RowIndex->RowIndex == Chooser->ResultsStructs.Num())
+		else if (RowIndex->RowIndex == SpecialIndex_Fallback)
+		{
+			if (ColumnName == Handles)
+			{
+				return SNullWidget::NullWidget;
+			}
+			else if (ColumnName == Result) 
+			{
+				TSharedPtr<SWidget> ResultWidget = FObjectChooserWidgetFactories::CreateWidget(false, Chooser->GetContextOwner(), FObjectChooserBase::StaticStruct(), Chooser->FallbackResult.GetMutableMemory(), Chooser->FallbackResult.GetScriptStruct(), Chooser->OutputObjectType,
+				FOnStructPicked::CreateLambda([this](const UScriptStruct* ChosenStruct)
+				{
+					const FScopedTransaction Transaction(LOCTEXT("Change Row Result Type", "Change Row Result Type"));
+					Chooser->Modify(true);
+					Chooser->FallbackResult.InitializeAs(ChosenStruct);
+					FObjectChooserWidgetFactories::CreateWidget(false, Chooser->GetContextOwner(), FObjectChooserBase::StaticStruct(), Chooser->FallbackResult.GetMutableMemory(), ChosenStruct, Chooser->OutputObjectType, FOnStructPicked(), &CacheBorder
+						,FChooserWidgetValueChanged(), LOCTEXT("Fallback Result", "Fallback Result: (None)"));
+				}),
+				&CacheBorder
+				,FChooserWidgetValueChanged(), LOCTEXT("Fallback Result", "Fallback Result: (None)")
+				);
+			
+				return ResultWidget.ToSharedRef();
+			}
+			else
+			{
+				const int ColumnIndex = ColumnName.GetNumber() - 1;
+				if (ColumnIndex < Chooser->ColumnsStructs.Num() && ColumnIndex >=0)
+				{
+					FChooserColumnBase* Column = &Chooser->ColumnsStructs[ColumnIndex].GetMutable<FChooserColumnBase>();
+					const UStruct * ColumnStruct = Chooser->ColumnsStructs[ColumnIndex].GetScriptStruct();
+
+					TSharedPtr<SWidget> ColumnWidget = FObjectChooserWidgetFactories::CreateColumnWidget(Column, ColumnStruct, Chooser->GetContextOwner(), -2);
+				
+					if (ColumnWidget.IsValid())
+					{
+						return ColumnWidget.ToSharedRef();
+					}
+				}
+			}
+		}
+		else if (RowIndex->RowIndex == SpecialIndex_AddRow)
 		{
 			// on the row past the end, show an Add button in the result column
 			if (ColumnName == Result)
@@ -759,7 +862,7 @@ private:
 
 TSharedRef<ITableRow> FChooserTableEditor::GenerateTableRow(TSharedPtr<FChooserTableRow> InItem, const TSharedRef<STableViewBase>& OwnerTable)
 {
-	UChooserTable* Chooser = Cast<UChooserTable>(EditingObjects[0]);
+	UChooserTable* Chooser = GetChooser();
 
 	return SNew(SChooserTableRow, OwnerTable)
 		.Entry(InItem).Chooser(Chooser).Editor(this);
@@ -769,14 +872,14 @@ void FChooserTableEditor::SelectRootProperties()
 {
 	if( DetailsView.IsValid() )
 	{
-		// Make sure details window is pointing to our object
-		DetailsView->SetObjects( EditingObjects );
+		// point the details view to the main table
+		DetailsView->SetObject( GetRootChooser() );
 	}
 }
 
 int FChooserTableEditor::MoveRow(int SourceRowIndex, int TargetRowIndex)
 {
-	UChooserTable* Chooser = Cast<UChooserTable>(EditingObjects[0]);
+	UChooserTable* Chooser = GetChooser();
 	TargetRowIndex = FMath::Min(TargetRowIndex,Chooser->ResultsStructs.Num());
 
 	const FScopedTransaction Transaction(LOCTEXT("Move Row", "Move Row"));
@@ -823,7 +926,7 @@ void FChooserTableEditor::ClearSelectedRows()
 
 void FChooserTableEditor::UpdateTableColumns()
 {
-	UChooserTable* Chooser = Cast<UChooserTable>(EditingObjects[0]);
+	UChooserTable* Chooser = GetChooser();
 
 	HeaderRow->ClearColumns();
 
@@ -841,7 +944,7 @@ void FChooserTableEditor::UpdateTableColumns()
 	{
 		FChooserColumnBase& Column = Chooser->ColumnsStructs[ColumnIndex].GetMutable<FChooserColumnBase>();
 
-		TSharedPtr<SWidget> HeaderWidget = FObjectChooserWidgetFactories::CreateColumnWidget(&Column, Chooser->ColumnsStructs[ColumnIndex].GetScriptStruct(), Chooser, -1);
+		TSharedPtr<SWidget> HeaderWidget = FObjectChooserWidgetFactories::CreateColumnWidget(&Column, Chooser->ColumnsStructs[ColumnIndex].GetScriptStruct(), Chooser->GetContextOwner(), -1);
 		if (!HeaderWidget.IsValid())
 		{
 			HeaderWidget = SNullWidget::NullWidget;
@@ -863,7 +966,7 @@ void FChooserTableEditor::UpdateTableColumns()
 					FUIAction(
 						FExecuteAction::CreateLambda([this,Chooser,ColumnIndex, ColumnId, &Column]()
 						{
-							SelectColumn(ColumnId.GetNumber() - 1);
+							SelectColumn(Chooser, ColumnId.GetNumber() - 1);
 						})
 						)
 					);
@@ -918,7 +1021,7 @@ void FChooserTableEditor::UpdateTableColumns()
 						TSharedRef<SWidget> Widget = FModuleManager::LoadModuleChecked<FStructViewerModule>("StructViewer").CreateStructViewer(Options, FOnStructPicked::CreateLambda([this, ColumnIndex](const UScriptStruct* ChosenStruct)
 						{
 							const FScopedTransaction Transaction(LOCTEXT("SetColumnInputType", "Set Column Input Type"));
-							UChooserTable* ChooserTable = Cast<UChooserTable>(EditingObjects[0]);
+							UChooserTable* ChooserTable = GetChooser();
 							ChooserTable->ColumnsStructs[ColumnIndex].GetMutable<FChooserColumnBase>().SetInputType(ChosenStruct);
 							ChooserTable->Modify(true);
 							UpdateTableColumns();
@@ -941,12 +1044,11 @@ void FChooserTableEditor::UpdateTableColumns()
 					// unclear why this color is coming out much darker
 					return (SelectedColumn && SelectedColumn->Column == ColumnId.GetNumber() - 1) ? FSlateColor(FColor(0x00, 0x70, 0xe0, 0xFF)) : FSlateColor(FLinearColor(0.05f,0.05f,0.05f));
 				})
-				.OnMouseButtonDown_Lambda([this, ColumnIndex, ColumnId](	const FGeometry&, const FPointerEvent& PointerEvent)
+				.OnMouseButtonDown_Lambda([this, Chooser, ColumnIndex, ColumnId](	const FGeometry&, const FPointerEvent& PointerEvent)
 				{
-					UChooserTable* Chooser = Cast<UChooserTable>(EditingObjects[0]);
 					TableView->ClearSelection();
 				
-					SelectColumn(ColumnId.GetNumber() - 1);
+					SelectColumn(Chooser, ColumnId.GetNumber() - 1);
 					return FReply::Handled();
 				})
 				[
@@ -976,7 +1078,7 @@ void FChooserTableEditor::UpdateTableColumns()
 void FChooserTableEditor::AddColumn(const UScriptStruct* ColumnType)
 {
 	CreateColumnComboButton->SetIsOpen(false);
-	UChooserTable* Chooser = Cast<UChooserTable>(EditingObjects[0]);
+	UChooserTable* Chooser = GetChooser();
 	const FScopedTransaction Transaction(LOCTEXT("Add Column Transaction", "Add Column"));
 	Chooser->Modify(true);
 
@@ -1021,15 +1123,15 @@ void FChooserTableEditor::AddColumn(const UScriptStruct* ColumnType)
 	UpdateTableColumns();
 	UpdateTableRows();
 
-	SelectColumn(InsertIndex);
+	SelectColumn(Chooser, InsertIndex);
 }
 
 TSharedRef<SDockTab> FChooserTableEditor::SpawnTableTab( const FSpawnTabArgs& Args )
 {
 	check( Args.GetTabId() == TableTabId );
 
-	UChooserTable* Chooser = Cast<UChooserTable>(EditingObjects[0]);
-	
+	UChooserTable* Chooser = GetChooser();
+
 	// + button to create new columns
 	
 	CreateColumnComboButton = SNew(SComboButton).OnGetMenuContent_Lambda([this]()
@@ -1069,7 +1171,7 @@ TSharedRef<SDockTab> FChooserTableEditor::SpawnTableTab( const FSpawnTabArgs& Ar
 			TSharedRef<SWidget> Widget = FModuleManager::LoadModuleChecked<FStructViewerModule>("StructViewer").CreateStructViewer(Options, FOnStructPicked::CreateLambda([this](const UScriptStruct* ChosenStruct)
 			{
 				CreateRowComboButton->SetIsOpen(false);
-				UChooserTable* Chooser = Cast<UChooserTable>(EditingObjects[0]);
+				UChooserTable* Chooser = GetChooser();
 				const FScopedTransaction Transaction(LOCTEXT("Add Row Transaction", "Add Row"));
 				Chooser->Modify(true);
 
@@ -1094,7 +1196,7 @@ TSharedRef<SDockTab> FChooserTableEditor::SpawnTableTab( const FSpawnTabArgs& Ar
 					if (Event.GetKey() == EKeys::Delete)
 					{
 						const FScopedTransaction Transaction(LOCTEXT("Delete Row Transaction", "Delete Row"));
-						UChooserTable* Chooser = Cast<UChooserTable>(EditingObjects[0]);
+						UChooserTable* Chooser = GetChooser();
 						Chooser->Modify(true);
 						// delete selected rows.
 						TArray<uint32> RowsToDelete;
@@ -1138,7 +1240,14 @@ TSharedRef<SDockTab> FChooserTableEditor::SpawnTableTab( const FSpawnTabArgs& Ar
 							SelectedRow->ClearFlags(RF_Standalone);
 						}
 						SelectedRows.SetNum(0);
-						UChooserTable* Chooser = Cast<UChooserTable>(EditingObjects[0]);
+						UChooserTable* Chooser = GetChooser();
+						
+						if (!Chooser->ResultsStructs.IsValidIndex(SelectedItem->RowIndex))
+						{
+							SelectRootProperties();
+							return;
+						}
+						
 						// Get the list of objects to edit the details of
 						TObjectPtr<UChooserRowDetails> Selection = NewObject<UChooserRowDetails>();
 						Selection->Chooser = Chooser;
@@ -1161,36 +1270,92 @@ TSharedRef<SDockTab> FChooserTableEditor::SpawnTableTab( const FSpawnTabArgs& Ar
 				})
     			.OnGenerateRow_Raw(this, &FChooserTableEditor::GenerateTableRow)
 				.HeaderRow(HeaderRow);
+
+
+	TSharedRef<SComboButton> EditChooserTableButton = SNew(SComboButton)
+		.ButtonStyle(FAppStyle::Get(), "GraphBreadcrumbButton");
 	
+	EditChooserTableButton->SetOnGetMenuContent(
+    		FOnGetContent::CreateLambda(
+    			[ EditChooserTableButton, this]()
+                			{
+    							FMenuBuilder MenuBuilder(true, nullptr);
+                            
+								UObject* RootChooser = GetRootChooser();
+								TArray<UObject*> ObjectsInPackage;
+								GetObjectsWithOuter(RootChooser->GetPackage(), ObjectsInPackage);
+
+								for (UObject* Object : ObjectsInPackage)
+								{
+									if (UChooserTable* Chooser = Cast<UChooserTable>(Object))
+									{
+										MenuBuilder.AddMenuEntry( FText::FromString(Chooser->GetName()), LOCTEXT("AddExistingObjectTooltip", "Add a reference to this existing Chooser Table."), FSlateIcon(),
+											FUIAction(FExecuteAction::CreateLambda([this, EditChooserTableButton, Chooser, RootChooser]()
+											{
+												while(GetChooser() != RootChooser)
+												{
+													PopChooserTableToEdit();
+												}
+												if (Chooser != RootChooser)
+												{
+													PushChooserTableToEdit(Chooser);
+												}
+												EditChooserTableButton->SetIsOpen(false);
+											})));
+									}
+                            	}
+    
+    							return MenuBuilder.MakeWidget();
+    
+                			})
+    		);
+
 	return SNew(SDockTab)
 		.Label( LOCTEXT("ChooserTableTitle", "Chooser Table") )
 		.TabColorScale( GetTabColorScale() )
 		.OnCanCloseTab_Lambda([]() { return false; })
 		[
-			
-			SNew(SScrollBox).Orientation(Orient_Horizontal)
-			+ SScrollBox::Slot()
+			SNew(SVerticalBox)
+			+ SVerticalBox::Slot().AutoHeight().Padding(3)
 			[
-				TableView.ToSharedRef()
+				SNew(SHorizontalBox)
+				+ SHorizontalBox::Slot().AutoWidth()
+				[
+					EditChooserTableButton
+				]
+				+ SHorizontalBox::Slot().FillWidth(1)
+				[
+					BreadcrumbTrail.ToSharedRef()
+				]
+			]
+			+ SVerticalBox::Slot().FillHeight(1)
+			[
+				SNew(SScrollBox).Orientation(Orient_Horizontal)
+				+ SScrollBox::Slot()
+				[
+					TableView.ToSharedRef()
+				]
 			]
 		];
 }
 
 void FChooserTableEditor::UpdateTableRows()
 {
-	UChooserTable* Chooser = Cast<UChooserTable>(EditingObjects[0]);
-	int32 OldNum = TableRows.Num();
+	UChooserTable* Chooser = GetChooser();
 	int32 NewNum = Chooser->ResultsStructs.Num();
 
 	// Sync the TableRows array which drives the ui table to match the number of results.
-	// Add 1 at the end, for the "Add Row" control
-	TableRows.SetNum(NewNum + 1);
-	
-	for(int32 i = OldNum; i < NewNum+1; i++)
+	TableRows.SetNum(0, false);
+	for(int i =0; i < NewNum; i++)
 	{
-		TableRows[i] = MakeShared<FChooserTableRow>(i);
+		TableRows.Add(MakeShared<FChooserTableRow>(i));
 	}
 
+	// Add one at the end, for the Fallback result
+	TableRows.Add(MakeShared<FChooserTableRow>(SChooserTableRow::SpecialIndex_Fallback));
+	// Add one at the end, for the "Add Row" control
+	TableRows.Add(MakeShared<FChooserTableRow>(SChooserTableRow::SpecialIndex_AddRow));
+
 	// Make sure each column has the same number of row datas as there are results
 	for(FInstancedStruct& ColumnData : Chooser->ColumnsStructs)
 	{
@@ -1223,7 +1388,7 @@ void FChooserTableEditor::OnObjectsReplaced(const TMap<UObject*, UObject*>& Repl
 
 	if (bChangedAny)
 	{
-		DetailsView->SetObjects(EditingObjects);
+		SelectRootProperties();
 	}
 }
 
@@ -1250,21 +1415,25 @@ TSharedRef<FChooserTableEditor> FChooserTableEditor::CreateEditor( const EToolki
 	return NewEditor;
 }
 	
-void FChooserTableEditor::SelectColumn(int Index)
+void FChooserTableEditor::SelectColumn(UChooserTable* ChooserEditor, int Index)
 {
 	UChooserTable* Chooser = GetChooser();
-	if (Index < Chooser->ColumnsStructs.Num())
-	{
-		if (SelectedColumn == nullptr)
-		{
-			SelectedColumn = NewObject<UChooserColumnDetails>();
-			SelectedColumn->SetFlags(RF_Standalone);
-			SelectedColumn->Chooser = Chooser;
-		}
-
-		SelectedColumn->Column = Index;
-		DetailsView->SetObject(SelectedColumn, true);
-	}
+   	if (Index < Chooser->ColumnsStructs.Num())
+   	{
+   		if (SelectedColumn == nullptr)
+   		{
+   			SelectedColumn = NewObject<UChooserColumnDetails>();
+   			SelectedColumn->SetFlags(RF_Standalone);
+   		}
+   
+   		SelectedColumn->Chooser = Chooser;
+   		SelectedColumn->Column = Index;
+   		DetailsView->SetObject(SelectedColumn, true);
+   	}
+   	else
+   	{
+   		SelectRootProperties();
+   	}
 }
 	
 void FChooserTableEditor::ClearSelectedColumn()
@@ -1364,6 +1533,157 @@ TSharedRef<SWidget> CreateEvaluateChooserWidget(bool bReadOnly, UObject* Transac
 			EvaluateChooser->Chooser = Cast<UChooserTable>(AssetData.GetAsset());
 		});
 }
+
+TSharedRef<SWidget> CreateNestedChooserWidget(bool bReadOnly, UObject* TransactionObject, void* Value, UClass* ResultBaseClass, FChooserWidgetValueChanged ValueChanged)
+{
+	FNestedChooser* NestedChooser = static_cast<FNestedChooser*>(Value);
+	
+	TSharedRef<SComboButton> Button = SNew(SComboButton)
+			.ContentPadding(0)
+			.ButtonContent()
+			[
+				SNew(SHorizontalBox)
+				+ SHorizontalBox::Slot().AutoWidth()
+				[
+					SNew(SImage).Image(FChooserEditorStyle::Get().GetBrush("ChooserEditor.ChooserTableIconSmall"))
+				]
+				+ SHorizontalBox::Slot().FillWidth(1.0).Padding(2)
+				[
+					SNew(STextBlock)
+					.Font(IDetailLayoutBuilder::GetDetailFont())
+					.Text_Lambda([NestedChooser]()
+					{
+						if (NestedChooser->Chooser)
+						{
+							return FText::FromString(NestedChooser->Chooser->GetName());
+						}
+						else
+						{
+							return LOCTEXT("None", "None");
+						}
+					})
+				]
+			];
+			
+	Button->SetOnGetMenuContent(
+		FOnGetContent::CreateLambda(
+			[ Button, NestedChooser, TransactionObject]()
+            			{
+							FMenuBuilder MenuBuilder(true, nullptr);
+                        
+                        	MenuBuilder.BeginSection(NAME_None, LOCTEXT("Nested Chooser", "Nested Chooser"));
+                        	{
+							MenuBuilder.AddMenuEntry(
+								LOCTEXT("None", "None"), 
+								LOCTEXT("NewNestedChooser_Tooltip", "Create a new Chooser Table embedded in this asset, and reference it"),
+								FSlateIcon(),
+								FUIAction(FExecuteAction::CreateLambda([Button, NestedChooser, TransactionObject]()
+								{
+									const FScopedTransaction Transaction(LOCTEXT("Clear Nested Chooser", "Clear Nested Chooser"));
+									TransactionObject->Modify(true);
+									Button->SetIsOpen(false);
+									NestedChooser->Chooser = nullptr;
+								})));
+								
+                        		MenuBuilder.AddMenuEntry(
+                        			LOCTEXT("New Nested Chooser", "New Nested Chooser"), 
+                        			LOCTEXT("NewNestedChooser_Tooltip", "Create a new Chooser Table embedded in this asset, and reference it"),
+                        			FSlateIcon(),
+                        			FUIAction(FExecuteAction::CreateLambda([Button, NestedChooser, TransactionObject]()
+                        			{
+                        				Button->SetIsOpen(false);
+                        				
+										TSharedRef<STextEntryPopup> TextEntry =
+											SNew(STextEntryPopup)
+											.Label(LOCTEXT("NewNestedChooserLabel", "New Chooser Name"))
+											.OnTextCommitted_Lambda([NestedChooser, TransactionObject](FText InText, ETextCommit::Type InCommitType)
+											{
+												const FScopedTransaction Transaction(LOCTEXT("Assign New Nested Chooser", "Assign New Nested Chooser"));
+												TransactionObject->Modify(true);
+											
+												FSlateApplication::Get().DismissAllMenus();
+												NestedChooser->Chooser = NewObject<UChooserTable>(TransactionObject->GetPackage(), UChooserTable::StaticClass(), FName(InText.ToString()));
+												NestedChooser->Chooser->ParentTable = Cast<UChooserTable>(TransactionObject);
+											});
+							
+										FSlateApplication& SlateApp = FSlateApplication::Get();
+										SlateApp.PushMenu(
+											SlateApp.GetInteractiveTopLevelWindows()[0],
+											FWidgetPath(),
+											TextEntry,
+											SlateApp.GetCursorPos(),
+											FPopupTransitionEffect::TypeInPopup
+											);
+                        			})));
+								
+								MenuBuilder.AddSubMenu(LOCTEXT("Select Existing", "Select Existing"),
+									LOCTEXT("SelectExisting_Tooltip", "Select Existing Embedded UChooserTable from this asset"),
+									FNewMenuDelegate::CreateLambda([Button, NestedChooser, TransactionObject](FMenuBuilder& SubMenuBuilder)
+									{
+										SubMenuBuilder.BeginSection("Existing", LOCTEXT("Existing Choosers", "Existing Choosers"));
+
+										TArray<UObject*> ObjectsInPackage;
+										GetObjectsWithOuter(TransactionObject->GetPackage(), ObjectsInPackage);
+
+										UObject* RootChooser = TransactionObject->GetPackage()->FindAssetInPackage();
+
+										for (UObject* Object : ObjectsInPackage)
+										{
+											 // exclude the main chooser 
+											if (Object != RootChooser)
+											{
+												if (UChooserTable* Chooser = Cast<UChooserTable>(Object))
+												{
+													SubMenuBuilder.AddMenuEntry( FText::FromString(Chooser->GetName()), LOCTEXT("AddExistingObjectTooltip", "Add a reference to this existing Chooser Table."), FSlateIcon(),
+														FUIAction(FExecuteAction::CreateLambda([Chooser, Button, NestedChooser, TransactionObject]()
+														{
+															const FScopedTransaction Transaction(LOCTEXT("Set Nested Chooser", "Set Nested Chooser"));
+															TransactionObject->Modify(true);
+															Button->SetIsOpen(false);
+															NestedChooser->Chooser = Chooser;
+														})));
+												}
+											}
+										}
+										
+
+									}));
+                                     
+									
+									
+                        	}
+
+							return MenuBuilder.MakeWidget();
+
+            			})
+		);
+		
+	return SNew(SHorizontalBox)
+		+ SHorizontalBox::Slot().FillWidth(1.0)
+		[
+			Button
+		]
+		+ SHorizontalBox::Slot().AutoWidth()
+		[
+			SNew(SButton)
+			.Text(LOCTEXT("Edit","Edit"))
+			.OnClicked_Lambda([NestedChooser, TransactionObject]()
+			{
+				if (NestedChooser->Chooser)
+				{
+					UObject* RootChooser = TransactionObject->GetPackage()->FindAssetInPackage();
+					IAssetEditorInstance* Editor = GEditor->GetEditorSubsystem<UAssetEditorSubsystem>()->FindEditorForAsset(RootChooser, true);
+					if (Editor && Editor->GetEditorName() == FChooserTableEditor::EditorName)
+					{
+						// currently nested tables are only editable in the Chooser Table Editor
+						FChooserTableEditor* ChooserEditor = static_cast<FChooserTableEditor*>(Editor);
+						ChooserEditor->PushChooserTableToEdit(NestedChooser->Chooser);
+					}
+				}
+				return FReply::Handled();
+			})
+		];
+}
 	
 class FChooserDetails : public IDetailCustomization
 {
@@ -1507,6 +1827,7 @@ void FChooserTableEditor::RegisterWidgets()
 	FObjectChooserWidgetFactories::RegisterWidgetCreator(FAssetChooser::StaticStruct(), CreateAssetWidget);
 	FObjectChooserWidgetFactories::RegisterWidgetCreator(FClassChooser::StaticStruct(), CreateClassWidget);
 	FObjectChooserWidgetFactories::RegisterWidgetCreator(FEvaluateChooser::StaticStruct(), CreateEvaluateChooserWidget);
+	FObjectChooserWidgetFactories::RegisterWidgetCreator(FNestedChooser::StaticStruct(), CreateNestedChooserWidget);
 
 	FPropertyEditorModule& PropertyModule = FModuleManager::LoadModuleChecked<FPropertyEditorModule>("PropertyEditor");
 	
diff --git a/Engine/Plugins/Experimental/Chooser/Source/ChooserEditor/Private/ChooserTableEditor.h b/Engine/Plugins/Experimental/Chooser/Source/ChooserEditor/Private/ChooserTableEditor.h
--- a/Engine/Plugins/Experimental/Chooser/Source/ChooserEditor/Private/ChooserTableEditor.h
+++ b/Engine/Plugins/Experimental/Chooser/Source/ChooserEditor/Private/ChooserTableEditor.h
@@ -3,14 +3,15 @@
 #pragma once
 
 #include "CoreMinimal.h"
-#include "Toolkits/IToolkitHost.h"
-#include "Toolkits/AssetEditorToolkit.h"
-#include "PropertyEditorDelegates.h"
-#include "Widgets/Views/SListView.h"
-#include "Widgets/Views/STableViewBase.h"
 #include "Chooser.h"
 #include "EditorUndoClient.h"
+#include "PropertyEditorDelegates.h"
 #include "Misc/NotifyHook.h"
+#include "Toolkits/AssetEditorToolkit.h"
+#include "Toolkits/IToolkitHost.h"
+#include "Widgets/Navigation/SBreadcrumbTrail.h"
+#include "Widgets/Views/SListView.h"
+#include "Widgets/Views/STableViewBase.h"
 #include "ChooserTableEditor.generated.h"
 
 class SComboButton;
@@ -65,6 +66,8 @@ namespace UE::ChooserEditor
 		/** Destructor */
 		virtual ~FChooserTableEditor();
 
+		virtual FName GetEditorName() const override;
+
 		/** IToolkit interface */
 		virtual FName GetToolkitFName() const override;
 		virtual FText GetBaseToolkitName() const override;
@@ -75,7 +78,7 @@ namespace UE::ChooserEditor
 		virtual bool IsPrimaryEditor() const override { return true; }
 		virtual bool IsSimpleAssetEditor() const override { return false; }
 		virtual void InitToolMenuContext(FToolMenuContext& MenuContext) override;
-		
+
 		/** FEditorUndoClient Interface */
 		virtual void PostUndo(bool bSuccess) override;
 		virtual void PostRedo(bool bSuccess) override;
@@ -84,8 +87,13 @@ namespace UE::ChooserEditor
 		virtual void NotifyPreChange( FProperty* PropertyAboutToChange ) override;
 		virtual void NotifyPostChange( const FPropertyChangedEvent& PropertyChangedEvent, FProperty* PropertyThatChanged) override;
 
-		UChooserTable* GetChooser() { return Cast<UChooserTable>(EditingObjects[0]); }
-		const UChooserTable* GetChooser() const { return Cast<UChooserTable>(EditingObjects[0]); }
+		UChooserTable* GetRootChooser() { return Cast<UChooserTable>(EditingObjects[0]); }
+		UChooserTable* GetChooser() { return BreadcrumbTrail->PeekCrumb(); }
+		const UChooserTable* GetChooser() const { return BreadcrumbTrail->PeekCrumb(); }
+
+		void PushChooserTableToEdit(UChooserTable* Chooser);
+		void PopChooserTableToEdit();
+		void RefreshAll();
 	
 		/** Used to show or hide certain properties */
 		void SetPropertyVisibilityDelegate(FIsPropertyVisible InVisibilityDelegate);
@@ -99,7 +107,7 @@ namespace UE::ChooserEditor
 		};
 
 		void UpdateTableRows();
-		void SelectColumn(int Index);
+		void SelectColumn(UChooserTable* Chooser, int Index);
 		void ClearSelectedColumn();
 		void DeleteColumn(int Index);
 		void AddColumn(const UScriptStruct* ColumnType);
@@ -142,14 +150,17 @@ namespace UE::ChooserEditor
 		UChooserColumnDetails* SelectedColumn = nullptr;
 		TArray<TObjectPtr<UChooserRowDetails>> SelectedRows;
 
+		TSharedPtr<SBreadcrumbTrail<UChooserTable*>> BreadcrumbTrail;
+		
 		void UpdateTableColumns();
 		TArray<TSharedPtr<FChooserTableRow>> TableRows;
 	
 		TSharedPtr<SComboButton> CreateColumnComboButton;
 		TSharedPtr<SComboButton> CreateRowComboButton;
-		
+
 		TSharedPtr<SHeaderRow> HeaderRow;
 		TSharedPtr<SListView<TSharedPtr<FChooserTableRow>>> TableView;
+
 	public:
 
 		TSharedPtr<SComboButton>& GetCreateRowComboButton() { return CreateRowComboButton; };
@@ -162,6 +173,8 @@ namespace UE::ChooserEditor
 		static TSharedRef<FChooserTableEditor> CreateEditor( const EToolkitMode::Type Mode, const TSharedPtr< IToolkitHost >& InitToolkitHost, const TArray<UObject*>& ObjectsToEdit, FGetDetailsViewObjects GetDetailsViewObjects = FGetDetailsViewObjects() );
 
 		static void RegisterWidgets();
+		
+		static FName EditorName;
 	};
 }
 
diff --git a/Engine/Plugins/Experimental/Chooser/Source/ChooserEditor/Private/EnumColumnEditor.cpp b/Engine/Plugins/Experimental/Chooser/Source/ChooserEditor/Private/EnumColumnEditor.cpp
--- a/Engine/Plugins/Experimental/Chooser/Source/ChooserEditor/Private/EnumColumnEditor.cpp
+++ b/Engine/Plugins/Experimental/Chooser/Source/ChooserEditor/Private/EnumColumnEditor.cpp
@@ -16,19 +16,23 @@
 
 namespace UE::ChooserEditor
 {
+
 	
 // Wrapper widget for EnumComboBox which will reconstruct the combo box when the Enum has changed
 template <typename ColumnType>
 class SEnumCell : public SCompoundWidget
 {
 public:
+
+	DECLARE_DELEGATE_OneParam(FOnValueSet, int);
+	
 	SLATE_BEGIN_ARGS(SEnumCell)
-		: _RowIndex(-1)
 	{}
 
 	SLATE_ARGUMENT(UObject*, TransactionObject)
 	SLATE_ARGUMENT(ColumnType*, EnumColumn)
-	SLATE_ATTRIBUTE(int, RowIndex);
+	SLATE_ATTRIBUTE(int32, EnumValue);
+	SLATE_EVENT(FOnValueSet, OnValueSet)
             
 	SLATE_END_ARGS()
 
@@ -42,31 +46,12 @@ public:
 				{
 					return SNew(SEnumComboBox, Enum)
 						.IsEnabled_Lambda([this](){ return IsEnabled(); } )
-						.CurrentValue_Lambda([this]()
+						.CurrentValue(EnumValue)
+						.OnEnumSelectionChanged_Lambda([this](int32 InEnumValue, ESelectInfo::Type)
 						{
-							int Row = RowIndex.Get();
-							if (EnumColumn->RowValues.IsValidIndex(Row))
-							{
-								return EnumColumn->RowValues.IsValidIndex(Row) ? static_cast<int32>(EnumColumn->RowValues[Row].Value) : 0;
-							}
-							else
-							{
-								return static_cast<int32>(EnumColumn->TestValue);
-							}
-						})
-						.OnEnumSelectionChanged_Lambda([this](int32 EnumValue, ESelectInfo::Type)
-						{
-							int Row = RowIndex.Get();
-							if (EnumColumn->RowValues.IsValidIndex(Row))
-							{
-								const FScopedTransaction Transaction(LOCTEXT("Edit RHS", "Edit Enum Value"));
-								TransactionObject->Modify(true);
-								EnumColumn->RowValues[Row].Value = static_cast<uint8>(EnumValue);
-							}
-							else
-							{
-								EnumColumn->TestValue = EnumValue;
-							}
+							const FScopedTransaction Transaction(LOCTEXT("Edit RHS", "Edit Enum Value"));
+							TransactionObject->Modify(true);
+							OnValueSet.ExecuteIfBound(InEnumValue);
 						});
 				}
 			}
@@ -100,9 +85,10 @@ public:
 		SetEnabled(InArgs._IsEnabled);
 		
 		SetCanTick(true);
-		RowIndex = InArgs._RowIndex;
 		EnumColumn = InArgs._EnumColumn;
 		TransactionObject = InArgs._TransactionObject;
+		EnumValue = InArgs._EnumValue;
+		OnValueSet = InArgs._OnValueSet;
 
 		if (EnumColumn)
 		{
@@ -137,13 +123,20 @@ private:
 	TSharedPtr<SBorder> EnumComboBorder;
 	TAttribute<int> RowIndex;
 	FDelegateHandle EnumChangedHandle;
+	
+	FOnValueSet OnValueSet;
+	TAttribute<int32> EnumValue;
 };
 
 TSharedRef<SWidget> CreateEnumColumnWidget(UChooserTable* Chooser, FChooserColumnBase* Column, int Row)
 {
 	FEnumColumn* EnumColumn = static_cast<FEnumColumn*>(Column);
 	
-	if (Row < 0)
+	if (Row == ColumnWidget_SpecialIndex_Fallback)
+	{
+		return SNullWidget::NullWidget;
+	}
+	if (Row == ColumnWidget_SpecialIndex_Header)
 	{
 		// create column header widget
 		TSharedPtr<SWidget> InputValueWidget = nullptr;
@@ -178,7 +171,9 @@ TSharedRef<SWidget> CreateEnumColumnWidget(UChooserTable* Chooser, FChooserColum
 			]
 			+ SVerticalBox::Slot()
 			[
-				SNew(SEnumCell<FEnumColumn>).TransactionObject(Chooser).EnumColumn(EnumColumn).RowIndex(Row)
+				SNew(SEnumCell<FEnumColumn>).TransactionObject(Chooser).EnumColumn(EnumColumn)
+					.OnValueSet_Lambda([EnumColumn](int Value) { EnumColumn->TestValue = Value; })
+					.EnumValue_Lambda([EnumColumn]() { return EnumColumn->TestValue; })
 					.IsEnabled_Lambda([Chooser] { return !Chooser->HasDebugTarget(); })
 			];
 		}
@@ -227,7 +222,18 @@ TSharedRef<SWidget> CreateEnumColumnWidget(UChooserTable* Chooser, FChooserColum
 			]
 			+ SHorizontalBox::Slot().FillWidth(1)
 			[
-				SNew(SEnumCell<FEnumColumn>).TransactionObject(Chooser).EnumColumn(EnumColumn).RowIndex(Row)
+				SNew(SEnumCell<FEnumColumn>).TransactionObject(Chooser).EnumColumn(EnumColumn)
+					.OnValueSet_Lambda([EnumColumn, Row](int Value)
+					{
+						if (EnumColumn->RowValues.IsValidIndex(Row))
+						{
+							EnumColumn->RowValues[Row].Value = static_cast<uint8>(Value);
+						}
+					})
+					.EnumValue_Lambda([EnumColumn, Row]()
+					{
+						return EnumColumn->RowValues.IsValidIndex(Row) ? static_cast<int32>(EnumColumn->RowValues[Row].Value) : 0;
+					})
 				.Visibility_Lambda([EnumColumn,Column, Row]()
 					{
 						return (EnumColumn->RowValues.IsValidIndex(Row) &&
@@ -242,7 +248,7 @@ TSharedRef<SWidget> CreateOutputEnumColumnWidget(UChooserTable* Chooser, FChoose
 {
 	FOutputEnumColumn* EnumColumn = static_cast<FOutputEnumColumn*>(Column);
 	
-	if (Row < 0)
+	if (Row == ColumnWidget_SpecialIndex_Header)
 	{
 		// create column header widget
 		TSharedPtr<SWidget> InputValueWidget = nullptr;
@@ -277,19 +283,34 @@ TSharedRef<SWidget> CreateOutputEnumColumnWidget(UChooserTable* Chooser, FChoose
 			]
 			+ SVerticalBox::Slot()
 			[
-				SNew(SEnumCell<FOutputEnumColumn>).TransactionObject(Chooser).EnumColumn(EnumColumn).RowIndex(Row).IsEnabled(false)
+				SNew(SEnumCell<FOutputEnumColumn>).TransactionObject(Chooser).EnumColumn(EnumColumn).IsEnabled(false)
+				.EnumValue_Lambda([EnumColumn]() { return static_cast<int32>(EnumColumn->TestValue); })
 			];
 		}
 
 		return ColumnHeaderWidget;
 	}
+	else if (Row == ColumnWidget_SpecialIndex_Fallback)
+	{
+		return 	SNew(SEnumCell<FOutputEnumColumn>).TransactionObject(Chooser).EnumColumn(EnumColumn)
+        			.OnValueSet_Lambda([EnumColumn](int Value) { EnumColumn->FallbackValue.Value = Value; })
+        			.EnumValue_Lambda([EnumColumn]() { return EnumColumn->FallbackValue.Value; });
+	}
 
 	// create cell widget
 	
-	return
-
-
-	SNew(SEnumCell<FOutputEnumColumn>).TransactionObject(Chooser).EnumColumn(EnumColumn).RowIndex(Row);
+	return SNew(SEnumCell<FOutputEnumColumn>).TransactionObject(Chooser).EnumColumn(EnumColumn)
+		.OnValueSet_Lambda([EnumColumn, Row](int Value)
+		{
+			if (EnumColumn->RowValues.IsValidIndex(Row))
+			{
+				EnumColumn->RowValues[Row].Value = static_cast<uint8>(Value);
+			}
+		})
+		.EnumValue_Lambda([EnumColumn, Row]()
+		{
+			return EnumColumn->RowValues.IsValidIndex(Row) ? static_cast<int32>(EnumColumn->RowValues[Row].Value) : 0;
+		});
 }
 
 TSharedRef<SWidget> CreateEnumPropertyWidget(bool bReadOnly, UObject* TransactionObject, void* Value, UClass* ResultBaseClass, FChooserWidgetValueChanged ValueChanged)
diff --git a/Engine/Plugins/Experimental/Chooser/Source/ChooserEditor/Private/FloatRangeColumnEditor.cpp b/Engine/Plugins/Experimental/Chooser/Source/ChooserEditor/Private/FloatRangeColumnEditor.cpp
--- a/Engine/Plugins/Experimental/Chooser/Source/ChooserEditor/Private/FloatRangeColumnEditor.cpp
+++ b/Engine/Plugins/Experimental/Chooser/Source/ChooserEditor/Private/FloatRangeColumnEditor.cpp
@@ -30,7 +30,11 @@ TSharedRef<SWidget> CreateFloatRangeColumnWidget(UChooserTable* Chooser, FChoose
 {
 	FFloatRangeColumn* FloatRangeColumn = static_cast<FFloatRangeColumn*>(Column);
 
-	if (Row < 0)
+	if (Row == ColumnWidget_SpecialIndex_Fallback)
+	{
+		return SNullWidget::NullWidget;
+	}
+	else if (Row == ColumnWidget_SpecialIndex_Header)
 	{
 		// create column header widget
 		TSharedPtr<SWidget> InputValueWidget = nullptr;
diff --git a/Engine/Plugins/Experimental/Chooser/Source/ChooserEditor/Private/GameplayTagColumnEditor.cpp b/Engine/Plugins/Experimental/Chooser/Source/ChooserEditor/Private/GameplayTagColumnEditor.cpp
--- a/Engine/Plugins/Experimental/Chooser/Source/ChooserEditor/Private/GameplayTagColumnEditor.cpp
+++ b/Engine/Plugins/Experimental/Chooser/Source/ChooserEditor/Private/GameplayTagColumnEditor.cpp
@@ -18,7 +18,11 @@ TSharedRef<SWidget> CreateGameplayTagColumnWidget(UChooserTable* Chooser, FChoos
 {
 	FGameplayTagColumn* GameplayTagColumn = static_cast<struct FGameplayTagColumn*>(Column);
 	
-	if (Row < 0)
+	if (Row == ColumnWidget_SpecialIndex_Fallback)
+	{
+		return SNullWidget::NullWidget;
+	}
+	else if (Row == ColumnWidget_SpecialIndex_Header)
 	{
 		// create column header widget
 		TSharedPtr<SWidget> InputValueWidget = nullptr;
diff --git a/Engine/Plugins/Experimental/Chooser/Source/ChooserEditor/Private/ObjectChooserWidgetFactories.cpp b/Engine/Plugins/Experimental/Chooser/Source/ChooserEditor/Private/ObjectChooserWidgetFactories.cpp
--- a/Engine/Plugins/Experimental/Chooser/Source/ChooserEditor/Private/ObjectChooserWidgetFactories.cpp
+++ b/Engine/Plugins/Experimental/Chooser/Source/ChooserEditor/Private/ObjectChooserWidgetFactories.cpp
@@ -63,7 +63,7 @@ TSharedPtr<SWidget> FObjectChooserWidgetFactories::CreateWidget(bool bReadOnly,
 	return nullptr;
 }
 
-TSharedPtr<SWidget> FObjectChooserWidgetFactories::CreateWidget(bool bReadOnly, UObject* TransactionObject, const UScriptStruct* BaseType, void* Value, const UStruct* ValueType, UClass* ResultBaseClass, const FOnStructPicked& CreateClassCallback, TSharedPtr<SBorder>* InnerWidget, FChooserWidgetValueChanged ValueChanged)
+TSharedPtr<SWidget> FObjectChooserWidgetFactories::CreateWidget(bool bReadOnly, UObject* TransactionObject, const UScriptStruct* BaseType, void* Value, const UStruct* ValueType, UClass* ResultBaseClass, const FOnStructPicked& CreateClassCallback, TSharedPtr<SBorder>* InnerWidget, FChooserWidgetValueChanged ValueChanged, FText NullValueDisplayText)
 {
 	TSharedPtr<SWidget> LeftWidget = CreateWidget(bReadOnly, TransactionObject, Value, ValueType, ResultBaseClass, ValueChanged);
 	if (bReadOnly)
@@ -74,7 +74,10 @@ TSharedPtr<SWidget> FObjectChooserWidgetFactories::CreateWidget(bool bReadOnly,
 	
 	if (!LeftWidget.IsValid())
 	{
-		LeftWidget = SNew(STextBlock).Text(LOCTEXT("SelectDataType", "Select Data Type..."));
+		LeftWidget = SNew(STextBlock)
+			.Font(FAppStyle::GetFontStyle("PropertyWindow.NormalFont"))
+			.Margin(2)
+			.Text(NullValueDisplayText.IsEmpty() ? LOCTEXT("SelectDataType", "Select Data Type..." ) : NullValueDisplayText);
 	}
 
 	// button for replacing data with a different Data Interface class
@@ -86,6 +89,7 @@ TSharedPtr<SWidget> FObjectChooserWidgetFactories::CreateWidget(bool bReadOnly,
 		FStructViewerInitializationOptions Options;
 		Options.StructFilter = MakeShared<FStructFilter>(BaseType);
 		Options.NameTypeToDisplay = EStructViewerNameTypeToDisplay::DisplayName;
+		Options.bShowNoneOption = true;
 		
 		TSharedRef<SWidget> Widget = FModuleManager::LoadModuleChecked<FStructViewerModule>("StructViewer").CreateStructViewer(Options, FOnStructPicked::CreateLambda(
 			[Button, CreateClassCallback](const UScriptStruct* ChosenStruct)
diff --git a/Engine/Plugins/Experimental/Chooser/Source/ChooserEditor/Private/ObjectColumnEditor.cpp b/Engine/Plugins/Experimental/Chooser/Source/ChooserEditor/Private/ObjectColumnEditor.cpp
--- a/Engine/Plugins/Experimental/Chooser/Source/ChooserEditor/Private/ObjectColumnEditor.cpp
+++ b/Engine/Plugins/Experimental/Chooser/Source/ChooserEditor/Private/ObjectColumnEditor.cpp
@@ -93,7 +93,11 @@ namespace UE::ChooserEditor
 
 		FObjectColumn* ObjectColumn = static_cast<FObjectColumn*>(Column);
 
-		if (Row < 0)
+		if (Row == ColumnWidget_SpecialIndex_Fallback)
+		{
+			return SNullWidget::NullWidget;
+		}
+		else if (Row == ColumnWidget_SpecialIndex_Header)
 		{
 			// create column header widget
 			TSharedPtr<SWidget> InputValueWidget = nullptr;
diff --git a/Engine/Plugins/Experimental/Chooser/Source/ChooserEditor/Private/OutputFloatColumnEditor.cpp b/Engine/Plugins/Experimental/Chooser/Source/ChooserEditor/Private/OutputFloatColumnEditor.cpp
--- a/Engine/Plugins/Experimental/Chooser/Source/ChooserEditor/Private/OutputFloatColumnEditor.cpp
+++ b/Engine/Plugins/Experimental/Chooser/Source/ChooserEditor/Private/OutputFloatColumnEditor.cpp
@@ -19,7 +19,7 @@ TSharedRef<SWidget> CreateOutputFloatColumnWidget(UChooserTable* Chooser, FChoos
 {
 	FOutputFloatColumn* OutputFloatColumn = static_cast<FOutputFloatColumn*>(Column);
 
-	if (Row < 0)
+    if (Row == ColumnWidget_SpecialIndex_Header)
 	{
 		// create column header widget
 		TSharedPtr<SWidget> InputValueWidget = nullptr;
@@ -61,6 +61,20 @@ TSharedRef<SWidget> CreateOutputFloatColumnWidget(UChooserTable* Chooser, FChoos
 
 		return ColumnHeaderWidget;
 	}
+	if (Row == ColumnWidget_SpecialIndex_Fallback)
+    {
+		return SNew(SNumericEntryBox<double>)
+        		.Value_Lambda([OutputFloatColumn]()
+        		{
+        			return OutputFloatColumn->FallbackValue;
+        		})
+        		.OnValueCommitted_Lambda([Chooser, OutputFloatColumn](double NewValue, ETextCommit::Type CommitType)
+        		{
+					const FScopedTransaction Transaction(LOCTEXT("Edit Min Value", "Edit Min Value"));
+					Chooser->Modify(true);
+					OutputFloatColumn->FallbackValue = NewValue;
+        		});	
+    }
 
 	// create cell widget
 	return SNew(SNumericEntryBox<double>)
diff --git a/Engine/Plugins/Experimental/Chooser/Source/ChooserEditor/Private/OutputStructColumnEditor.cpp b/Engine/Plugins/Experimental/Chooser/Source/ChooserEditor/Private/OutputStructColumnEditor.cpp
--- a/Engine/Plugins/Experimental/Chooser/Source/ChooserEditor/Private/OutputStructColumnEditor.cpp
+++ b/Engine/Plugins/Experimental/Chooser/Source/ChooserEditor/Private/OutputStructColumnEditor.cpp
@@ -16,7 +16,7 @@ namespace UE::ChooserEditor
 	
 TSharedRef<SWidget> CreateOutputStructColumnWidget(UChooserTable* Chooser, FChooserColumnBase* Column, int Row)
 {
-	if (Row < 0)
+	if (Row == ColumnWidget_SpecialIndex_Header)
 	{
 		// create column header widget
 		TSharedPtr<SWidget> InputValueWidget = nullptr;
diff --git a/Engine/Plugins/Experimental/Chooser/Source/ChooserEditor/Private/RandomizeColumnEditor.cpp b/Engine/Plugins/Experimental/Chooser/Source/ChooserEditor/Private/RandomizeColumnEditor.cpp
--- a/Engine/Plugins/Experimental/Chooser/Source/ChooserEditor/Private/RandomizeColumnEditor.cpp
+++ b/Engine/Plugins/Experimental/Chooser/Source/ChooserEditor/Private/RandomizeColumnEditor.cpp
@@ -19,7 +19,12 @@ namespace UE::ChooserEditor
 TSharedRef<SWidget> CreateRandomizeColumnWidget(UChooserTable* Chooser, FChooserColumnBase* Column, int Row)
 {
 	FRandomizeColumn* RandomizeColumn = static_cast<FRandomizeColumn*>(Column);
-	if (Row < 0)
+	
+	if (Row == ColumnWidget_SpecialIndex_Fallback)
+	{
+		return SNullWidget::NullWidget;
+	}
+	if (Row == ColumnWidget_SpecialIndex_Header)
 	{
 		const FSlateBrush* ColumnIcon = FAppStyle::Get().GetBrush("Icons.Help");
 		
diff --git a/Engine/Plugins/Experimental/Chooser/Source/ChooserEditor/Private/RewindDebuggerChooser.cpp b/Engine/Plugins/Experimental/Chooser/Source/ChooserEditor/Private/RewindDebuggerChooser.cpp
--- a/Engine/Plugins/Experimental/Chooser/Source/ChooserEditor/Private/RewindDebuggerChooser.cpp
+++ b/Engine/Plugins/Experimental/Chooser/Source/ChooserEditor/Private/RewindDebuggerChooser.cpp
@@ -34,11 +34,12 @@ void FRewindDebuggerChooser::Update(float DeltaTime, IRewindDebugger* RewindDebu
 				{
 					if (UChooserTable* Chooser = Cast<UChooserTable>(ChooserObject))
 					{
-						if (Chooser->HasDebugTarget())
+						UChooserTable* ContextOwner = Chooser->GetContextOwner();
+						if (ContextOwner->HasDebugTarget())
 						{
 							if (UObject* ContextObject = FObjectTrace::GetObjectFromId(OwnerId))
 							{
-								if (Chooser->GetDebugTarget() == ContextObject)
+								if (ContextOwner->GetDebugTarget() == ContextObject)
 								{
 									Chooser->SetDebugSelectedRow(ChooserEvaluationData.SelectedIndex);
 									Chooser->bDebugTestValuesValid = true;
diff --git a/Engine/Plugins/Experimental/Chooser/Source/ChooserEditor/Public/ObjectChooserWidgetFactories.h b/Engine/Plugins/Experimental/Chooser/Source/ChooserEditor/Public/ObjectChooserWidgetFactories.h
--- a/Engine/Plugins/Experimental/Chooser/Source/ChooserEditor/Public/ObjectChooserWidgetFactories.h
+++ b/Engine/Plugins/Experimental/Chooser/Source/ChooserEditor/Public/ObjectChooserWidgetFactories.h
@@ -13,6 +13,9 @@ struct FChooserColumnBase;
 
 namespace UE::ChooserEditor
 {
+	enum {  ColumnWidget_SpecialIndex_Header = -1 };
+	enum {  ColumnWidget_SpecialIndex_Fallback = -2 };
+	
 	DECLARE_DELEGATE(FChooserWidgetValueChanged)
 	typedef TFunction<TSharedRef<SWidget>(bool bReadOnly, UObject* TransactionObject, void* Value, UClass* ResultBaseClass, FChooserWidgetValueChanged)> FChooserWidgetCreator;
 	typedef TFunction<TSharedRef<SWidget>(UChooserTable* Chooser, FChooserColumnBase* Column, int Row)> FColumnWidgetCreator;
@@ -23,9 +26,10 @@ namespace UE::ChooserEditor
 
 		static TSharedPtr<SWidget> CreateWidget(bool ReadOnly, UObject* TransactionObject, void* Value, const UStruct* ValueType, UClass* ResultBaseClass, FChooserWidgetValueChanged ValueChanged = FChooserWidgetValueChanged());
 		static TSharedPtr<SWidget> CreateWidget(bool ReadOnly, UObject* TransactionObject, const UScriptStruct* BaseType, void* Value, const UStruct* ValueType, UClass* ResultBaseClass,
-										const FOnStructPicked& CreateClassCallback, TSharedPtr<SBorder>* InnerWidget = nullptr, FChooserWidgetValueChanged ValueChanged = FChooserWidgetValueChanged());
+										const FOnStructPicked& CreateClassCallback, TSharedPtr<SBorder>* InnerWidget = nullptr, FChooserWidgetValueChanged ValueChanged = FChooserWidgetValueChanged(),
+										FText NullValueDisplayText = FText());
 		
-		static TSharedPtr<SWidget> CreateColumnWidget(FChooserColumnBase* Column, const UStruct* ColumnType, UChooserTable* Chooser, int Row);
+		static TSharedPtr<SWidget> CreateColumnWidget(FChooserColumnBase* Column, const UStruct* ColumnType, UChooserTable* Chooser, int Row); // todo: chooser should be a UObject
 		
 		static void RegisterWidgets();
 
