Fix for t-pose from ChooserPlayer animgraph node in certain situations.

diff --git a/Engine/Plugins/Experimental/Chooser/Source/Chooser/Private/AnimNode_ChooserPlayer.cpp b/Engine/Plugins/Experimental/Chooser/Source/Chooser/Private/AnimNode_ChooserPlayer.cpp
--- a/Engine/Plugins/Experimental/Chooser/Source/Chooser/Private/AnimNode_ChooserPlayer.cpp
+++ b/Engine/Plugins/Experimental/Chooser/Source/Chooser/Private/AnimNode_ChooserPlayer.cpp
@@ -93,6 +93,17 @@ void FAnimNode_ChooserPlayer::UpdateAssetPlayer(const FAnimationUpdateContext& C
 	const bool bJustBecameRelevant = !UpdateCounter.WasSynchronizedCounter(Context.AnimInstanceProxy->GetUpdateCounter());
 	UpdateCounter.SynchronizeWith(Context.AnimInstanceProxy->GetUpdateCounter());
 
+	if (CurrentAsset != nullptr && AnimPlayers.IsEmpty())
+	{
+		// catch the case where the blendstack has been reset, and we need to reinitialize it.
+		if (EvaluationFrequency != EChooserEvaluationFrequency::OnInitialUpdate)
+		{
+			// force reselection for all modes other than OnInitialUpdate (in that case we will just restart the current animation below)
+			CurrentAsset = nullptr;
+		}
+
+	}
+
 	UAnimationAsset* NewAsset = CurrentAsset;
 
 	if (EvaluationFrequency == EChooserEvaluationFrequency::OnUpdate || CurrentAsset == nullptr)
@@ -121,7 +132,7 @@ void FAnimNode_ChooserPlayer::UpdateAssetPlayer(const FAnimationUpdateContext& C
 	// - if the mirror setting has changed
 	// - for playback rate of 0, when the start time changes - for choosing poses as frames of an animation sequence
 	// - if the curve values are different
-	if (bJustBecameRelevant || NewAsset != CurrentAsset ||
+	if (bJustBecameRelevant || NewAsset != CurrentAsset || AnimPlayers.IsEmpty() ||
 		CurrentMirror != Settings.bMirror ||
 		(CurrentStartTime != Settings.StartTime && Settings.PlaybackRate == 0.0f) ||
 		CurrentCurveOverridesHash != Settings.CurveOverrides.Hash)
