proxy table and proxy asset icons

diff --git a/Engine/Plugins/Experimental/Chooser/Source/ProxyTable/Private/ProxyAsset.cpp b/Engine/Plugins/Experimental/Chooser/Source/ProxyTable/Private/ProxyAsset.cpp
--- a/Engine/Plugins/Experimental/Chooser/Source/ProxyTable/Private/ProxyAsset.cpp
+++ b/Engine/Plugins/Experimental/Chooser/Source/ProxyTable/Private/ProxyAsset.cpp
@@ -6,6 +6,21 @@
 #include "UObject/Package.h"
 #include "LookupProxy.h"
 
+FName UProxyAsset::TypeTagName = "ProxyType";
+
+void UProxyAsset::GetAssetRegistryTags(TArray<FAssetRegistryTag>& OutTags) const
+{
+	FString ResultTypeName;
+	if (Type)
+	{
+		ResultTypeName = Type.GetName();
+	}
+	OutTags.Add({TypeTagName, ResultTypeName, FAssetRegistryTag::TT_Alphabetical});
+	
+	
+	UObject::GetAssetRegistryTags(OutTags);
+}
+
 #if WITH_EDITOR
 void UProxyAsset::PostEditUndo()
 {
diff --git a/Engine/Plugins/Experimental/Chooser/Source/ProxyTable/Public/ProxyAsset.h b/Engine/Plugins/Experimental/Chooser/Source/ProxyTable/Public/ProxyAsset.h
--- a/Engine/Plugins/Experimental/Chooser/Source/ProxyTable/Public/ProxyAsset.h
+++ b/Engine/Plugins/Experimental/Chooser/Source/ProxyTable/Public/ProxyAsset.h
@@ -19,6 +19,10 @@ class UProxyAsset : public UObject, public IHasContextClass
 public:
 	UProxyAsset() {}
 
+	virtual void GetAssetRegistryTags(TArray<FAssetRegistryTag>& OutTags) const override;
+
+	static PROXYTABLE_API FName TypeTagName;
+
 #if WITH_EDITOR
 	FProxyTypeChanged OnTypeChanged;
 	virtual void PostEditUndo() override;
@@ -38,7 +42,6 @@ public:
 	UPROPERTY(EditAnywhere, Meta = (ExcludeBaseStruct, BaseStruct = "/Script/Chooser.ContextObjectTypeBase"), Category = "Input")
 	TArray<FInstancedStruct> ContextData;
 
-	
 	UPROPERTY(EditAnywhere, Meta = (ExcludeBaseStruct, BaseStruct ="/Script/ProxyTable.ChooserParameterProxyTableBase"), Category = "Proxy Table Reference")
 	FInstancedStruct ProxyTable;
 
diff --git a/Engine/Plugins/Experimental/Chooser/Source/ProxyTableEditor/Private/AssetDefinition_ProxyAsset.cpp b/Engine/Plugins/Experimental/Chooser/Source/ProxyTableEditor/Private/AssetDefinition_ProxyAsset.cpp
--- a/Engine/Plugins/Experimental/Chooser/Source/ProxyTableEditor/Private/AssetDefinition_ProxyAsset.cpp
+++ b/Engine/Plugins/Experimental/Chooser/Source/ProxyTableEditor/Private/AssetDefinition_ProxyAsset.cpp
@@ -1,8 +1,13 @@
 // Copyright Epic Games, Inc. All Rights Reserved.
 
 #include "AssetDefinition_ProxyAsset.h"
+
+#include "ProxyTableEditor.h"
+#include "ProxyTableEditorStyle.h"
 #include "Toolkits/SimpleAssetEditor.h"
 
+#include "Animation/AnimSequence.h"
+
 #define LOCTEXT_NAMESPACE "UAssetDefinition_ProxyAsset" 
 
 EAssetCommandResult UAssetDefinition_ProxyAsset::OpenAssets(const FAssetOpenArgs& OpenArgs) const
@@ -13,4 +18,23 @@ EAssetCommandResult UAssetDefinition_ProxyAsset::OpenAssets(const FAssetOpenArgs
 	return EAssetCommandResult::Handled;
 }
 
+const FSlateBrush* UAssetDefinition_ProxyAsset::GetThumbnailBrush(const FAssetData& InAssetData, const FName InClassName) const
+{
+	FString ProxyType;
+	if (InAssetData.GetTagValue<FString>(UProxyAsset::TypeTagName, ProxyType))
+	{
+		if (ProxyType.Contains("AnimSequence"))
+		{
+			return UE::ProxyTableEditor::FProxyTableEditorStyle::Get().GetBrush("ProxyTableEditor.ProxyAnimSequenceIconLarge");
+		}
+	}
+	
+	return UE::ProxyTableEditor::FProxyTableEditorStyle::Get().GetBrush("ProxyTableEditor.ProxyAssetIconLarge");
+}
+
+const FSlateBrush* UAssetDefinition_ProxyAsset::GetIconBrush(const FAssetData& InAssetData, const FName InClassName) const
+{
+	return UE::ProxyTableEditor::FProxyTableEditorStyle::Get().GetBrush("ProxyTableEditor.ProxyAssetIconLarge");
+}
+
 #undef LOCTEXT_NAMESPACE
\ No newline at end of file
diff --git a/Engine/Plugins/Experimental/Chooser/Source/ProxyTableEditor/Private/AssetDefinition_ProxyAsset.h b/Engine/Plugins/Experimental/Chooser/Source/ProxyTableEditor/Private/AssetDefinition_ProxyAsset.h
--- a/Engine/Plugins/Experimental/Chooser/Source/ProxyTableEditor/Private/AssetDefinition_ProxyAsset.h
+++ b/Engine/Plugins/Experimental/Chooser/Source/ProxyTableEditor/Private/AssetDefinition_ProxyAsset.h
@@ -23,5 +23,8 @@ public:
 		return Categories;
 	}
 	virtual EAssetCommandResult OpenAssets(const FAssetOpenArgs& OpenArgs) const override;
+	
+	const FSlateBrush* GetIconBrush(const FAssetData& InAssetData, const FName InClassName) const override;
+	const FSlateBrush* GetThumbnailBrush(const FAssetData& InAssetData, const FName InClassName) const override;
 	// UAssetDefinition End
 };
diff --git a/Engine/Plugins/Experimental/Chooser/Source/ProxyTableEditor/Private/AssetDefinition_ProxyTable.cpp b/Engine/Plugins/Experimental/Chooser/Source/ProxyTableEditor/Private/AssetDefinition_ProxyTable.cpp
--- a/Engine/Plugins/Experimental/Chooser/Source/ProxyTableEditor/Private/AssetDefinition_ProxyTable.cpp
+++ b/Engine/Plugins/Experimental/Chooser/Source/ProxyTableEditor/Private/AssetDefinition_ProxyTable.cpp
@@ -1,6 +1,9 @@
 // Copyright Epic Games, Inc. All Rights Reserved.
 
 #include "AssetDefinition_ProxyTable.h"
+
+#include <ProxyTableEditorStyle.h>
+
 #include "ProxyTableEditor.h"
 
 #define LOCTEXT_NAMESPACE "UAssetDefinition_ProxyTable" 
@@ -16,5 +19,14 @@ EAssetCommandResult UAssetDefinition_ProxyTable::OpenAssets(const FAssetOpenArgs
 
 	return EAssetCommandResult::Handled;
 }
+const FSlateBrush* UAssetDefinition_ProxyTable::GetThumbnailBrush(const FAssetData& InAssetData, const FName InClassName) const
+{
+	return UE::ProxyTableEditor::FProxyTableEditorStyle::Get().GetBrush("ProxyTableEditor.ProxyTableIconLarge");
+}
+
+const FSlateBrush* UAssetDefinition_ProxyTable::GetIconBrush(const FAssetData& InAssetData, const FName InClassName) const
+{
+	return UE::ProxyTableEditor::FProxyTableEditorStyle::Get().GetBrush("ProxyTableEditor.ProxyTableIconLarge");
+}
 
 #undef LOCTEXT_NAMESPACE
\ No newline at end of file
diff --git a/Engine/Plugins/Experimental/Chooser/Source/ProxyTableEditor/Private/AssetDefinition_ProxyTable.h b/Engine/Plugins/Experimental/Chooser/Source/ProxyTableEditor/Private/AssetDefinition_ProxyTable.h
--- a/Engine/Plugins/Experimental/Chooser/Source/ProxyTableEditor/Private/AssetDefinition_ProxyTable.h
+++ b/Engine/Plugins/Experimental/Chooser/Source/ProxyTableEditor/Private/AssetDefinition_ProxyTable.h
@@ -23,5 +23,8 @@ public:
 		return Categories;
 	}
 	virtual EAssetCommandResult OpenAssets(const FAssetOpenArgs& OpenArgs) const override;
+	
+	const FSlateBrush* GetIconBrush(const FAssetData& InAssetData, const FName InClassName) const override;
+	const FSlateBrush* GetThumbnailBrush(const FAssetData& InAssetData, const FName InClassName) const override;
 	// UAssetDefinition End
 };
diff --git a/Engine/Plugins/Experimental/Chooser/Source/ProxyTableEditor/Private/ProxyAssetFactory.h b/Engine/Plugins/Experimental/Chooser/Source/ProxyTableEditor/Private/ProxyAssetFactory.h
--- a/Engine/Plugins/Experimental/Chooser/Source/ProxyTableEditor/Private/ProxyAssetFactory.h
+++ b/Engine/Plugins/Experimental/Chooser/Source/ProxyTableEditor/Private/ProxyAssetFactory.h
@@ -19,4 +19,6 @@ class UProxyAssetFactory : public UFactory
 	{
 		return FactoryCreateNew(Class, InParent, Name, Flags, Context, Warn, NAME_None);
 	}
+
+	virtual FName GetNewAssetThumbnailOverride() const override { return "ProxyTableEditor.ProxyAssetIconLarge"; }
 };
\ No newline at end of file
diff --git a/Engine/Plugins/Experimental/Chooser/Source/ProxyTableEditor/Private/ProxyTableEditorModule.cpp b/Engine/Plugins/Experimental/Chooser/Source/ProxyTableEditor/Private/ProxyTableEditorModule.cpp
--- a/Engine/Plugins/Experimental/Chooser/Source/ProxyTableEditor/Private/ProxyTableEditorModule.cpp
+++ b/Engine/Plugins/Experimental/Chooser/Source/ProxyTableEditor/Private/ProxyTableEditorModule.cpp
@@ -5,6 +5,7 @@
 #include "ProxyTableEditorCommands.h"
 #include "StructOutputDataCustomization.h"
 #include "PropertyEditorModule.h"
+#include "ProxyTableEditorStyle.h"
 
 #define LOCTEXT_NAMESPACE "ProxyTableEditorModule"
 
@@ -13,6 +14,7 @@ namespace UE::ProxyTableEditor
 
 void FModule::StartupModule()
 {
+	FProxyTableEditorStyle::Initialize();
 	FProxyTableEditor::RegisterWidgets();
 
 	FProxyTableEditorCommands::Register();
@@ -25,6 +27,9 @@ void FModule::ShutdownModule()
 {
 	FProxyTableEditorCommands::Unregister();
 	FPropertyEditorModule& PropertyModule = FModuleManager::LoadModuleChecked<FPropertyEditorModule>("PropertyEditor");
+	PropertyModule.UnregisterCustomPropertyTypeLayout(FProxyStructOutput::StaticStruct()->GetFName());
+	
+	FProxyTableEditorStyle::Shutdown();
 }
 
 }
diff --git a/Engine/Plugins/Experimental/Chooser/Source/ProxyTableEditor/Private/ProxyTableEditorStyle.cpp b/Engine/Plugins/Experimental/Chooser/Source/ProxyTableEditor/Private/ProxyTableEditorStyle.cpp
new file mode 100644
--- /dev/null
+++ b/Engine/Plugins/Experimental/Chooser/Source/ProxyTableEditor/Private/ProxyTableEditorStyle.cpp
@@ -0,0 +1,54 @@
+// Copyright Epic Games, Inc. All Rights Reserved.
+
+#include "ProxyTableEditorStyle.h"
+#include "Styling/CoreStyle.h"
+#include "Styling/SlateStyle.h"
+#include "Styling/SlateStyleMacros.h"
+#include "Styling/SlateStyleRegistry.h"
+
+namespace UE::ProxyTableEditor
+{
+	
+TSharedPtr< FProxyTableEditorStyle > FProxyTableEditorStyle::StyleInstance = nullptr;
+
+
+FProxyTableEditorStyle::FProxyTableEditorStyle() :
+	FSlateStyleSet("ProxyTableEditorStyle")
+{
+	const FVector2D Icon16x16(16.0f, 16.0f);
+	const FVector2D Icon24x24(24.0f, 24.0f);
+
+	SetContentRoot(FPaths::EnginePluginsDir() / TEXT("Experimental/Chooser/Content"));
+	SetCoreContentRoot(FPaths::EngineContentDir() / TEXT(""));
+
+	// proxy asset icon
+	Set("ProxyTableEditor.ProxyAssetIconLarge", new IMAGE_BRUSH_SVG("ProxyAsset_64", Icon24x24));
+	// proxy sequence icon
+	Set("ProxyTableEditor.ProxyAnimSequenceIconLarge", new IMAGE_BRUSH_SVG("ProxyAnimSequence_64", Icon24x24));
+
+	// proxy table icon
+	Set("ProxyTableEditor.ProxyTableIconLarge", new IMAGE_BRUSH_SVG("ProxyTable_64", Icon24x24));
+}
+
+void FProxyTableEditorStyle::Initialize()
+{
+	if (!StyleInstance.IsValid())
+	{
+		StyleInstance = MakeShared<FProxyTableEditorStyle>();
+		FSlateStyleRegistry::RegisterSlateStyle(*StyleInstance);
+	}
+}
+
+void FProxyTableEditorStyle::Shutdown()
+{
+	FSlateStyleRegistry::UnRegisterSlateStyle(*StyleInstance);
+	ensure(StyleInstance.IsUnique());
+	StyleInstance.Reset();
+}
+
+const ISlateStyle& FProxyTableEditorStyle::Get()
+{
+	return *StyleInstance;
+}
+
+}
\ No newline at end of file
diff --git a/Engine/Plugins/Experimental/Chooser/Source/ProxyTableEditor/Private/ProxyTableEditorStyle.h b/Engine/Plugins/Experimental/Chooser/Source/ProxyTableEditor/Private/ProxyTableEditorStyle.h
new file mode 100644
--- /dev/null
+++ b/Engine/Plugins/Experimental/Chooser/Source/ProxyTableEditor/Private/ProxyTableEditorStyle.h
@@ -0,0 +1,24 @@
+// Copyright Epic Games, Inc. All Rights Reserved.
+#pragma once
+
+#include "CoreMinimal.h"
+#include "Styling/SlateStyle.h"
+
+namespace UE::ProxyTableEditor
+{
+	
+class FProxyTableEditorStyle : FSlateStyleSet
+{
+public:
+    FProxyTableEditorStyle();
+	static void Initialize();
+	static void Shutdown();
+
+	static const ISlateStyle& Get();
+
+private:
+
+	static TSharedPtr< FProxyTableEditorStyle > StyleInstance;
+};
+
+}
diff --git a/Engine/Plugins/Experimental/Chooser/Source/ProxyTableEditor/Private/ProxyTableFactory.h b/Engine/Plugins/Experimental/Chooser/Source/ProxyTableEditor/Private/ProxyTableFactory.h
--- a/Engine/Plugins/Experimental/Chooser/Source/ProxyTableEditor/Private/ProxyTableFactory.h
+++ b/Engine/Plugins/Experimental/Chooser/Source/ProxyTableEditor/Private/ProxyTableFactory.h
@@ -19,4 +19,5 @@ class UProxyTableFactory : public UFactory
 	{
 		return FactoryCreateNew(Class, InParent, Name, Flags, Context, Warn, NAME_None);
 	}
+	virtual FName GetNewAssetThumbnailOverride() const override { return "ProxyTableEditor.ProxyTableIconLarge"; }
 };
\ No newline at end of file
