Follow up changes for nested chooser table support
- Fix for drag/drop crash
- improved display for fallback result
- fix result type filtering in asset pickers not using main chooser result type

diff --git a/Engine/Plugins/Experimental/Chooser/Source/ChooserEditor/Private/ChooserEditorStyle.cpp b/Engine/Plugins/Experimental/Chooser/Source/ChooserEditor/Private/ChooserEditorStyle.cpp
--- a/Engine/Plugins/Experimental/Chooser/Source/ChooserEditor/Private/ChooserEditorStyle.cpp
+++ b/Engine/Plugins/Experimental/Chooser/Source/ChooserEditor/Private/ChooserEditorStyle.cpp
@@ -26,6 +26,8 @@ FChooserEditorStyle::FChooserEditorStyle() :
 
 	// tab icon
 	Set("ChooserEditor.ChooserTableIconSmall", new IMAGE_BRUSH_SVG("ChooserIcon_16", Icon16x16));
+	
+	Set("ChooserEditor.FallbackIcon", new IMAGE_BRUSH_SVG("Fallback", Icon16x16));
 }
 
 void FChooserEditorStyle::Initialize()
diff --git a/Engine/Plugins/Experimental/Chooser/Source/ChooserEditor/Private/ChooserTableEditor.cpp b/Engine/Plugins/Experimental/Chooser/Source/ChooserEditor/Private/ChooserTableEditor.cpp
--- a/Engine/Plugins/Experimental/Chooser/Source/ChooserEditor/Private/ChooserTableEditor.cpp
+++ b/Engine/Plugins/Experimental/Chooser/Source/ChooserEditor/Private/ChooserTableEditor.cpp
@@ -678,18 +678,19 @@ public:
 			if (ColumnName == Handles)
 			{
 				// row drag handle
-			
 				return SNew(SChooserRowHandle).ChooserEditor(Editor).RowIndex(RowIndex->RowIndex);
 			}
 			else if (ColumnName == Result) 
 			{
-				TSharedPtr<SWidget> ResultWidget = FObjectChooserWidgetFactories::CreateWidget(false, Chooser->GetContextOwner(), FObjectChooserBase::StaticStruct(), Chooser->ResultsStructs[RowIndex->RowIndex].GetMutableMemory(), Chooser->ResultsStructs[RowIndex->RowIndex].GetScriptStruct(), Chooser->OutputObjectType,
+				UChooserTable* ContextOwner = Chooser->GetContextOwner();
+				TSharedPtr<SWidget> ResultWidget = FObjectChooserWidgetFactories::CreateWidget(false, ContextOwner, FObjectChooserBase::StaticStruct(), Chooser->ResultsStructs[RowIndex->RowIndex].GetMutableMemory(), Chooser->ResultsStructs[RowIndex->RowIndex].GetScriptStruct(), ContextOwner->OutputObjectType,
 				FOnStructPicked::CreateLambda([this, RowIndex=RowIndex->RowIndex](const UScriptStruct* ChosenStruct)
 				{
+					UChooserTable* ContextOwner = Chooser->GetContextOwner();
 					const FScopedTransaction Transaction(LOCTEXT("Change Row Result Type", "Change Row Result Type"));
 					Chooser->Modify(true);
 					Chooser->ResultsStructs[RowIndex].InitializeAs(ChosenStruct);
-					FObjectChooserWidgetFactories::CreateWidget(false, Chooser->GetContextOwner(), FObjectChooserBase::StaticStruct(), Chooser->ResultsStructs[RowIndex].GetMutableMemory(), ChosenStruct, Chooser->OutputObjectType, FOnStructPicked(), &CacheBorder);
+					FObjectChooserWidgetFactories::CreateWidget(false, ContextOwner, FObjectChooserBase::StaticStruct(), Chooser->ResultsStructs[RowIndex].GetMutableMemory(), ChosenStruct, ContextOwner->OutputObjectType, FOnStructPicked(), &CacheBorder);
 				}),
 				&CacheBorder
 				);
@@ -754,24 +755,40 @@ public:
 		{
 			if (ColumnName == Handles)
 			{
-				return SNullWidget::NullWidget;
+				return SNew(SBox).Padding(0.0f) .HAlign(HAlign_Center) .VAlign(VAlign_Center) .WidthOverride(16.0f)
+					[
+						SNew(SImage)
+						.Image(FChooserEditorStyle::Get().GetBrush("ChooserEditor.FallbackIcon"))
+						.ToolTipText(LOCTEXT("FallbackTooltip","Fallback result:  Returned if all rows failed."))
+					];
 			}
 			else if (ColumnName == Result) 
 			{
-				TSharedPtr<SWidget> ResultWidget = FObjectChooserWidgetFactories::CreateWidget(false, Chooser->GetContextOwner(), FObjectChooserBase::StaticStruct(), Chooser->FallbackResult.GetMutableMemory(), Chooser->FallbackResult.GetScriptStruct(), Chooser->OutputObjectType,
+				UChooserTable* ContextOwner = Chooser->GetContextOwner();
+				TSharedPtr<SWidget> ResultWidget = FObjectChooserWidgetFactories::CreateWidget(false,ContextOwner, FObjectChooserBase::StaticStruct(), Chooser->FallbackResult.GetMutableMemory(), Chooser->FallbackResult.GetScriptStruct(), ContextOwner->OutputObjectType,
 				FOnStructPicked::CreateLambda([this](const UScriptStruct* ChosenStruct)
 				{
+				UChooserTable* ContextOwner = Chooser->GetContextOwner();
 					const FScopedTransaction Transaction(LOCTEXT("Change Row Result Type", "Change Row Result Type"));
 					Chooser->Modify(true);
 					Chooser->FallbackResult.InitializeAs(ChosenStruct);
-					FObjectChooserWidgetFactories::CreateWidget(false, Chooser->GetContextOwner(), FObjectChooserBase::StaticStruct(), Chooser->FallbackResult.GetMutableMemory(), ChosenStruct, Chooser->OutputObjectType, FOnStructPicked(), &CacheBorder
+					FObjectChooserWidgetFactories::CreateWidget(false, ContextOwner, FObjectChooserBase::StaticStruct(), Chooser->FallbackResult.GetMutableMemory(), ChosenStruct, ContextOwner->OutputObjectType, FOnStructPicked(), &CacheBorder
 						,FChooserWidgetValueChanged(), LOCTEXT("Fallback Result", "Fallback Result: (None)"));
 				}),
 				&CacheBorder
 				,FChooserWidgetValueChanged(), LOCTEXT("Fallback Result", "Fallback Result: (None)")
 				);
-			
-				return ResultWidget.ToSharedRef();
+				
+				return SNew(SOverlay)
+						+ SOverlay::Slot()
+						[
+							ResultWidget.ToSharedRef()
+						]
+						+ SOverlay::Slot().VAlign(VAlign_Top)
+						[
+							SNew(SSeparator).SeparatorImage(FCoreStyle::Get().GetBrush("FocusRectangle"))
+							.Visibility_Lambda([this]() { return bDragActive  ? EVisibility::Visible : EVisibility::Hidden; })
+						];
 			}
 			else
 			{
@@ -834,7 +851,12 @@ public:
 			if (Chooser == Operation->ChooserEditor->GetChooser())
 			{
 				int NewRowIndex;
-				if (bDropAbove)
+				if (!Chooser->ResultsStructs.IsValidIndex(RowIndex->RowIndex))
+				{
+					// for special (negative) indices, move to the end
+					NewRowIndex = Editor->MoveRow(Operation->RowIndex, Chooser->ResultsStructs.Num());
+				}
+				else if (bDropAbove)
 				{
 					NewRowIndex = Editor->MoveRow(Operation->RowIndex, RowIndex->RowIndex);
 				}
@@ -1521,7 +1543,8 @@ TSharedRef<SWidget> CreateEvaluateChooserWidget(bool bReadOnly, UObject* Transac
 		 	{
 		 		if (UChooserTable* Chooser = Cast<UChooserTable>(InAssetData.GetAsset()))
 		 		{
-		 			return !(Chooser->OutputObjectType && Chooser->OutputObjectType->IsChildOf(ResultBaseClass));
+		 			UChooserTable* ContextOwner = Chooser->GetContextOwner();
+		 			return !(ContextOwner->OutputObjectType && ContextOwner->OutputObjectType->IsChildOf(ResultBaseClass));
 		 		}
 		 	}
 			return true;
