ChooserTrack fixes for loading trace files from disk

diff --git a/Engine/Plugins/Experimental/Chooser/Source/Chooser/Internal/Chooser.h b/Engine/Plugins/Experimental/Chooser/Source/Chooser/Internal/Chooser.h
--- a/Engine/Plugins/Experimental/Chooser/Source/Chooser/Internal/Chooser.h
+++ b/Engine/Plugins/Experimental/Chooser/Source/Chooser/Internal/Chooser.h
@@ -38,22 +38,14 @@ public:
 	void SetDebugSelectedRow(int32 Index) const { DebugSelectedRow = Index; }
 	int32 GetDebugSelectedRow() const { return DebugSelectedRow; }
 	bool HasDebugTarget() const { return !DebugTargetName.IsEmpty(); }
-	const UObject* GetDebugTarget() const { return DebugTarget.Get(); }
+	
+	void SetDebugTarget(FString Name) { DebugTargetName = Name; }
+	void ResetDebugTarget() { DebugTargetName = ""; }
 	const FString& GetDebugTargetName() const { return DebugTargetName; }
-	void SetDebugTarget(TWeakObjectPtr<const UObject> Target)
-	{
-		DebugTarget = Target;
-		if (const UObject* DebugObject = DebugTarget.Get())
-		{
-			DebugTargetName = DebugObject->GetName();
-		}
-	}
-	void ResetDebugTarget()
-	{
-		DebugTarget.Reset();
-		DebugTargetName = "";
-	}
-	void IterateRecentContextObjects(TFunction<void(const UObject*)> Callback) const;
+	
+
+	void AddRecentContextObject(const FString& ObjectName) const { RecentContextObjects.Add(ObjectName); }
+	void IterateRecentContextObjects(TFunction<void(const FString&)> Callback) const;
 	void UpdateDebugging(FChooserEvaluationContext& Context) const;
 	
 	// enable display of which cells pass/fail based on current TestValue for each column
@@ -69,7 +61,7 @@ private:
 	EObjectChooserResultType CachedPreviousResultType = EObjectChooserResultType::ObjectResult;
 	
 	// objects this chooser has been recently evaluated on
-	mutable TSet<TWeakObjectPtr<const UObject>> RecentContextObjects;
+	mutable TSet<FString> RecentContextObjects;
 	mutable FCriticalSection DebugLock;
 	// reference to the UObject in PIE  which we want to get debug info for
 	mutable TWeakObjectPtr<const UObject> DebugTarget;
diff --git a/Engine/Plugins/Experimental/Chooser/Source/Chooser/Private/Chooser.cpp b/Engine/Plugins/Experimental/Chooser/Source/Chooser/Private/Chooser.cpp
--- a/Engine/Plugins/Experimental/Chooser/Source/Chooser/Private/Chooser.cpp
+++ b/Engine/Plugins/Experimental/Chooser/Source/Chooser/Private/Chooser.cpp
@@ -200,15 +200,12 @@ void UChooserTable::PostEditChangeProperty(FPropertyChangedEvent& PropertyChange
 	}
 }
 
-void UChooserTable::IterateRecentContextObjects(TFunction<void(const UObject*)> Callback) const
+void UChooserTable::IterateRecentContextObjects(TFunction<void(const FString&)> Callback) const
 {
 	FScopeLock Lock(&DebugLock);
-	for(TWeakObjectPtr<const UObject>& Object : RecentContextObjects)
+	for(const FString& ObjectName : RecentContextObjects)
 	{
-		if (Object.IsValid())
-		{
-			Callback(Object.Get());
-		}
+		Callback(ObjectName);
 	}
 }
 
@@ -224,18 +221,9 @@ void UChooserTable::UpdateDebugging(FChooserEvaluationContext& Context) const
 		{
 			if (UObject* ContextObject = ObjectParam->Object)
 			{
-				RecentContextObjects.Add(MakeWeakObjectPtr(ContextObject));
-				
-				if (ContextOwner->DebugTarget == nullptr && !DebugTargetName.IsEmpty())
-				{
-					// if the DebugTargetName is set, but not the DebugTarget, it means that PIE has been restarted, so try matching by name
-					if (ContextObject->GetName() == ContextOwner->DebugTargetName)
-					{
-						ContextOwner->DebugTarget = ContextObject;
-					}
-				}
+				RecentContextObjects.Add(ContextObject->GetName());
 				
-				if (ContextObject == ContextOwner->DebugTarget) 
+				if (ContextObject->GetName() == ContextOwner->GetDebugTargetName()) 
 				{
 					bDebugTestValuesValid = true;
 					Context.DebuggingInfo.bCurrentDebugTarget = true;
diff --git a/Engine/Plugins/Experimental/Chooser/Source/ChooserEditor/Private/ChooserTableEditor.cpp b/Engine/Plugins/Experimental/Chooser/Source/ChooserEditor/Private/ChooserTableEditor.cpp
--- a/Engine/Plugins/Experimental/Chooser/Source/ChooserEditor/Private/ChooserTableEditor.cpp
+++ b/Engine/Plugins/Experimental/Chooser/Source/ChooserEditor/Private/ChooserTableEditor.cpp
@@ -177,28 +177,25 @@ void FChooserTableEditor::MakeDebugTargetMenu(UToolMenu* InToolMenu)
 
 	const UChooserTable* Chooser = GetChooser();
 
-	Chooser->IterateRecentContextObjects([this, InToolMenu](const UObject* Object)
+	Chooser->IterateRecentContextObjects([this, InToolMenu](const FString& ObjectName)
 		{
 			InToolMenu->AddMenuEntry(
 						SectionName,
 						FToolMenuEntry::InitMenuEntry(
-							Object->GetFName(),
-							FText::FromString(Object->GetName()),
+							FName(ObjectName),
+							FText::FromString(ObjectName),
 							LOCTEXT("Select Object ToolTip", "Select this object as the debug target"),
 							FSlateIcon(),
 							FUIAction(
-								FExecuteAction::CreateLambda([this, ObjectPtr = MakeWeakObjectPtr(Object)]()
+								FExecuteAction::CreateLambda([this, ObjectName]()
 								{
-									if(ObjectPtr.IsValid())
+									UChooserTable* Chooser = GetRootChooser();
+									Chooser->SetDebugTarget(ObjectName);
+									Chooser->bDebugTestValuesValid = false;
+									if (!Chooser->bEnableDebugTesting)
 									{
-										UChooserTable* Chooser = GetRootChooser();
-										Chooser->SetDebugTarget(ObjectPtr);
-										Chooser->bDebugTestValuesValid = false;
-										if (!Chooser->bEnableDebugTesting)
-										{
-											Chooser->bEnableDebugTesting = true;
-											UpdateTableColumns();
-										}
+										Chooser->bEnableDebugTesting = true;
+										UpdateTableColumns();
 									}
 								}),
 								FCanExecuteAction()
diff --git a/Engine/Plugins/Experimental/Chooser/Source/ChooserEditor/Private/ChooserTrack.cpp b/Engine/Plugins/Experimental/Chooser/Source/ChooserEditor/Private/ChooserTrack.cpp
--- a/Engine/Plugins/Experimental/Chooser/Source/ChooserEditor/Private/ChooserTrack.cpp
+++ b/Engine/Plugins/Experimental/Chooser/Source/ChooserEditor/Private/ChooserTrack.cpp
@@ -7,7 +7,6 @@
 #include "Editor.h"
 #include "IGameplayProvider.h"
 #include "IRewindDebugger.h"
-#include "ObjectTrace.h"
 #include "Subsystems/AssetEditorSubsystem.h"
 
 #define LOCTEXT_NAMESPACE "ChooserTrack"
@@ -22,11 +21,10 @@ FChooserTrack::FChooserTrack(uint64 InObjectId, uint64 InChooserId) :
 	EventData = MakeShared<SEventTimelineView::FTimelineEventData>();
 	Icon = FSlateIcon("EditorStyle", "Sequencer.Tracks.Event", "Sequencer.Tracks.Event");
 	
-	ChooserTable = FObjectTrace::GetObjectFromId(InChooserId);
-	if (ChooserTable)
-	{
-		TrackName = FText::FromString(ChooserTable->GetName());
-	}
+	const IGameplayProvider* GameplayProvider = IRewindDebugger::Instance()->GetAnalysisSession()->ReadProvider<IGameplayProvider>("GameplayProvider");
+	const FObjectInfo& ChooserInfo = GameplayProvider->GetObjectInfo(InChooserId);
+	TrackName = FText::FromString(ChooserInfo.Name);
+	ChooserTable = FindObject<UChooserTable>(nullptr, ChooserInfo.PathName);
 }
 
 FChooserTrack::~FChooserTrack()
@@ -109,15 +107,14 @@ bool FChooserTrack::HandleDoubleClickInternal()
 		const IGameplayProvider* GameplayProvider = AnalysisSession->ReadProvider<IGameplayProvider>("GameplayProvider");
 		const FObjectInfo& AssetInfo = GameplayProvider->GetObjectInfo(GetChooserId());
 
-		if (UObject* OwnerObject = FObjectTrace::GetObjectFromId(ObjectId))
+		// attach chooser table editor debugging
+		if (ChooserTable)
 		{
-			if (UChooserTable* Chooser = Cast<UChooserTable>(FObjectTrace::GetObjectFromId(GetChooserId())))
-			{
-				// attach chooser table editor debugging
-				Chooser->SetDebugTarget(OwnerObject);
-				Chooser->bEnableDebugTesting = true;
-			}
+			const FObjectInfo& OwnerObjectInfo  = GameplayProvider->GetObjectInfo(ObjectId);
+			ChooserTable->SetDebugTarget(OwnerObjectInfo.Name);
+			ChooserTable->bEnableDebugTesting = true;
 		}
+		
 
 		GEditor->GetEditorSubsystem<UAssetEditorSubsystem>()->OpenEditorForAsset(AssetInfo.PathName);
 
diff --git a/Engine/Plugins/Experimental/Chooser/Source/ChooserEditor/Private/ChooserTrack.h b/Engine/Plugins/Experimental/Chooser/Source/ChooserEditor/Private/ChooserTrack.h
--- a/Engine/Plugins/Experimental/Chooser/Source/ChooserEditor/Private/ChooserTrack.h
+++ b/Engine/Plugins/Experimental/Chooser/Source/ChooserEditor/Private/ChooserTrack.h
@@ -5,8 +5,9 @@
 #include "RewindDebuggerTrack.h"
 #include "Textures/SlateIcon.h"
 #include "SEventTimelineView.h"
-#include "IDetailsView.h"
 
+class UChooserTable;
+	
 namespace UE::ChooserEditor
 {
 
@@ -32,7 +33,7 @@ private:
 	mutable TSharedPtr<SEventTimelineView::FTimelineEventData> EventData;
 	mutable int EventUpdateRequested = 0;
 
-	UObject* ChooserTable;
+	UChooserTable* ChooserTable = 0;
 	uint64 ObjectId;
 	uint64 ChooserId;
 	FSlateIcon Icon;
diff --git a/Engine/Plugins/Experimental/Chooser/Source/ChooserEditor/Private/RewindDebuggerChooser.cpp b/Engine/Plugins/Experimental/Chooser/Source/ChooserEditor/Private/RewindDebuggerChooser.cpp
--- a/Engine/Plugins/Experimental/Chooser/Source/ChooserEditor/Private/RewindDebuggerChooser.cpp
+++ b/Engine/Plugins/Experimental/Chooser/Source/ChooserEditor/Private/RewindDebuggerChooser.cpp
@@ -4,7 +4,7 @@
 
 #include "Chooser.h"
 #include "ChooserProvider.h"
-#include "ObjectTrace.h"
+#include "IGameplayProvider.h"
 
 FRewindDebuggerChooser::FRewindDebuggerChooser()
 {
@@ -17,53 +17,54 @@ void FRewindDebuggerChooser::Update(float DeltaTime, IRewindDebugger* RewindDebu
 	TraceServices::FAnalysisSessionReadScope SessionReadScope(*AnalysisSession);
 	
 	const FChooserProvider* ChooserProvider = AnalysisSession->ReadProvider<FChooserProvider>(FChooserProvider::ProviderName);
+	const IGameplayProvider* GameplayProvider = AnalysisSession->ReadProvider<IGameplayProvider>("GameplayProvider");
 	
 	const TraceServices::IFrameProvider& FrameProvider = TraceServices::ReadFrameProvider(*RewindDebugger->GetAnalysisSession());
 	TraceServices::FFrame Frame;
 	if (FrameProvider.GetFrameFromTime(ETraceFrameType::TraceFrameType_Game, RewindDebugger->CurrentTraceTime(), Frame))
 	{
-		
-		ChooserProvider->EnumerateChooserEvaluationTimelines([ChooserProvider, RewindDebugger, Frame](uint64 OwnerId, const FChooserProvider::ChooserEvaluationTimeline& ChooserEvaluationTimeline)
+		ChooserProvider->EnumerateChooserEvaluationTimelines([GameplayProvider, ChooserProvider, RewindDebugger, Frame](uint64 OwnerId, const FChooserProvider::ChooserEvaluationTimeline& ChooserEvaluationTimeline)
 		{
 			double StartTime = RewindDebugger->GetScrubTime();
 			double EndTime = StartTime + Frame.EndTime - Frame.StartTime;
 			
-			ChooserEvaluationTimeline.EnumerateEvents(StartTime, EndTime, [ChooserProvider, StartTime, EndTime, OwnerId](double InStartTime, double InEndTime, uint32 InDepth, const FChooserEvaluationData& ChooserEvaluationData)
+			ChooserEvaluationTimeline.EnumerateEvents(StartTime, EndTime, [GameplayProvider, ChooserProvider, StartTime, EndTime, OwnerId](double InStartTime, double InEndTime, uint32 InDepth, const FChooserEvaluationData& ChooserEvaluationData)  
 			{
-				if (UObject* ChooserObject = FObjectTrace::GetObjectFromId(ChooserEvaluationData.ChooserId))
+				const FObjectInfo& ChooserInfo = GameplayProvider->GetObjectInfo(ChooserEvaluationData.ChooserId);
+				
+				if (UChooserTable* Chooser = FindObject<UChooserTable>(nullptr, ChooserInfo.PathName))
 				{
-					if (UChooserTable* Chooser = Cast<UChooserTable>(ChooserObject))
+					const FObjectInfo& ContextObjectInfo = GameplayProvider->GetObjectInfo(OwnerId);
+					
+					// add to recent context objects list, so that this object is selectable as a target in the chooser editor
+                 	Chooser->AddRecentContextObject(ContextObjectInfo.Name);
+					
+					UChooserTable* ContextOwner = Chooser->GetContextOwner();
+					if (ContextOwner->HasDebugTarget())
 					{
-						UChooserTable* ContextOwner = Chooser->GetContextOwner();
-						if (ContextOwner->HasDebugTarget())
+						if (ContextOwner->GetDebugTargetName() == ContextObjectInfo.Name)
 						{
-							if (UObject* ContextObject = FObjectTrace::GetObjectFromId(OwnerId))
+							Chooser->SetDebugSelectedRow(ChooserEvaluationData.SelectedIndex);
+							Chooser->bDebugTestValuesValid = true;
+
+							ChooserProvider->ReadChooserValueTimeline(OwnerId, [StartTime, EndTime, Chooser, &ChooserEvaluationData](const FChooserProvider::ChooserValueTimeline& ValueTimeline)
 							{
-								if (ContextOwner->GetDebugTarget() == ContextObject)
+								ValueTimeline.EnumerateEvents(StartTime,EndTime,[Chooser, ChooserEvaluationData](double StartTime, double EndTime, uint32 Depth, const FChooserValueData& ValueData)
 								{
-									Chooser->SetDebugSelectedRow(ChooserEvaluationData.SelectedIndex);
-									Chooser->bDebugTestValuesValid = true;
-
-									ChooserProvider->ReadChooserValueTimeline(OwnerId, [StartTime, EndTime, Chooser, &ChooserEvaluationData](const FChooserProvider::ChooserValueTimeline& ValueTimeline)
+									if (ChooserEvaluationData.ChooserId == ValueData.ChooserId)
 									{
-										ValueTimeline.EnumerateEvents(StartTime,EndTime,[Chooser, ChooserEvaluationData](double StartTime, double EndTime, uint32 Depth, const FChooserValueData& ValueData)
+										for(FInstancedStruct& ColumnStruct : Chooser->ColumnsStructs)
 										{
-											if (ChooserEvaluationData.ChooserId == ValueData.ChooserId)
+											FChooserColumnBase& Column = ColumnStruct.GetMutable<FChooserColumnBase>();
+											if(ValueData.Key == Column.GetInputValue()->GetDebugName())
 											{
-												for(FInstancedStruct& ColumnStruct : Chooser->ColumnsStructs)
-												{
-													FChooserColumnBase& Column = ColumnStruct.GetMutable<FChooserColumnBase>();
-													if(ValueData.Key == Column.GetInputValue()->GetDebugName())
-													{
-														Column.SetTestValue(ValueData.Value);
-													}
-												}
+												Column.SetTestValue(ValueData.Value);
 											}
-											return TraceServices::EEventEnumerate::Continue;
-										});
-									});
-								}
-							}
+										}
+									}
+									return TraceServices::EEventEnumerate::Continue;
+								});
+							});
 						}
 					}
 				}
