ChooserTableEditor code cleanup.  Moving various classes out of ChooserTableEditor.cpp in to separate files

diff --git a/Engine/Plugins/Experimental/Chooser/Source/ChooserEditor/Private/ChooserDetails.cpp b/Engine/Plugins/Experimental/Chooser/Source/ChooserEditor/Private/ChooserDetails.cpp
new file mode 100644
--- /dev/null
+++ b/Engine/Plugins/Experimental/Chooser/Source/ChooserEditor/Private/ChooserDetails.cpp
@@ -0,0 +1,119 @@
+// Copyright Epic Games, Inc. All Rights Reserved.
+
+#include "ChooserDetails.h"
+
+#include "ChooserTableEditor.h"
+
+#include "Chooser.h"
+#include "ChooserEditorWidgets.h"
+#include "DetailCategoryBuilder.h"
+#include "DetailLayoutBuilder.h"
+#include "IDetailsView.h"
+#include "PropertyCustomizationHelpers.h"
+#include "PropertyEditorModule.h"
+#include "ScopedTransaction.h"
+
+#define LOCTEXT_NAMESPACE "ChooserDetails"
+
+namespace UE::ChooserEditor
+{
+	
+void FChooserDetails::CustomizeDetails(IDetailLayoutBuilder& DetailBuilder)
+{
+	TArray<TWeakObjectPtr<UObject>> Objects;
+	DetailBuilder.GetObjectsBeingCustomized(Objects);
+	check(!Objects.IsEmpty());
+
+	UChooserTable* Chooser = Cast<UChooserTable>(Objects[0]);
+	
+	IDetailCategoryBuilder& HiddenCategory = DetailBuilder.EditCategory(TEXT("Hidden"));
+
+	TArray<TSharedRef<IPropertyHandle>> HiddenProperties;
+	HiddenCategory.GetDefaultProperties(HiddenProperties);
+	for(TSharedRef<IPropertyHandle>& PropertyHandle :  HiddenProperties)
+	{
+		// these (Results and Columns arrays) need to be hidden when showing the root ChooserTable properties
+		// but still need to be EditAnywhere so that the Properties exist for displaying when you select a row or column (eg by FChooserRowDetails below)
+		PropertyHandle->MarkHiddenByCustomization();
+	}
+}
+
+// Make the details panel show the values for the selected row, showing each column value
+void FChooserRowDetails::CustomizeDetails(IDetailLayoutBuilder& DetailBuilder)
+{
+	TArray<TWeakObjectPtr<UObject>> Objects;
+	DetailBuilder.GetObjectsBeingCustomized(Objects);
+	check(!Objects.IsEmpty());
+
+	UChooserRowDetails* Row = Cast<UChooserRowDetails>(Objects[0]);
+	UChooserTable* Chooser = Row->Chooser;
+	
+	if (Chooser->ResultsStructs.IsValidIndex(Row->Row))
+	{
+		IDetailCategoryBuilder& PropertiesCategory = DetailBuilder.EditCategory("Row Properties");
+
+		TSharedPtr<IPropertyHandle> ChooserProperty = DetailBuilder.GetProperty("Chooser", Row->StaticClass());
+		DetailBuilder.HideProperty(ChooserProperty);
+	
+		TSharedPtr<IPropertyHandle> ResultsArrayProperty = ChooserProperty->GetChildHandle("ResultsStructs");
+		TSharedPtr<IPropertyHandle> CurrentResultProperty = ResultsArrayProperty->AsArray()->GetElement(Row->Row);
+		IDetailPropertyRow& NewResultProperty = PropertiesCategory.AddProperty(CurrentResultProperty);
+		NewResultProperty.DisplayName(LOCTEXT("ResultColumnName","Result"));
+		NewResultProperty.ShowPropertyButtons(false); // hide array add button
+		NewResultProperty.ShouldAutoExpand(true);
+	
+		for(int ColumnIndex=0; ColumnIndex<Chooser->ColumnsStructs.Num(); ColumnIndex++)
+		{
+			FChooserColumnBase& Column = Chooser->ColumnsStructs[ColumnIndex].GetMutable<FChooserColumnBase>();
+			TSharedRef<FStructOnScope> StructOnScope = MakeShared<FStructOnScope>(Chooser->ColumnsStructs[ColumnIndex].GetScriptStruct(), reinterpret_cast<uint8*>(&Column));
+			TSharedPtr<IPropertyHandle> ColumnDataProperty = DetailBuilder.AddStructurePropertyData({StructOnScope}, Column.RowValuesPropertyName());
+			uint32 NumElements = 0;
+			ColumnDataProperty->AsArray()->GetNumElements(NumElements);
+			if (Row->Row < (int)NumElements)
+			{
+				TSharedRef<IPropertyHandle> CellData = ColumnDataProperty->AsArray()->GetElement(Row->Row);
+	
+				IDetailPropertyRow& NewColumnProperty = PropertiesCategory.AddProperty(CellData);
+				FText DisplayName = LOCTEXT("No Input Value", "No Input Value");
+				if (FChooserParameterBase* InputValue = Column.GetInputValue())
+				{
+					InputValue->GetDisplayName(DisplayName); 
+				}
+				NewColumnProperty.DisplayName(DisplayName);
+				NewColumnProperty.ShowPropertyButtons(false); // hide array add button
+				NewColumnProperty.ShouldAutoExpand(true);
+			}
+		}
+	}
+}
+
+// Make the details panel show the values for the selected row, showing each column value
+void FChooserColumnDetails::CustomizeDetails(IDetailLayoutBuilder& DetailBuilder)
+{
+	TArray<TWeakObjectPtr<UObject>> Objects;
+	DetailBuilder.GetObjectsBeingCustomized(Objects);
+	check(!Objects.IsEmpty());
+
+	UChooserColumnDetails* Column = Cast<UChooserColumnDetails>(Objects[0]);
+	UChooserTable* Chooser = Column->Chooser;
+	
+	if (Chooser->ColumnsStructs.IsValidIndex(Column->Column))
+	{
+		IDetailCategoryBuilder& PropertiesCategory = DetailBuilder.EditCategory("Column Properties");
+
+		TSharedPtr<IPropertyHandle> ChooserProperty = DetailBuilder.GetProperty("Chooser", Column->StaticClass());
+		DetailBuilder.HideProperty(ChooserProperty);
+	
+		TSharedPtr<IPropertyHandle> ColumnsArrayProperty = ChooserProperty->GetChildHandle("ColumnsStructs");
+		TSharedPtr<IPropertyHandle> CurrentColumnProperty = ColumnsArrayProperty->AsArray()->GetElement(Column->Column);
+		IDetailPropertyRow& NewColumnProperty = PropertiesCategory.AddProperty(CurrentColumnProperty);
+		NewColumnProperty.DisplayName(LOCTEXT("Selected Column","Selected Column"));
+		NewColumnProperty.ShowPropertyButtons(false); // hide array add button
+		NewColumnProperty.ShouldAutoExpand(true);
+	}
+}
+
+
+}
+
+#undef LOCTEXT_NAMESPACE
diff --git a/Engine/Plugins/Experimental/Chooser/Source/ChooserEditor/Private/ChooserDetails.h b/Engine/Plugins/Experimental/Chooser/Source/ChooserEditor/Private/ChooserDetails.h
new file mode 100644
--- /dev/null
+++ b/Engine/Plugins/Experimental/Chooser/Source/ChooserEditor/Private/ChooserDetails.h
@@ -0,0 +1,71 @@
+// Copyright Epic Games, Inc. All Rights Reserved.
+
+#pragma once
+
+#include "CoreMinimal.h"
+#include "IDetailCustomization.h"
+#include "ChooserDetails.generated.h"
+
+class UChooserTable;
+
+// Class used for chooser editor details customization
+UCLASS()
+class UChooserRowDetails : public UObject
+{
+	GENERATED_BODY()
+public:
+	UPROPERTY(EditAnywhere, Instanced, Category="Hidden")
+	TObjectPtr<UChooserTable> Chooser;
+	int Row;
+};
+
+namespace UE::ChooserEditor
+{
+	
+	class FChooserDetails : public IDetailCustomization
+	{
+	public:
+		FChooserDetails() {};
+		virtual ~FChooserDetails() override {};
+
+		static TSharedRef<IDetailCustomization> MakeInstance()
+		{
+			return MakeShareable( new FChooserDetails() );
+		}
+
+		// IDetailCustomization interface
+		virtual void CustomizeDetails(class IDetailLayoutBuilder& DetailBuilder) override;
+	};
+		
+	class FChooserRowDetails : public IDetailCustomization
+	{
+	public:
+		FChooserRowDetails() {};
+		virtual ~FChooserRowDetails() override {};
+
+		static TSharedRef<IDetailCustomization> MakeInstance()
+		{
+			return MakeShareable( new FChooserRowDetails() );
+		}
+
+		// IDetailCustomization interface
+		virtual void CustomizeDetails(class IDetailLayoutBuilder& DetailBuilder) override;
+	};
+	
+	class FChooserColumnDetails : public IDetailCustomization
+    {
+    public:
+    	FChooserColumnDetails() {};
+    	virtual ~FChooserColumnDetails() override {};
+    
+    	static TSharedRef<IDetailCustomization> MakeInstance()
+    	{
+    		return MakeShareable( new FChooserColumnDetails() );
+    	}
+    
+    	// IDetailCustomization interface
+    	virtual void CustomizeDetails(class IDetailLayoutBuilder& DetailBuilder) override;
+    };
+
+	
+}
\ No newline at end of file
diff --git a/Engine/Plugins/Experimental/Chooser/Source/ChooserEditor/Private/ChooserEditorWidgets.cpp b/Engine/Plugins/Experimental/Chooser/Source/ChooserEditor/Private/ChooserEditorWidgets.cpp
new file mode 100644
--- /dev/null
+++ b/Engine/Plugins/Experimental/Chooser/Source/ChooserEditor/Private/ChooserEditorWidgets.cpp
@@ -0,0 +1,252 @@
+// Copyright Epic Games, Inc. All Rights Reserved.
+
+#include "ChooserEditorWidgets.h"
+
+#include "Chooser.h"
+#include "ChooserEditorStyle.h"
+#include "ChooserTableEditor.h"
+#include "DetailLayoutBuilder.h"
+#include "ObjectChooserClassFilter.h"
+#include "ObjectChooserWidgetFactories.h"
+#include "ObjectChooser_Asset.h"
+#include "ObjectChooser_Class.h"
+#include "PropertyCustomizationHelpers.h"
+#include "SClassViewer.h"
+#include "ScopedTransaction.h"
+#include "Editor/EditorEngine.h"
+#include "Framework/Application/SlateApplication.h"
+#include "Framework/MultiBox/MultiBoxBuilder.h"
+#include "Widgets/Input/SButton.h"
+#include "Widgets/Input/SComboButton.h"
+#include "Widgets/Input/STextEntryPopup.h"
+
+#define LOCTEXT_NAMESPACE "ChooserEditorWidgets"
+
+namespace UE::ChooserEditor
+{
+
+TSharedRef<SWidget> CreateAssetWidget(bool bReadOnly, UObject* TransactionObject, void* Value, UClass* ResultBaseClass, FChooserWidgetValueChanged ValueChanged)
+{
+	FAssetChooser* DIAsset = static_cast<FAssetChooser*>(Value);
+
+	UObject* Asset = DIAsset->Asset;
+	
+	return SNew(SObjectPropertyEntryBox)
+		.IsEnabled(!bReadOnly)
+		.AllowedClass(ResultBaseClass ? ResultBaseClass : UObject::StaticClass())
+		.ObjectPath_Lambda([DIAsset](){ return DIAsset->Asset ? DIAsset->Asset.GetPath() : "";})
+		.OnObjectChanged_Lambda([TransactionObject, DIAsset](const FAssetData& AssetData)
+		{
+			const FScopedTransaction Transaction(LOCTEXT("Edit Asset", "Edit Asset"));
+			TransactionObject->Modify(true);
+			DIAsset->Asset = AssetData.GetAsset();
+		});
+}
+	
+TSharedRef<SWidget> CreateClassWidget(bool bReadOnly, UObject* TransactionObject, void* Value, UClass* ResultBaseClass, FChooserWidgetValueChanged ValueChanged)
+{
+	FClassChooser* ClassChooser = static_cast<FClassChooser*>(Value);
+
+	UClass* Class = ClassChooser->Class;
+
+	return SNew(SClassPropertyEntryBox)
+			.IsEnabled(!bReadOnly)
+			.MetaClass(ResultBaseClass ? ResultBaseClass : UObject::StaticClass())
+			.SelectedClass_Lambda([ClassChooser]()
+			{
+				return ClassChooser->Class;
+			})
+			.OnSetClass_Lambda([TransactionObject, ClassChooser](const UClass* SelectedClass)
+			{
+				const FScopedTransaction Transaction(LOCTEXT("Edit Class", "Edit Class"));
+				TransactionObject->Modify(true);
+				ClassChooser->Class = const_cast<UClass*>(SelectedClass);
+			});
+}
+
+TSharedRef<SWidget> CreateEvaluateChooserWidget(bool bReadOnly, UObject* TransactionObject, void* Value, UClass* ResultBaseClass, FChooserWidgetValueChanged ValueChanged)
+{
+	FEvaluateChooser* EvaluateChooser = static_cast<FEvaluateChooser*>(Value);
+	
+	return SNew(SObjectPropertyEntryBox)
+		.IsEnabled(!bReadOnly)
+		.AllowedClass(UChooserTable::StaticClass())
+		.ObjectPath_Lambda([EvaluateChooser](){ return EvaluateChooser->Chooser ? EvaluateChooser->Chooser.GetPath() : "";})
+		.OnShouldFilterAsset_Lambda([ResultBaseClass](const FAssetData& InAssetData)
+		{
+			if (ResultBaseClass == nullptr)
+		 	{
+		 		return false;
+		 	}
+		 	if (InAssetData.IsInstanceOf(UChooserTable::StaticClass()))
+		 	{
+		 		if (UChooserTable* Chooser = Cast<UChooserTable>(InAssetData.GetAsset()))
+		 		{
+		 			UChooserTable* ContextOwner = Chooser->GetContextOwner();
+		 			return !(ContextOwner->OutputObjectType && ContextOwner->OutputObjectType->IsChildOf(ResultBaseClass));
+		 		}
+		 	}
+			return true;
+		})
+		.OnObjectChanged_Lambda([TransactionObject, EvaluateChooser](const FAssetData& AssetData)
+		{
+			const FScopedTransaction Transaction(LOCTEXT("Edit Chooser", "Edit Chooser"));
+			TransactionObject->Modify(true);
+			EvaluateChooser->Chooser = Cast<UChooserTable>(AssetData.GetAsset());
+		});
+}
+
+TSharedRef<SWidget> CreateNestedChooserWidget(bool bReadOnly, UObject* TransactionObject, void* Value, UClass* ResultBaseClass, FChooserWidgetValueChanged ValueChanged)
+{
+	FNestedChooser* NestedChooser = static_cast<FNestedChooser*>(Value);
+	
+	TSharedRef<SComboButton> Button = SNew(SComboButton)
+			.ContentPadding(0)
+			.ButtonContent()
+			[
+				SNew(SHorizontalBox)
+				+ SHorizontalBox::Slot().AutoWidth()
+				[
+					SNew(SImage).Image(FChooserEditorStyle::Get().GetBrush("ChooserEditor.ChooserTableIconSmall"))
+				]
+				+ SHorizontalBox::Slot().FillWidth(1.0).Padding(2)
+				[
+					SNew(STextBlock)
+					.Font(IDetailLayoutBuilder::GetDetailFont())
+					.Text_Lambda([NestedChooser]()
+					{
+						if (NestedChooser->Chooser)
+						{
+							return FText::FromString(NestedChooser->Chooser->GetName());
+						}
+						else
+						{
+							return LOCTEXT("None", "None");
+						}
+					})
+				]
+			];
+			
+	Button->SetOnGetMenuContent(
+		FOnGetContent::CreateLambda(
+			[ Button, NestedChooser, TransactionObject]()
+            			{
+							FMenuBuilder MenuBuilder(true, nullptr);
+                        
+                        	MenuBuilder.BeginSection(NAME_None, LOCTEXT("Nested Chooser", "Nested Chooser"));
+                        	{
+							MenuBuilder.AddMenuEntry(
+								LOCTEXT("None", "None"), 
+								LOCTEXT("NewNestedChooser_Tooltip", "Create a new Chooser Table embedded in this asset, and reference it"),
+								FSlateIcon(),
+								FUIAction(FExecuteAction::CreateLambda([Button, NestedChooser, TransactionObject]()
+								{
+									const FScopedTransaction Transaction(LOCTEXT("Clear Nested Chooser", "Clear Nested Chooser"));
+									TransactionObject->Modify(true);
+									Button->SetIsOpen(false);
+									NestedChooser->Chooser = nullptr;
+								})));
+								
+                        		MenuBuilder.AddMenuEntry(
+                        			LOCTEXT("New Nested Chooser", "New Nested Chooser"), 
+                        			LOCTEXT("NewNestedChooser_Tooltip", "Create a new Chooser Table embedded in this asset, and reference it"),
+                        			FSlateIcon(),
+                        			FUIAction(FExecuteAction::CreateLambda([Button, NestedChooser, TransactionObject]()
+                        			{
+                        				Button->SetIsOpen(false);
+                        				
+										TSharedRef<STextEntryPopup> TextEntry =
+											SNew(STextEntryPopup)
+											.Label(LOCTEXT("NewNestedChooserLabel", "New Chooser Name"))
+											.OnTextCommitted_Lambda([NestedChooser, TransactionObject](FText InText, ETextCommit::Type InCommitType)
+											{
+												const FScopedTransaction Transaction(LOCTEXT("Assign New Nested Chooser", "Assign New Nested Chooser"));
+												TransactionObject->Modify(true);
+											
+												FSlateApplication::Get().DismissAllMenus();
+												NestedChooser->Chooser = NewObject<UChooserTable>(TransactionObject->GetPackage(), UChooserTable::StaticClass(), FName(InText.ToString()));
+												NestedChooser->Chooser->ParentTable = Cast<UChooserTable>(TransactionObject);
+											});
+							
+										FSlateApplication& SlateApp = FSlateApplication::Get();
+										SlateApp.PushMenu(
+											SlateApp.GetInteractiveTopLevelWindows()[0],
+											FWidgetPath(),
+											TextEntry,
+											SlateApp.GetCursorPos(),
+											FPopupTransitionEffect::TypeInPopup
+											);
+                        			})));
+								
+								MenuBuilder.AddSubMenu(LOCTEXT("Select Existing", "Select Existing"),
+									LOCTEXT("SelectExisting_Tooltip", "Select Existing Embedded UChooserTable from this asset"),
+									FNewMenuDelegate::CreateLambda([Button, NestedChooser, TransactionObject](FMenuBuilder& SubMenuBuilder)
+									{
+										SubMenuBuilder.BeginSection("Existing", LOCTEXT("Existing Choosers", "Existing Choosers"));
+
+										TArray<UObject*> ObjectsInPackage;
+										GetObjectsWithOuter(TransactionObject->GetPackage(), ObjectsInPackage);
+
+										UObject* RootChooser = TransactionObject->GetPackage()->FindAssetInPackage();
+
+										for (UObject* Object : ObjectsInPackage)
+										{
+											 // exclude the main chooser 
+											if (Object != RootChooser)
+											{
+												if (UChooserTable* Chooser = Cast<UChooserTable>(Object))
+												{
+													SubMenuBuilder.AddMenuEntry( FText::FromString(Chooser->GetName()), LOCTEXT("AddExistingObjectTooltip", "Add a reference to this existing Chooser Table."), FSlateIcon(),
+														FUIAction(FExecuteAction::CreateLambda([Chooser, Button, NestedChooser, TransactionObject]()
+														{
+															const FScopedTransaction Transaction(LOCTEXT("Set Nested Chooser", "Set Nested Chooser"));
+															TransactionObject->Modify(true);
+															Button->SetIsOpen(false);
+															NestedChooser->Chooser = Chooser;
+														})));
+												}
+											}
+										}
+										
+
+									}));
+                                     
+									
+									
+                        	}
+
+							return MenuBuilder.MakeWidget();
+
+            			})
+		);
+		
+	return SNew(SHorizontalBox)
+		+ SHorizontalBox::Slot().FillWidth(1.0)
+		[
+			Button
+		]
+		+ SHorizontalBox::Slot().AutoWidth()
+		[
+			SNew(SButton)
+			.Text(LOCTEXT("Edit","Edit"))
+			.OnClicked_Lambda([NestedChooser, TransactionObject]()
+			{
+				if (NestedChooser->Chooser)
+				{
+					UObject* RootChooser = TransactionObject->GetPackage()->FindAssetInPackage();
+					IAssetEditorInstance* Editor = GEditor->GetEditorSubsystem<UAssetEditorSubsystem>()->FindEditorForAsset(RootChooser, true);
+					if (Editor && Editor->GetEditorName() == FChooserTableEditor::EditorName)
+					{
+						// currently nested tables are only editable in the Chooser Table Editor
+						FChooserTableEditor* ChooserEditor = static_cast<FChooserTableEditor*>(Editor);
+						ChooserEditor->PushChooserTableToEdit(NestedChooser->Chooser);
+					}
+				}
+				return FReply::Handled();
+			})
+		];
+}
+	
+}
+
+#undef LOCTEXT_NAMESPACE
diff --git a/Engine/Plugins/Experimental/Chooser/Source/ChooserEditor/Private/ChooserEditorWidgets.h b/Engine/Plugins/Experimental/Chooser/Source/ChooserEditor/Private/ChooserEditorWidgets.h
new file mode 100644
--- /dev/null
+++ b/Engine/Plugins/Experimental/Chooser/Source/ChooserEditor/Private/ChooserEditorWidgets.h
@@ -0,0 +1,14 @@
+// Copyright Epic Games, Inc. All Rights Reserved.
+#pragma once
+
+#include "ObjectChooserWidgetFactories.h"
+
+namespace UE::ChooserEditor
+{
+
+TSharedRef<SWidget> CreateAssetWidget(bool bReadOnly, UObject* TransactionObject, void* Value, UClass* ResultBaseClass, FChooserWidgetValueChanged ValueChanged);
+TSharedRef<SWidget> CreateClassWidget(bool bReadOnly, UObject* TransactionObject, void* Value, UClass* ResultBaseClass, FChooserWidgetValueChanged ValueChanged);
+TSharedRef<SWidget> CreateEvaluateChooserWidget(bool bReadOnly, UObject* TransactionObject, void* Value, UClass* ResultBaseClass, FChooserWidgetValueChanged ValueChanged);
+TSharedRef<SWidget> CreateNestedChooserWidget(bool bReadOnly, UObject* TransactionObject, void* Value, UClass* ResultBaseClass, FChooserWidgetValueChanged ValueChanged);
+	
+}
\ No newline at end of file
diff --git a/Engine/Plugins/Experimental/Chooser/Source/ChooserEditor/Private/ChooserTableEditor.cpp b/Engine/Plugins/Experimental/Chooser/Source/ChooserEditor/Private/ChooserTableEditor.cpp
--- a/Engine/Plugins/Experimental/Chooser/Source/ChooserEditor/Private/ChooserTableEditor.cpp
+++ b/Engine/Plugins/Experimental/Chooser/Source/ChooserEditor/Private/ChooserTableEditor.cpp
@@ -2,46 +2,42 @@
 
 #include "ChooserTableEditor.h"
 
-#include "Framework/Docking/TabManager.h"
-#include "Modules/ModuleManager.h"
-#include "Styling/AppStyle.h"
-#include "Widgets/Docking/SDockTab.h"
-#include "PropertyEditorModule.h"
-#include "IDetailsView.h"
-#include "Framework/MultiBox/MultiBoxBuilder.h"
-#include "Widgets/Input/SHyperlink.h"
-#include "Widgets/Input/SNumericEntryBox.h"
-#include "Widgets/Views/SListView.h"
-#include "Widgets/Input/SComboButton.h"
-#include "SAssetDropTarget.h"
-#include "SClassViewer.h"
-#include "StructViewerModule.h"
-#include "SourceCodeNavigation.h"
 #include "Chooser.h"
-#include "ChooserEditorStyle.h"
+#include "ChooserDetails.h"
+#include "ChooserEditorWidgets.h"
 #include "ChooserFindProperties.h"
-#include "DetailCategoryBuilder.h"
-#include "DetailLayoutBuilder.h"
+#include "ChooserTableEditorCommands.h"
 #include "ClassViewerFilter.h"
+#include "DetailCategoryBuilder.h"
+#include "GraphEditorSettings.h"
+#include "IDetailsView.h"
 #include "IPropertyAccessEditor.h"
 #include "LandscapeRender.h"
-#include "Widgets/Layout/SWidgetSwitcher.h"
-#include "DragAndDrop/DecoratedDragDropOp.h"
-#include "ObjectChooser_Asset.h"
-#include "ObjectChooser_Class.h"
 #include "ObjectChooserClassFilter.h"
 #include "ObjectChooserWidgetFactories.h"
-#include "GraphEditorSettings.h"
-#include "IDetailCustomization.h"
+#include "ObjectChooser_Asset.h"
+#include "ObjectChooser_Class.h"
 #include "PersonaModule.h"
 #include "PropertyCustomizationHelpers.h"
+#include "PropertyEditorModule.h"
+#include "SAssetDropTarget.h"
+#include "SClassViewer.h"
 #include "ScopedTransaction.h"
-#include "Widgets/Layout/SSeparator.h"
-#include "ChooserTableEditorCommands.h"
-#include "Animation/AnimBlueprintGeneratedClass.h"
-#include "Animation/AnimInstance.h"
-#include "Widgets/Input/STextEntryPopup.h"
+#include "SourceCodeNavigation.h"
+#include "StructViewerModule.h"
+#include "DragAndDrop/DecoratedDragDropOp.h"
+#include "Framework/Docking/TabManager.h"
+#include "Framework/MultiBox/MultiBoxBuilder.h"
+#include "Modules/ModuleManager.h"
+#include "Styling/AppStyle.h"
+#include "SChooserTableRow.h"
+#include "Widgets/Docking/SDockTab.h"
+#include "Widgets/Input/SComboButton.h"
+#include "Widgets/Input/SHyperlink.h"
+#include "Widgets/Input/SNumericEntryBox.h"
 #include "Widgets/Layout/SScrollBox.h"
+#include "Widgets/Layout/SWidgetSwitcher.h"
+#include "Widgets/Views/SListView.h"
 
 #define LOCTEXT_NAMESPACE "ChooserEditor"
 
@@ -345,8 +341,6 @@ void FChooserTableEditor::InitEditor( const EToolkitMode::Type Mode, const TShar
 		
 	FAnimAssetFindReplaceConfig FindReplaceConfig;
 	FindReplaceConfig.InitialProcessorClass = UChooserFindProperties::StaticClass();
-
-	FPersonaModule& PersonaModule = FModuleManager::LoadModuleChecked<FPersonaModule>("Persona");
 }
 
 FName FChooserTableEditor::GetToolkitFName() const
@@ -541,342 +535,6 @@ TSharedRef<SDockTab> FChooserTableEditor::SpawnFindReplaceTab( const FSpawnTabAr
 	];
 }
 
-class FChooserRowDragDropOp : public FDecoratedDragDropOp
-{
-public:
-	DRAG_DROP_OPERATOR_TYPE(FWidgetTemplateDragDropOp, FDecoratedDragDropOp)
-
-	FChooserTableEditor* ChooserEditor;
-	uint32 RowIndex;
-
-	/** Constructs the drag drop operation */
-	static TSharedRef<FChooserRowDragDropOp> New(FChooserTableEditor* InEditor, uint32 InRowIndex)
-	{
-		TSharedRef<FChooserRowDragDropOp> Operation = MakeShareable(new FChooserRowDragDropOp());
-		Operation->ChooserEditor = InEditor;
-		Operation->RowIndex = InRowIndex;
-		Operation->DefaultHoverText = LOCTEXT("Chooser Row", "Chooser Row");
-		Operation->CurrentHoverText = Operation->DefaultHoverText;
-		
-		Operation->Construct();
-
-		return Operation;
-	};
-};
-
-class SChooserRowHandle : public SCompoundWidget
-{
-public:
-	SLATE_BEGIN_ARGS(SChooserRowHandle)
-	{}
-	SLATE_DEFAULT_SLOT(FArguments, Content)
-	SLATE_ARGUMENT(FChooserTableEditor*, ChooserEditor)
-	SLATE_ARGUMENT(uint32, RowIndex)
-	SLATE_END_ARGS()
-
-	void Construct(const FArguments& InArgs)
-	{
-		ChooserEditor = InArgs._ChooserEditor;
-		RowIndex = InArgs._RowIndex;
-
-		ChildSlot
-		[
-			SNew(SOverlay)
-			+ SOverlay::Slot()
-			[
-				SNew(SBox) .Padding(0.0f) .HAlign(HAlign_Center) .VAlign(VAlign_Center) .WidthOverride(16.0f)
-				[
-					SNew(SImage)
-					.Image(FCoreStyle::Get().GetBrush("VerticalBoxDragIndicatorShort"))
-				]
-			]
-			+ SOverlay::Slot()
-			[
-				SNew(SBox).Padding(0.0f) .HAlign(HAlign_Center) .VAlign(VAlign_Center) .WidthOverride(16.0f)
-				[
-					SNew(SImage)
-					.Visibility_Lambda([this]()
-					{
-						return ChooserEditor->GetChooser()->bDebugTestValuesValid && RowIndex == ChooserEditor->GetChooser()->GetDebugSelectedRow() ? EVisibility::HitTestInvisible : EVisibility::Hidden;
-					})
-					.Image(FAppStyle::Get().GetBrush("Icons.ArrowRight"))
-				]
-			]
-		];
-	}
-	
-	virtual FReply OnMouseButtonDoubleClick(const FGeometry& InMyGeometry, const FPointerEvent& InMouseEvent) override
-	{
-		UChooserTable* Chooser = ChooserEditor->GetChooser();
-
-		if (const FNestedChooser* NestedChooserResult = Chooser->ResultsStructs[RowIndex].GetPtr<FNestedChooser>())
-		{
-			if (NestedChooserResult->Chooser)
-			{
-				ChooserEditor->PushChooserTableToEdit(NestedChooserResult->Chooser);
-			}
-		}
-		return FReply::Handled();
-	}
-
-	FReply OnMouseButtonDown(const FGeometry& MyGeometry, const FPointerEvent& MouseEvent) override
-	{
-		ChooserEditor->SelectRow(RowIndex);
-		return FReply::Handled().DetectDrag(SharedThis(this), EKeys::LeftMouseButton);
-	};
-
-	FReply OnDragDetected(const FGeometry& MyGeometry, const FPointerEvent& MouseEvent) override
-	{
-		// clear row selection so that delete key can't cause the selected row to be deleted
-		ChooserEditor->ClearSelectedRows();
-		
-		TSharedRef<FChooserRowDragDropOp> DragDropOp = FChooserRowDragDropOp::New(ChooserEditor, RowIndex);
-		return FReply::Handled().BeginDragDrop(DragDropOp);
-	}
-
-private:
-	FChooserTableEditor* ChooserEditor = nullptr;
-	uint32 RowIndex;
-};
-
-
-class SChooserTableRow : public SMultiColumnTableRow<TSharedPtr<FChooserTableEditor::FChooserTableRow>>
-{
-public:
-	SLATE_BEGIN_ARGS(SChooserTableRow) {}
-	/** The list item for this row */
-	SLATE_ARGUMENT(TSharedPtr<FChooserTableEditor::FChooserTableRow>, Entry)
-	SLATE_ARGUMENT(UChooserTable*, Chooser)
-	SLATE_ARGUMENT(FChooserTableEditor*, Editor)
-	SLATE_END_ARGS()
-
-	enum { SpecialIndex_AddRow = -1, SpecialIndex_Fallback = -2 };
-
-	void Construct(const FArguments& Args, const TSharedRef<STableViewBase>& OwnerTableView)
-	{
-		RowIndex = Args._Entry;
-		Chooser = Args._Chooser;
-		Editor = Args._Editor;
-
-		SMultiColumnTableRow<TSharedPtr<FChooserTableEditor::FChooserTableRow>>::Construct(
-			FSuperRowType::FArguments(),
-			OwnerTableView
-		);
-	}
-
-	/** Overridden from SMultiColumnTableRow.  Generates a widget for this column of the list view. */
-	virtual TSharedRef<SWidget> GenerateWidgetForColumn(const FName& ColumnName) override
-	{
-		static FName Result = "Result";
-		static FName Handles = "Handles";
-	
-		if (Chooser->ResultsStructs.IsValidIndex(RowIndex->RowIndex))
-		{
-			if (ColumnName == Handles)
-			{
-				// row drag handle
-				return SNew(SChooserRowHandle).ChooserEditor(Editor).RowIndex(RowIndex->RowIndex);
-			}
-			else if (ColumnName == Result) 
-			{
-				UChooserTable* ContextOwner = Chooser->GetContextOwner();
-				TSharedPtr<SWidget> ResultWidget = FObjectChooserWidgetFactories::CreateWidget(false, ContextOwner, FObjectChooserBase::StaticStruct(), Chooser->ResultsStructs[RowIndex->RowIndex].GetMutableMemory(), Chooser->ResultsStructs[RowIndex->RowIndex].GetScriptStruct(), ContextOwner->OutputObjectType,
-				FOnStructPicked::CreateLambda([this, RowIndex=RowIndex->RowIndex](const UScriptStruct* ChosenStruct)
-				{
-					UChooserTable* ContextOwner = Chooser->GetContextOwner();
-					const FScopedTransaction Transaction(LOCTEXT("Change Row Result Type", "Change Row Result Type"));
-					Chooser->Modify(true);
-					Chooser->ResultsStructs[RowIndex].InitializeAs(ChosenStruct);
-					FObjectChooserWidgetFactories::CreateWidget(false, ContextOwner, FObjectChooserBase::StaticStruct(), Chooser->ResultsStructs[RowIndex].GetMutableMemory(), ChosenStruct, ContextOwner->OutputObjectType, FOnStructPicked(), &CacheBorder);
-				}),
-				&CacheBorder
-				);
-			
-				return SNew(SOverlay)
-						+ SOverlay::Slot()
-						[
-							ResultWidget.ToSharedRef()
-						]
-						+ SOverlay::Slot().VAlign(VAlign_Bottom)
-						[
-							SNew(SSeparator).SeparatorImage(FCoreStyle::Get().GetBrush("FocusRectangle"))
-							.Visibility_Lambda([this]() { return bDragActive && !bDropAbove ? EVisibility::Visible : EVisibility::Hidden; })
-						]
-						+ SOverlay::Slot().VAlign(VAlign_Top)
-						[
-							SNew(SSeparator).SeparatorImage(FCoreStyle::Get().GetBrush("FocusRectangle"))
-							.Visibility_Lambda([this]() { return bDragActive && bDropAbove ? EVisibility::Visible : EVisibility::Hidden; })
-						];
-			}
-			else
-			{
-				const int ColumnIndex = ColumnName.GetNumber() - 1;
-				if (ColumnIndex < Chooser->ColumnsStructs.Num() && ColumnIndex >=0)
-				{
-					FChooserColumnBase* Column = &Chooser->ColumnsStructs[ColumnIndex].GetMutable<FChooserColumnBase>();
-					const UStruct * ColumnStruct = Chooser->ColumnsStructs[ColumnIndex].GetScriptStruct();
-
-					TSharedPtr<SWidget> ColumnWidget = FObjectChooserWidgetFactories::CreateColumnWidget(Column, ColumnStruct, Chooser->GetContextOwner(), RowIndex->RowIndex);
-				
-					if (ColumnWidget.IsValid())
-					{
-						return SNew(SOverlay)
-						+ SOverlay::Slot()
-						[
-							ColumnWidget.ToSharedRef()
-						]
-						+ SOverlay::Slot()
-						[
-							SNew(SColorBlock).Visibility(EVisibility::HitTestInvisible).Color_Lambda(
-									[this,Column]()
-									{
-										if (Chooser->bDebugTestValuesValid && Column->HasFilters())
-										{
-											if (Column->EditorTestFilter(RowIndex->RowIndex))
-											{
-												return FLinearColor(0.0,1.0,0.0,0.30);
-											}
-											else
-											{
-												return FLinearColor(1.0,0.0,0.0,0.20);
-											}
-										}
-										return FLinearColor::Transparent;
-									})
-						];
-					}
-				}
-			}
-		}
-		else if (RowIndex->RowIndex == SpecialIndex_Fallback)
-		{
-			if (ColumnName == Handles)
-			{
-				return SNew(SBox).Padding(0.0f) .HAlign(HAlign_Center) .VAlign(VAlign_Center) .WidthOverride(16.0f)
-					[
-						SNew(SImage)
-						.Image(FChooserEditorStyle::Get().GetBrush("ChooserEditor.FallbackIcon"))
-						.ToolTipText(LOCTEXT("FallbackTooltip","Fallback result:  Returned if all rows failed."))
-					];
-			}
-			else if (ColumnName == Result) 
-			{
-				UChooserTable* ContextOwner = Chooser->GetContextOwner();
-				TSharedPtr<SWidget> ResultWidget = FObjectChooserWidgetFactories::CreateWidget(false,ContextOwner, FObjectChooserBase::StaticStruct(), Chooser->FallbackResult.GetMutableMemory(), Chooser->FallbackResult.GetScriptStruct(), ContextOwner->OutputObjectType,
-				FOnStructPicked::CreateLambda([this](const UScriptStruct* ChosenStruct)
-				{
-				UChooserTable* ContextOwner = Chooser->GetContextOwner();
-					const FScopedTransaction Transaction(LOCTEXT("Change Row Result Type", "Change Row Result Type"));
-					Chooser->Modify(true);
-					Chooser->FallbackResult.InitializeAs(ChosenStruct);
-					FObjectChooserWidgetFactories::CreateWidget(false, ContextOwner, FObjectChooserBase::StaticStruct(), Chooser->FallbackResult.GetMutableMemory(), ChosenStruct, ContextOwner->OutputObjectType, FOnStructPicked(), &CacheBorder
-						,FChooserWidgetValueChanged(), LOCTEXT("Fallback Result", "Fallback Result: (None)"));
-				}),
-				&CacheBorder
-				,FChooserWidgetValueChanged(), LOCTEXT("Fallback Result", "Fallback Result: (None)")
-				);
-				
-				return SNew(SOverlay)
-						+ SOverlay::Slot()
-						[
-							ResultWidget.ToSharedRef()
-						]
-						+ SOverlay::Slot().VAlign(VAlign_Top)
-						[
-							SNew(SSeparator).SeparatorImage(FCoreStyle::Get().GetBrush("FocusRectangle"))
-							.Visibility_Lambda([this]() { return bDragActive  ? EVisibility::Visible : EVisibility::Hidden; })
-						];
-			}
-			else
-			{
-				const int ColumnIndex = ColumnName.GetNumber() - 1;
-				if (ColumnIndex < Chooser->ColumnsStructs.Num() && ColumnIndex >=0)
-				{
-					FChooserColumnBase* Column = &Chooser->ColumnsStructs[ColumnIndex].GetMutable<FChooserColumnBase>();
-					const UStruct * ColumnStruct = Chooser->ColumnsStructs[ColumnIndex].GetScriptStruct();
-
-					TSharedPtr<SWidget> ColumnWidget = FObjectChooserWidgetFactories::CreateColumnWidget(Column, ColumnStruct, Chooser->GetContextOwner(), -2);
-				
-					if (ColumnWidget.IsValid())
-					{
-						return ColumnWidget.ToSharedRef();
-					}
-				}
-			}
-		}
-		else if (RowIndex->RowIndex == SpecialIndex_AddRow)
-		{
-			// on the row past the end, show an Add button in the result column
-			if (ColumnName == Result)
-			{
-				return Editor->GetCreateRowComboButton().ToSharedRef();
-			}
-		}
-
-		return SNullWidget::NullWidget;
-	}
-
-	virtual void OnDragEnter(const FGeometry& MyGeometry, const FDragDropEvent& DragDropEvent) override
-	{
-		if (TSharedPtr<FChooserRowDragDropOp> Operation = DragDropEvent.GetOperationAs<FChooserRowDragDropOp>())
-		{
-			bDragActive = true;
-			float Center = MyGeometry.Position.Y + MyGeometry.Size.Y;
-			bDropAbove = DragDropEvent.GetScreenSpacePosition().Y < Center;
-		}
-	}
-	virtual void OnDragLeave(const FDragDropEvent& DragDropEvent) override
-	{
-		bDragActive = false;
-	}
-
-	virtual FReply OnDragOver(const FGeometry& MyGeometry, const FDragDropEvent& DragDropEvent) override
-	{
-		if (TSharedPtr<FChooserRowDragDropOp> Operation = DragDropEvent.GetOperationAs<FChooserRowDragDropOp>())
-		{
-			float Center = MyGeometry.AbsolutePosition.Y + MyGeometry.Size.Y/2;
-			bDropAbove = DragDropEvent.GetScreenSpacePosition().Y < Center;
-			return FReply::Handled();
-		}
-		return FReply::Unhandled();
-	}
-
-	virtual FReply OnDrop(const FGeometry& MyGeometry, const FDragDropEvent& DragDropEvent) override
-	{
-		if (TSharedPtr<FChooserRowDragDropOp> Operation = DragDropEvent.GetOperationAs<FChooserRowDragDropOp>())
-		{
-			if (Chooser == Operation->ChooserEditor->GetChooser())
-			{
-				int NewRowIndex;
-				if (!Chooser->ResultsStructs.IsValidIndex(RowIndex->RowIndex))
-				{
-					// for special (negative) indices, move to the end
-					NewRowIndex = Editor->MoveRow(Operation->RowIndex, Chooser->ResultsStructs.Num());
-				}
-				else if (bDropAbove)
-				{
-					NewRowIndex = Editor->MoveRow(Operation->RowIndex, RowIndex->RowIndex);
-				}
-				else
-				{
-					NewRowIndex = Editor->MoveRow(Operation->RowIndex, RowIndex->RowIndex+1);
-				}
-				Editor->SelectRow(NewRowIndex);
-				
-				return FReply::Handled();		
-			}
-		}
-		return FReply::Unhandled();
-	}
-
-private:
-	TSharedPtr<FChooserTableEditor::FChooserTableRow> RowIndex;
-	UChooserTable* Chooser;
-	FChooserTableEditor* Editor;
-	TSharedPtr<SBorder> CacheBorder;
-	bool bDragActive = false;
-	bool bDropAbove = false;
-};
 
 
 TSharedRef<ITableRow> FChooserTableEditor::GenerateTableRow(TSharedPtr<FChooserTableRow> InItem, const TSharedRef<STableViewBase>& OwnerTable)
@@ -1483,365 +1141,6 @@ void FChooserTableEditor::DeleteColumn(int Index)
 	}
 }
 
-TSharedRef<SWidget> CreateAssetWidget(bool bReadOnly, UObject* TransactionObject, void* Value, UClass* ResultBaseClass, FChooserWidgetValueChanged ValueChanged)
-{
-	FAssetChooser* DIAsset = static_cast<FAssetChooser*>(Value);
-
-	UObject* Asset = DIAsset->Asset;
-	
-	return SNew(SObjectPropertyEntryBox)
-		.IsEnabled(!bReadOnly)
-		.AllowedClass(ResultBaseClass ? ResultBaseClass : UObject::StaticClass())
-		.ObjectPath_Lambda([DIAsset](){ return DIAsset->Asset ? DIAsset->Asset.GetPath() : "";})
-		.OnObjectChanged_Lambda([TransactionObject, DIAsset](const FAssetData& AssetData)
-		{
-			const FScopedTransaction Transaction(LOCTEXT("Edit Asset", "Edit Asset"));
-			TransactionObject->Modify(true);
-			DIAsset->Asset = AssetData.GetAsset();
-		});
-}
-	
-TSharedRef<SWidget> CreateClassWidget(bool bReadOnly, UObject* TransactionObject, void* Value, UClass* ResultBaseClass, FChooserWidgetValueChanged ValueChanged)
-{
-	FClassChooser* ClassChooser = static_cast<FClassChooser*>(Value);
-
-	UClass* Class = ClassChooser->Class;
-
-	return SNew(SClassPropertyEntryBox)
-			.IsEnabled(!bReadOnly)
-			.MetaClass(ResultBaseClass ? ResultBaseClass : UObject::StaticClass())
-			.SelectedClass_Lambda([ClassChooser]()
-			{
-				return ClassChooser->Class;
-			})
-			.OnSetClass_Lambda([TransactionObject, ClassChooser](const UClass* SelectedClass)
-			{
-				const FScopedTransaction Transaction(LOCTEXT("Edit Class", "Edit Class"));
-				TransactionObject->Modify(true);
-				ClassChooser->Class = const_cast<UClass*>(SelectedClass);
-			});
-}
-
-TSharedRef<SWidget> CreateEvaluateChooserWidget(bool bReadOnly, UObject* TransactionObject, void* Value, UClass* ResultBaseClass, FChooserWidgetValueChanged ValueChanged)
-{
-	FEvaluateChooser* EvaluateChooser = static_cast<FEvaluateChooser*>(Value);
-	
-	return SNew(SObjectPropertyEntryBox)
-		.IsEnabled(!bReadOnly)
-		.AllowedClass(UChooserTable::StaticClass())
-		.ObjectPath_Lambda([EvaluateChooser](){ return EvaluateChooser->Chooser ? EvaluateChooser->Chooser.GetPath() : "";})
-		.OnShouldFilterAsset_Lambda([ResultBaseClass](const FAssetData& InAssetData)
-		{
-			if (ResultBaseClass == nullptr)
-		 	{
-		 		return false;
-		 	}
-		 	if (InAssetData.IsInstanceOf(UChooserTable::StaticClass()))
-		 	{
-		 		if (UChooserTable* Chooser = Cast<UChooserTable>(InAssetData.GetAsset()))
-		 		{
-		 			UChooserTable* ContextOwner = Chooser->GetContextOwner();
-		 			return !(ContextOwner->OutputObjectType && ContextOwner->OutputObjectType->IsChildOf(ResultBaseClass));
-		 		}
-		 	}
-			return true;
-		})
-		.OnObjectChanged_Lambda([TransactionObject, EvaluateChooser](const FAssetData& AssetData)
-		{
-			const FScopedTransaction Transaction(LOCTEXT("Edit Chooser", "Edit Chooser"));
-			TransactionObject->Modify(true);
-			EvaluateChooser->Chooser = Cast<UChooserTable>(AssetData.GetAsset());
-		});
-}
-
-TSharedRef<SWidget> CreateNestedChooserWidget(bool bReadOnly, UObject* TransactionObject, void* Value, UClass* ResultBaseClass, FChooserWidgetValueChanged ValueChanged)
-{
-	FNestedChooser* NestedChooser = static_cast<FNestedChooser*>(Value);
-	
-	TSharedRef<SComboButton> Button = SNew(SComboButton)
-			.ContentPadding(0)
-			.ButtonContent()
-			[
-				SNew(SHorizontalBox)
-				+ SHorizontalBox::Slot().AutoWidth()
-				[
-					SNew(SImage).Image(FChooserEditorStyle::Get().GetBrush("ChooserEditor.ChooserTableIconSmall"))
-				]
-				+ SHorizontalBox::Slot().FillWidth(1.0).Padding(2)
-				[
-					SNew(STextBlock)
-					.Font(IDetailLayoutBuilder::GetDetailFont())
-					.Text_Lambda([NestedChooser]()
-					{
-						if (NestedChooser->Chooser)
-						{
-							return FText::FromString(NestedChooser->Chooser->GetName());
-						}
-						else
-						{
-							return LOCTEXT("None", "None");
-						}
-					})
-				]
-			];
-			
-	Button->SetOnGetMenuContent(
-		FOnGetContent::CreateLambda(
-			[ Button, NestedChooser, TransactionObject]()
-            			{
-							FMenuBuilder MenuBuilder(true, nullptr);
-                        
-                        	MenuBuilder.BeginSection(NAME_None, LOCTEXT("Nested Chooser", "Nested Chooser"));
-                        	{
-							MenuBuilder.AddMenuEntry(
-								LOCTEXT("None", "None"), 
-								LOCTEXT("NewNestedChooser_Tooltip", "Create a new Chooser Table embedded in this asset, and reference it"),
-								FSlateIcon(),
-								FUIAction(FExecuteAction::CreateLambda([Button, NestedChooser, TransactionObject]()
-								{
-									const FScopedTransaction Transaction(LOCTEXT("Clear Nested Chooser", "Clear Nested Chooser"));
-									TransactionObject->Modify(true);
-									Button->SetIsOpen(false);
-									NestedChooser->Chooser = nullptr;
-								})));
-								
-                        		MenuBuilder.AddMenuEntry(
-                        			LOCTEXT("New Nested Chooser", "New Nested Chooser"), 
-                        			LOCTEXT("NewNestedChooser_Tooltip", "Create a new Chooser Table embedded in this asset, and reference it"),
-                        			FSlateIcon(),
-                        			FUIAction(FExecuteAction::CreateLambda([Button, NestedChooser, TransactionObject]()
-                        			{
-                        				Button->SetIsOpen(false);
-                        				
-										TSharedRef<STextEntryPopup> TextEntry =
-											SNew(STextEntryPopup)
-											.Label(LOCTEXT("NewNestedChooserLabel", "New Chooser Name"))
-											.OnTextCommitted_Lambda([NestedChooser, TransactionObject](FText InText, ETextCommit::Type InCommitType)
-											{
-												const FScopedTransaction Transaction(LOCTEXT("Assign New Nested Chooser", "Assign New Nested Chooser"));
-												TransactionObject->Modify(true);
-											
-												FSlateApplication::Get().DismissAllMenus();
-												NestedChooser->Chooser = NewObject<UChooserTable>(TransactionObject->GetPackage(), UChooserTable::StaticClass(), FName(InText.ToString()));
-												NestedChooser->Chooser->ParentTable = Cast<UChooserTable>(TransactionObject);
-											});
-							
-										FSlateApplication& SlateApp = FSlateApplication::Get();
-										SlateApp.PushMenu(
-											SlateApp.GetInteractiveTopLevelWindows()[0],
-											FWidgetPath(),
-											TextEntry,
-											SlateApp.GetCursorPos(),
-											FPopupTransitionEffect::TypeInPopup
-											);
-                        			})));
-								
-								MenuBuilder.AddSubMenu(LOCTEXT("Select Existing", "Select Existing"),
-									LOCTEXT("SelectExisting_Tooltip", "Select Existing Embedded UChooserTable from this asset"),
-									FNewMenuDelegate::CreateLambda([Button, NestedChooser, TransactionObject](FMenuBuilder& SubMenuBuilder)
-									{
-										SubMenuBuilder.BeginSection("Existing", LOCTEXT("Existing Choosers", "Existing Choosers"));
-
-										TArray<UObject*> ObjectsInPackage;
-										GetObjectsWithOuter(TransactionObject->GetPackage(), ObjectsInPackage);
-
-										UObject* RootChooser = TransactionObject->GetPackage()->FindAssetInPackage();
-
-										for (UObject* Object : ObjectsInPackage)
-										{
-											 // exclude the main chooser 
-											if (Object != RootChooser)
-											{
-												if (UChooserTable* Chooser = Cast<UChooserTable>(Object))
-												{
-													SubMenuBuilder.AddMenuEntry( FText::FromString(Chooser->GetName()), LOCTEXT("AddExistingObjectTooltip", "Add a reference to this existing Chooser Table."), FSlateIcon(),
-														FUIAction(FExecuteAction::CreateLambda([Chooser, Button, NestedChooser, TransactionObject]()
-														{
-															const FScopedTransaction Transaction(LOCTEXT("Set Nested Chooser", "Set Nested Chooser"));
-															TransactionObject->Modify(true);
-															Button->SetIsOpen(false);
-															NestedChooser->Chooser = Chooser;
-														})));
-												}
-											}
-										}
-										
-
-									}));
-                                     
-									
-									
-                        	}
-
-							return MenuBuilder.MakeWidget();
-
-            			})
-		);
-		
-	return SNew(SHorizontalBox)
-		+ SHorizontalBox::Slot().FillWidth(1.0)
-		[
-			Button
-		]
-		+ SHorizontalBox::Slot().AutoWidth()
-		[
-			SNew(SButton)
-			.Text(LOCTEXT("Edit","Edit"))
-			.OnClicked_Lambda([NestedChooser, TransactionObject]()
-			{
-				if (NestedChooser->Chooser)
-				{
-					UObject* RootChooser = TransactionObject->GetPackage()->FindAssetInPackage();
-					IAssetEditorInstance* Editor = GEditor->GetEditorSubsystem<UAssetEditorSubsystem>()->FindEditorForAsset(RootChooser, true);
-					if (Editor && Editor->GetEditorName() == FChooserTableEditor::EditorName)
-					{
-						// currently nested tables are only editable in the Chooser Table Editor
-						FChooserTableEditor* ChooserEditor = static_cast<FChooserTableEditor*>(Editor);
-						ChooserEditor->PushChooserTableToEdit(NestedChooser->Chooser);
-					}
-				}
-				return FReply::Handled();
-			})
-		];
-}
-	
-class FChooserDetails : public IDetailCustomization
-{
-public:
-	FChooserDetails() {};
-	virtual ~FChooserDetails() override {};
-
-	static TSharedRef<IDetailCustomization> MakeInstance()
-	{
-		return MakeShareable( new FChooserDetails() );
-	}
-
-	// IDetailCustomization interface
-	virtual void CustomizeDetails(class IDetailLayoutBuilder& DetailBuilder) override;
-};
-		
-void FChooserDetails::CustomizeDetails(IDetailLayoutBuilder& DetailBuilder)
-{
-	TArray<TWeakObjectPtr<UObject>> Objects;
-	DetailBuilder.GetObjectsBeingCustomized(Objects);
-
-	UChooserTable* Chooser = Cast<UChooserTable>(Objects[0]);
-	
-	IDetailCategoryBuilder& HiddenCategory = DetailBuilder.EditCategory(TEXT("Hidden"));
-
-	TArray<TSharedRef<IPropertyHandle>> HiddenProperties;
-	HiddenCategory.GetDefaultProperties(HiddenProperties);
-	for(TSharedRef<IPropertyHandle>& PropertyHandle :  HiddenProperties)
-	{
-		// these (Results and Columns arrays) need to be hidden when showing the root ChooserTable properties
-		// but still need to be EditAnywhere so that the Properties exist for displaying when you select a row or column (eg by FChooserRowDetails below)
-		PropertyHandle->MarkHiddenByCustomization();
-	}
-}
-
-class FChooserRowDetails : public IDetailCustomization
-{
-public:
-	FChooserRowDetails() {};
-	virtual ~FChooserRowDetails() override {};
-
-	static TSharedRef<IDetailCustomization> MakeInstance()
-	{
-		return MakeShareable( new FChooserRowDetails() );
-	}
-
-	// IDetailCustomization interface
-	virtual void CustomizeDetails(class IDetailLayoutBuilder& DetailBuilder) override;
-};
-
-// Make the details panel show the values for the selected row, showing each column value
-void FChooserRowDetails::CustomizeDetails(IDetailLayoutBuilder& DetailBuilder)
-{
-	TArray<TWeakObjectPtr<UObject>> Objects;
-	DetailBuilder.GetObjectsBeingCustomized(Objects);
-
-	UChooserRowDetails* Row = Cast<UChooserRowDetails>(Objects[0]);
-	UChooserTable* Chooser = Row->Chooser;
-	
-	if (Chooser->ResultsStructs.IsValidIndex(Row->Row))
-	{
-		IDetailCategoryBuilder& PropertiesCategory = DetailBuilder.EditCategory("Row Properties");
-
-		TSharedPtr<IPropertyHandle> ChooserProperty = DetailBuilder.GetProperty("Chooser", Row->StaticClass());
-		DetailBuilder.HideProperty(ChooserProperty);
-	
-		TSharedPtr<IPropertyHandle> ResultsArrayProperty = ChooserProperty->GetChildHandle("ResultsStructs");
-		TSharedPtr<IPropertyHandle> CurrentResultProperty = ResultsArrayProperty->AsArray()->GetElement(Row->Row);
-		IDetailPropertyRow& NewResultProperty = PropertiesCategory.AddProperty(CurrentResultProperty);
-		NewResultProperty.DisplayName(LOCTEXT("ResultColumnName","Result"));
-		NewResultProperty.ShowPropertyButtons(false); // hide array add button
-		NewResultProperty.ShouldAutoExpand(true);
-	
-		for(int ColumnIndex=0; ColumnIndex<Chooser->ColumnsStructs.Num(); ColumnIndex++)
-		{
-			FChooserColumnBase& Column = Chooser->ColumnsStructs[ColumnIndex].GetMutable<FChooserColumnBase>();
-			TSharedRef<FStructOnScope> StructOnScope = MakeShared<FStructOnScope>(Chooser->ColumnsStructs[ColumnIndex].GetScriptStruct(), reinterpret_cast<uint8*>(&Column));
-			TSharedPtr<IPropertyHandle> ColumnDataProperty = DetailBuilder.AddStructurePropertyData({StructOnScope}, Column.RowValuesPropertyName());
-			uint32 NumElements = 0;
-			ColumnDataProperty->AsArray()->GetNumElements(NumElements);
-			if (Row->Row < (int)NumElements)
-			{
-				TSharedRef<IPropertyHandle> CellData = ColumnDataProperty->AsArray()->GetElement(Row->Row);
-	
-				IDetailPropertyRow& NewColumnProperty = PropertiesCategory.AddProperty(CellData);
-				FText DisplayName = LOCTEXT("No Input Value", "No Input Value");
-				if (FChooserParameterBase* InputValue = Column.GetInputValue())
-				{
-					InputValue->GetDisplayName(DisplayName); 
-				}
-				NewColumnProperty.DisplayName(DisplayName);
-				NewColumnProperty.ShowPropertyButtons(false); // hide array add button
-				NewColumnProperty.ShouldAutoExpand(true);
-			}
-		}
-	}
-}
-
-class FChooserColumnDetails : public IDetailCustomization
-{
-public:
-	FChooserColumnDetails() {};
-	virtual ~FChooserColumnDetails() override {};
-
-	static TSharedRef<IDetailCustomization> MakeInstance()
-	{
-		return MakeShareable( new FChooserColumnDetails() );
-	}
-
-	// IDetailCustomization interface
-	virtual void CustomizeDetails(class IDetailLayoutBuilder& DetailBuilder) override;
-};
-
-// Make the details panel show the values for the selected row, showing each column value
-void FChooserColumnDetails::CustomizeDetails(IDetailLayoutBuilder& DetailBuilder)
-{
-	TArray<TWeakObjectPtr<UObject>> Objects;
-	DetailBuilder.GetObjectsBeingCustomized(Objects);
-
-	UChooserColumnDetails* Column = Cast<UChooserColumnDetails>(Objects[0]);
-	UChooserTable* Chooser = Column->Chooser;
-	
-	if (Chooser->ColumnsStructs.IsValidIndex(Column->Column))
-	{
-		IDetailCategoryBuilder& PropertiesCategory = DetailBuilder.EditCategory("Column Properties");
-
-		TSharedPtr<IPropertyHandle> ChooserProperty = DetailBuilder.GetProperty("Chooser", Column->StaticClass());
-		DetailBuilder.HideProperty(ChooserProperty);
-	
-		TSharedPtr<IPropertyHandle> ColumnsArrayProperty = ChooserProperty->GetChildHandle("ColumnsStructs");
-		TSharedPtr<IPropertyHandle> CurrentColumnProperty = ColumnsArrayProperty->AsArray()->GetElement(Column->Column);
-		IDetailPropertyRow& NewColumnProperty = PropertiesCategory.AddProperty(CurrentColumnProperty);
-		NewColumnProperty.DisplayName(LOCTEXT("Selected Column","Selected Column"));
-		NewColumnProperty.ShowPropertyButtons(false); // hide array add button
-		NewColumnProperty.ShouldAutoExpand(true);
-	}
-}
-
 void FChooserTableEditor::RegisterWidgets()
 {
 	FObjectChooserWidgetFactories::RegisterWidgetCreator(FAssetChooser::StaticStruct(), CreateAssetWidget);
diff --git a/Engine/Plugins/Experimental/Chooser/Source/ChooserEditor/Private/ChooserTableEditor.h b/Engine/Plugins/Experimental/Chooser/Source/ChooserEditor/Private/ChooserTableEditor.h
--- a/Engine/Plugins/Experimental/Chooser/Source/ChooserEditor/Private/ChooserTableEditor.h
+++ b/Engine/Plugins/Experimental/Chooser/Source/ChooserEditor/Private/ChooserTableEditor.h
@@ -17,17 +17,7 @@
 class SComboButton;
 class SEditableText;
 class IDetailsView;
-
-// Class used for chooser editor details customization
-UCLASS()
-class UChooserRowDetails : public UObject
-{
-	GENERATED_BODY()
-public:
-	UPROPERTY(EditAnywhere, Instanced, Category="Hidden")
-	TObjectPtr<UChooserTable> Chooser;
-	int Row;
-};
+class UChooserRowDetails;
 
 // Class used for chooser editor details customization
 UCLASS()
@@ -43,6 +33,8 @@ public:
 
 namespace UE::ChooserEditor
 {
+	struct FChooserTableRow;
+	
 	class FChooserTableEditor : public FAssetEditorToolkit, public FSelfRegisteringEditorUndoClient, public FNotifyHook
 	{
 	public:
@@ -100,11 +92,7 @@ namespace UE::ChooserEditor
 		/** Can be used to disable the details view making it read-only */
 		void SetPropertyEditingEnabledDelegate(FIsPropertyEditingEnabled InPropertyEditingDelegate);
 	
-		struct FChooserTableRow
-		{
-			FChooserTableRow(int32 i) { RowIndex = i; }
-			int32 RowIndex;
-		};
+
 
 		void UpdateTableRows();
 		void SelectColumn(UChooserTable* Chooser, int Index);
diff --git a/Engine/Plugins/Experimental/Chooser/Source/ChooserEditor/Private/SChooserRowHandle.cpp b/Engine/Plugins/Experimental/Chooser/Source/ChooserEditor/Private/SChooserRowHandle.cpp
new file mode 100644
--- /dev/null
+++ b/Engine/Plugins/Experimental/Chooser/Source/ChooserEditor/Private/SChooserRowHandle.cpp
@@ -0,0 +1,86 @@
+// Copyright Epic Games, Inc. All Rights Reserved.
+
+#include "SChooserRowHandle.h"
+#include "ChooserTableEditor.h"
+
+#define LOCTEXT_NAMESPACE "ChooserRowHandle"
+
+namespace UE::ChooserEditor
+{
+
+	TSharedRef<FChooserRowDragDropOp> FChooserRowDragDropOp::New(FChooserTableEditor* InEditor, uint32 InRowIndex)
+	{
+		TSharedRef<FChooserRowDragDropOp> Operation = MakeShareable(new FChooserRowDragDropOp());
+		Operation->ChooserEditor = InEditor;
+		Operation->RowIndex = InRowIndex;
+		Operation->DefaultHoverText = LOCTEXT("Chooser Row", "Chooser Row");
+		Operation->CurrentHoverText = Operation->DefaultHoverText;
+		
+		Operation->Construct();
+
+		return Operation;
+	}
+	
+	void SChooserRowHandle::Construct(const FArguments& InArgs)
+	{
+		ChooserEditor = InArgs._ChooserEditor;
+		RowIndex = InArgs._RowIndex;
+
+		ChildSlot
+		[
+			SNew(SOverlay)
+			+ SOverlay::Slot()
+			[
+				SNew(SBox) .Padding(0.0f) .HAlign(HAlign_Center) .VAlign(VAlign_Center) .WidthOverride(16.0f)
+				[
+					SNew(SImage)
+					.Image(FCoreStyle::Get().GetBrush("VerticalBoxDragIndicatorShort"))
+				]
+			]
+			+ SOverlay::Slot()
+			[
+				SNew(SBox).Padding(0.0f) .HAlign(HAlign_Center) .VAlign(VAlign_Center) .WidthOverride(16.0f)
+				[
+					SNew(SImage)
+					.Visibility_Lambda([this]()
+					{
+						return ChooserEditor->GetChooser()->bDebugTestValuesValid && RowIndex == ChooserEditor->GetChooser()->GetDebugSelectedRow() ? EVisibility::HitTestInvisible : EVisibility::Hidden;
+					})
+					.Image(FAppStyle::Get().GetBrush("Icons.ArrowRight"))
+				]
+			]
+		];
+	}
+	
+	FReply SChooserRowHandle::OnMouseButtonDoubleClick(const FGeometry& InMyGeometry, const FPointerEvent& InMouseEvent)
+	{
+		UChooserTable* Chooser = ChooserEditor->GetChooser();
+
+		if (const FNestedChooser* NestedChooserResult = Chooser->ResultsStructs[RowIndex].GetPtr<FNestedChooser>())
+		{
+			if (NestedChooserResult->Chooser)
+			{
+				ChooserEditor->PushChooserTableToEdit(NestedChooserResult->Chooser);
+			}
+		}
+		return FReply::Handled();
+	}
+
+	FReply SChooserRowHandle::OnMouseButtonDown(const FGeometry& MyGeometry, const FPointerEvent& MouseEvent)
+	{
+		ChooserEditor->SelectRow(RowIndex);
+		return FReply::Handled().DetectDrag(SharedThis(this), EKeys::LeftMouseButton);
+	};
+
+	FReply SChooserRowHandle::OnDragDetected(const FGeometry& MyGeometry, const FPointerEvent& MouseEvent)
+	{
+		// clear row selection so that delete key can't cause the selected row to be deleted
+		ChooserEditor->ClearSelectedRows();
+		
+		TSharedRef<FChooserRowDragDropOp> DragDropOp = FChooserRowDragDropOp::New(ChooserEditor, RowIndex);
+		return FReply::Handled().BeginDragDrop(DragDropOp);
+	}
+
+}
+
+#undef LOCTEXT_NAMESPACE
diff --git a/Engine/Plugins/Experimental/Chooser/Source/ChooserEditor/Private/SChooserRowHandle.h b/Engine/Plugins/Experimental/Chooser/Source/ChooserEditor/Private/SChooserRowHandle.h
new file mode 100644
--- /dev/null
+++ b/Engine/Plugins/Experimental/Chooser/Source/ChooserEditor/Private/SChooserRowHandle.h
@@ -0,0 +1,46 @@
+// Copyright Epic Games, Inc. All Rights Reserved.
+#pragma once
+
+#include "DragAndDrop/DecoratedDragDropOp.h"
+
+namespace UE::ChooserEditor
+{
+
+class FChooserTableEditor;
+
+class FChooserRowDragDropOp : public FDecoratedDragDropOp
+{
+public:
+	DRAG_DROP_OPERATOR_TYPE(FWidgetTemplateDragDropOp, FDecoratedDragDropOp)
+
+	FChooserTableEditor* ChooserEditor;
+	uint32 RowIndex;
+
+	/** Constructs the drag drop operation */
+	static TSharedRef<FChooserRowDragDropOp> New(FChooserTableEditor* InEditor, uint32 InRowIndex);
+};
+
+class SChooserRowHandle : public SCompoundWidget
+{
+public:
+	SLATE_BEGIN_ARGS(SChooserRowHandle)
+	{}
+	SLATE_DEFAULT_SLOT(FArguments, Content)
+	SLATE_ARGUMENT(FChooserTableEditor*, ChooserEditor)
+	SLATE_ARGUMENT(uint32, RowIndex)
+	SLATE_END_ARGS()
+
+	void Construct(const FArguments& InArgs);
+	
+	virtual FReply OnMouseButtonDoubleClick(const FGeometry& InMyGeometry, const FPointerEvent& InMouseEvent) override;
+	FReply OnMouseButtonDown(const FGeometry& MyGeometry, const FPointerEvent& MouseEvent) override;
+	FReply OnDragDetected(const FGeometry& MyGeometry, const FPointerEvent& MouseEvent) override;
+
+private:
+	FChooserTableEditor* ChooserEditor = nullptr;
+	uint32 RowIndex = 0;
+};
+
+}
+
+#undef LOCTEXT_NAMESPACE
diff --git a/Engine/Plugins/Experimental/Chooser/Source/ChooserEditor/Private/SChooserTableRow.cpp b/Engine/Plugins/Experimental/Chooser/Source/ChooserEditor/Private/SChooserTableRow.cpp
new file mode 100644
--- /dev/null
+++ b/Engine/Plugins/Experimental/Chooser/Source/ChooserEditor/Private/SChooserTableRow.cpp
@@ -0,0 +1,241 @@
+// Copyright Epic Games, Inc. All Rights Reserved.
+
+#include "SChooserTableRow.h"
+#include "Chooser.h"
+#include "ChooserEditorStyle.h"
+#include "ChooserTableEditor.h"
+#include "IObjectChooser.h"
+#include "ObjectChooserWidgetFactories.h"
+#include "SChooserRowHandle.h"
+#include "ScopedTransaction.h"
+#include "Widgets/Colors/SColorBlock.h"
+#include "Widgets/Input/SComboButton.h"
+#include "Widgets/Views/STableRow.h"
+#include "Widgets/Layout/SSeparator.h"
+
+
+
+#define LOCTEXT_NAMESPACE "ChooserTableRow"
+
+namespace UE::ChooserEditor
+{
+	void SChooserTableRow::Construct(const FArguments& Args, const TSharedRef<STableViewBase>& OwnerTableView)
+	{
+		RowIndex = Args._Entry;
+		Chooser = Args._Chooser;
+		Editor = Args._Editor;
+
+		SMultiColumnTableRow<TSharedPtr<FChooserTableRow>>::Construct(
+			FSuperRowType::FArguments(),
+			OwnerTableView
+		);
+	}
+
+	/** Overridden from SMultiColumnTableRow.  Generates a widget for this column of the list view. */
+	TSharedRef<SWidget> SChooserTableRow::GenerateWidgetForColumn(const FName& ColumnName)
+	{
+		static FName Result = "Result";
+		static FName Handles = "Handles";
+	
+		if (Chooser->ResultsStructs.IsValidIndex(RowIndex->RowIndex))
+		{
+			if (ColumnName == Handles)
+			{
+				// row drag handle
+				return SNew(SChooserRowHandle).ChooserEditor(Editor).RowIndex(RowIndex->RowIndex);
+			}
+			else if (ColumnName == Result) 
+			{
+				UChooserTable* ContextOwner = Chooser->GetContextOwner();
+				TSharedPtr<SWidget> ResultWidget = FObjectChooserWidgetFactories::CreateWidget(false, ContextOwner, FObjectChooserBase::StaticStruct(), Chooser->ResultsStructs[RowIndex->RowIndex].GetMutableMemory(), Chooser->ResultsStructs[RowIndex->RowIndex].GetScriptStruct(), ContextOwner->OutputObjectType,
+				FOnStructPicked::CreateLambda([this, RowIndex=RowIndex->RowIndex](const UScriptStruct* ChosenStruct)
+				{
+					UChooserTable* ContextOwner = Chooser->GetContextOwner();
+					const FScopedTransaction Transaction(LOCTEXT("Change Row Result Type", "Change Row Result Type"));
+					Chooser->Modify(true);
+					Chooser->ResultsStructs[RowIndex].InitializeAs(ChosenStruct);
+					FObjectChooserWidgetFactories::CreateWidget(false, ContextOwner, FObjectChooserBase::StaticStruct(), Chooser->ResultsStructs[RowIndex].GetMutableMemory(), ChosenStruct, ContextOwner->OutputObjectType, FOnStructPicked(), &CacheBorder);
+				}),
+				&CacheBorder
+				);
+			
+				return SNew(SOverlay)
+						+ SOverlay::Slot()
+						[
+							ResultWidget.ToSharedRef()
+						]
+						+ SOverlay::Slot().VAlign(VAlign_Bottom)
+						[
+							SNew(SSeparator).SeparatorImage(FCoreStyle::Get().GetBrush("FocusRectangle"))
+							.Visibility_Lambda([this]() { return bDragActive && !bDropAbove ? EVisibility::Visible : EVisibility::Hidden; })
+						]
+						+ SOverlay::Slot().VAlign(VAlign_Top)
+						[
+							SNew(SSeparator).SeparatorImage(FCoreStyle::Get().GetBrush("FocusRectangle"))
+							.Visibility_Lambda([this]() { return bDragActive && bDropAbove ? EVisibility::Visible : EVisibility::Hidden; })
+						];
+			}
+			else
+			{
+				const int ColumnIndex = ColumnName.GetNumber() - 1;
+				if (ColumnIndex < Chooser->ColumnsStructs.Num() && ColumnIndex >=0)
+				{
+					FChooserColumnBase* Column = &Chooser->ColumnsStructs[ColumnIndex].GetMutable<FChooserColumnBase>();
+					const UStruct * ColumnStruct = Chooser->ColumnsStructs[ColumnIndex].GetScriptStruct();
+
+					TSharedPtr<SWidget> ColumnWidget = FObjectChooserWidgetFactories::CreateColumnWidget(Column, ColumnStruct, Chooser->GetContextOwner(), RowIndex->RowIndex);
+				
+					if (ColumnWidget.IsValid())
+					{
+						return SNew(SOverlay)
+						+ SOverlay::Slot()
+						[
+							ColumnWidget.ToSharedRef()
+						]
+						+ SOverlay::Slot()
+						[
+							SNew(SColorBlock).Visibility(EVisibility::HitTestInvisible).Color_Lambda(
+									[this,Column]()
+									{
+										if (Chooser->bDebugTestValuesValid && Column->HasFilters())
+										{
+											if (Column->EditorTestFilter(RowIndex->RowIndex))
+											{
+												return FLinearColor(0.0,1.0,0.0,0.30);
+											}
+											else
+											{
+												return FLinearColor(1.0,0.0,0.0,0.20);
+											}
+										}
+										return FLinearColor::Transparent;
+									})
+						];
+					}
+				}
+			}
+		}
+		else if (RowIndex->RowIndex == SpecialIndex_Fallback)
+		{
+			if (ColumnName == Handles)
+			{
+				return SNew(SBox).Padding(0.0f) .HAlign(HAlign_Center) .VAlign(VAlign_Center) .WidthOverride(16.0f)
+					[
+						SNew(SImage)
+						.Image(FChooserEditorStyle::Get().GetBrush("ChooserEditor.FallbackIcon"))
+						.ToolTipText(LOCTEXT("FallbackTooltip","Fallback result:  Returned if all rows failed."))
+					];
+			}
+			else if (ColumnName == Result) 
+			{
+				UChooserTable* ContextOwner = Chooser->GetContextOwner();
+				TSharedPtr<SWidget> ResultWidget = FObjectChooserWidgetFactories::CreateWidget(false,ContextOwner, FObjectChooserBase::StaticStruct(), Chooser->FallbackResult.GetMutableMemory(), Chooser->FallbackResult.GetScriptStruct(), ContextOwner->OutputObjectType,
+				FOnStructPicked::CreateLambda([this](const UScriptStruct* ChosenStruct)
+				{
+				UChooserTable* ContextOwner = Chooser->GetContextOwner();
+					const FScopedTransaction Transaction(LOCTEXT("Change Row Result Type", "Change Row Result Type"));
+					Chooser->Modify(true);
+					Chooser->FallbackResult.InitializeAs(ChosenStruct);
+					FObjectChooserWidgetFactories::CreateWidget(false, ContextOwner, FObjectChooserBase::StaticStruct(), Chooser->FallbackResult.GetMutableMemory(), ChosenStruct, ContextOwner->OutputObjectType, FOnStructPicked(), &CacheBorder
+						,FChooserWidgetValueChanged(), LOCTEXT("Fallback Result", "Fallback Result: (None)"));
+				}),
+				&CacheBorder
+				,FChooserWidgetValueChanged(), LOCTEXT("Fallback Result", "Fallback Result: (None)")
+				);
+				
+				return SNew(SOverlay)
+						+ SOverlay::Slot()
+						[
+							ResultWidget.ToSharedRef()
+						]
+						+ SOverlay::Slot().VAlign(VAlign_Top)
+						[
+							SNew(SSeparator).SeparatorImage(FCoreStyle::Get().GetBrush("FocusRectangle"))
+							.Visibility_Lambda([this]() { return bDragActive  ? EVisibility::Visible : EVisibility::Hidden; })
+						];
+			}
+			else
+			{
+				const int ColumnIndex = ColumnName.GetNumber() - 1;
+				if (ColumnIndex < Chooser->ColumnsStructs.Num() && ColumnIndex >=0)
+				{
+					FChooserColumnBase* Column = &Chooser->ColumnsStructs[ColumnIndex].GetMutable<FChooserColumnBase>();
+					const UStruct * ColumnStruct = Chooser->ColumnsStructs[ColumnIndex].GetScriptStruct();
+
+					TSharedPtr<SWidget> ColumnWidget = FObjectChooserWidgetFactories::CreateColumnWidget(Column, ColumnStruct, Chooser->GetContextOwner(), -2);
+				
+					if (ColumnWidget.IsValid())
+					{
+						return ColumnWidget.ToSharedRef();
+					}
+				}
+			}
+		}
+		else if (RowIndex->RowIndex == SpecialIndex_AddRow)
+		{
+			// on the row past the end, show an Add button in the result column
+			if (ColumnName == Result)
+			{
+				return Editor->GetCreateRowComboButton().ToSharedRef();
+			}
+		}
+
+		return SNullWidget::NullWidget;
+	}
+
+	void SChooserTableRow::OnDragEnter(const FGeometry& MyGeometry, const FDragDropEvent& DragDropEvent)
+	{
+		if (TSharedPtr<FChooserRowDragDropOp> Operation = DragDropEvent.GetOperationAs<FChooserRowDragDropOp>())
+		{
+			bDragActive = true;
+			float Center = MyGeometry.Position.Y + MyGeometry.Size.Y;
+			bDropAbove = DragDropEvent.GetScreenSpacePosition().Y < Center;
+		}
+	}
+	void SChooserTableRow::OnDragLeave(const FDragDropEvent& DragDropEvent)
+	{
+		bDragActive = false;
+	}
+
+	FReply SChooserTableRow::OnDragOver(const FGeometry& MyGeometry, const FDragDropEvent& DragDropEvent)
+	{
+		if (TSharedPtr<FChooserRowDragDropOp> Operation = DragDropEvent.GetOperationAs<FChooserRowDragDropOp>())
+		{
+			float Center = MyGeometry.AbsolutePosition.Y + MyGeometry.Size.Y/2;
+			bDropAbove = DragDropEvent.GetScreenSpacePosition().Y < Center;
+			return FReply::Handled();
+		}
+		return FReply::Unhandled();
+	}
+
+	FReply SChooserTableRow::OnDrop(const FGeometry& MyGeometry, const FDragDropEvent& DragDropEvent)
+	{
+		if (TSharedPtr<FChooserRowDragDropOp> Operation = DragDropEvent.GetOperationAs<FChooserRowDragDropOp>())
+		{
+			if (Chooser == Operation->ChooserEditor->GetChooser())
+			{
+				int NewRowIndex;
+				if (!Chooser->ResultsStructs.IsValidIndex(RowIndex->RowIndex))
+				{
+					// for special (negative) indices, move to the end
+					NewRowIndex = Editor->MoveRow(Operation->RowIndex, Chooser->ResultsStructs.Num());
+				}
+				else if (bDropAbove)
+				{
+					NewRowIndex = Editor->MoveRow(Operation->RowIndex, RowIndex->RowIndex);
+				}
+				else
+				{
+					NewRowIndex = Editor->MoveRow(Operation->RowIndex, RowIndex->RowIndex+1);
+				}
+				Editor->SelectRow(NewRowIndex);
+				
+				return FReply::Handled();		
+			}
+		}
+		return FReply::Unhandled();
+	}
+
+}
+
+#undef LOCTEXT_NAMESPACE
diff --git a/Engine/Plugins/Experimental/Chooser/Source/ChooserEditor/Private/SChooserTableRow.h b/Engine/Plugins/Experimental/Chooser/Source/ChooserEditor/Private/SChooserTableRow.h
new file mode 100644
--- /dev/null
+++ b/Engine/Plugins/Experimental/Chooser/Source/ChooserEditor/Private/SChooserTableRow.h
@@ -0,0 +1,53 @@
+// Copyright Epic Games, Inc. All Rights Reserved.
+#pragma once
+
+#include "Widgets/Views/STableRow.h"
+
+class UChooserTable;
+
+namespace UE::ChooserEditor
+{
+
+class FChooserTableEditor;
+
+struct FChooserTableRow
+{
+	FChooserTableRow(int32 i) { RowIndex = i; }
+	int32 RowIndex;
+};
+
+class SChooserTableRow : public SMultiColumnTableRow<TSharedPtr<FChooserTableRow>>
+{
+public:
+SLATE_BEGIN_ARGS(SChooserTableRow) {}
+/** The list item for this row */
+	SLATE_ARGUMENT(TSharedPtr<FChooserTableRow>, Entry)
+	SLATE_ARGUMENT(UChooserTable*, Chooser)
+	SLATE_ARGUMENT(FChooserTableEditor*, Editor)
+	SLATE_END_ARGS()
+
+	enum { SpecialIndex_AddRow = -1, SpecialIndex_Fallback = -2 };
+
+	void Construct(const FArguments& Args, const TSharedRef<STableViewBase>& OwnerTableView);
+
+	/** Overridden from SMultiColumnTableRow.  Generates a widget for this column of the list view. */
+	virtual TSharedRef<SWidget> GenerateWidgetForColumn(const FName& ColumnName) override;
+	virtual void OnDragEnter(const FGeometry& MyGeometry, const FDragDropEvent& DragDropEvent) override;
+	virtual void OnDragLeave(const FDragDropEvent& DragDropEvent) override;
+	virtual FReply OnDragOver(const FGeometry& MyGeometry, const FDragDropEvent& DragDropEvent) override;
+	virtual FReply OnDrop(const FGeometry& MyGeometry, const FDragDropEvent& DragDropEvent) override;
+
+private:
+	TSharedPtr<FChooserTableRow> RowIndex;
+	UChooserTable* Chooser;
+	FChooserTableEditor* Editor;
+	TSharedPtr<SBorder> CacheBorder;
+	bool bDragActive = false;
+	bool bDropAbove = false;
+};
+
+
+	
+}
+
+#undef LOCTEXT_NAMESPACE
