remove extraneous undef LOCTEXT_NAMESPACE

diff --git a/Engine/Plugins/Experimental/Chooser/Source/ChooserEditor/Private/SChooserRowHandle.h b/Engine/Plugins/Experimental/Chooser/Source/ChooserEditor/Private/SChooserRowHandle.h
--- a/Engine/Plugins/Experimental/Chooser/Source/ChooserEditor/Private/SChooserRowHandle.h
+++ b/Engine/Plugins/Experimental/Chooser/Source/ChooserEditor/Private/SChooserRowHandle.h
@@ -42,5 +42,3 @@ private:
 };
 
 }
-
-#undef LOCTEXT_NAMESPACE
diff --git a/Engine/Plugins/Experimental/Chooser/Source/ChooserEditor/Private/SChooserTableRow.h b/Engine/Plugins/Experimental/Chooser/Source/ChooserEditor/Private/SChooserTableRow.h
--- a/Engine/Plugins/Experimental/Chooser/Source/ChooserEditor/Private/SChooserTableRow.h
+++ b/Engine/Plugins/Experimental/Chooser/Source/ChooserEditor/Private/SChooserTableRow.h
@@ -46,8 +46,4 @@ private:
 	bool bDropAbove = false;
 };
 
-
-	
 }
-
-#undef LOCTEXT_NAMESPACE
