Drag/Drop Assets to add to chooser

diff --git a/Engine/Plugins/Experimental/Chooser/Source/Chooser/Public/IChooserColumn.h b/Engine/Plugins/Experimental/Chooser/Source/Chooser/Public/IChooserColumn.h
--- a/Engine/Plugins/Experimental/Chooser/Source/Chooser/Public/IChooserColumn.h
+++ b/Engine/Plugins/Experimental/Chooser/Source/Chooser/Public/IChooserColumn.h
@@ -58,6 +58,7 @@ public:
 	virtual void SetNumRows(int32 NumRows) {}
 	virtual void DeleteRows(const TArray<uint32> & RowIndices) {}
 	virtual void MoveRow(int SourceIndex, int TargetIndex) {}
+	virtual void InsertRows(int Index, int Count) {}
 	
 	virtual UScriptStruct* GetInputBaseType() const { return nullptr; };
 	virtual const UScriptStruct* GetInputType() const { return nullptr; };
@@ -71,6 +72,7 @@ public:
 #endif
 };
 
+
 #if WITH_EDITOR
 #define CHOOSER_COLUMN_BOILERPLATE2(ParameterType, RowValuesProperty) \
 	virtual FName RowValuesPropertyName() override { return #RowValuesProperty; }\
@@ -86,6 +88,14 @@ public:
 			RowValuesProperty.Add(DefaultRowValue);\
 		}\
 	}\
+	virtual void InsertRows(int Index, int Count) override\
+	{\
+		RowValuesProperty.InsertUninitialized(Index, Count);\
+		for (int i=0;i<Count;i++)\
+		{\
+			RowValuesProperty[Index + i] = DefaultRowValue;\
+		}\
+	}\
 	virtual void DeleteRows(const TArray<uint32> & RowIndices )\
 	{\
 		for(uint32 Index : RowIndices)\
diff --git a/Engine/Plugins/Experimental/Chooser/Source/ChooserEditor/Private/ChooserTableEditor.cpp b/Engine/Plugins/Experimental/Chooser/Source/ChooserEditor/Private/ChooserTableEditor.cpp
--- a/Engine/Plugins/Experimental/Chooser/Source/ChooserEditor/Private/ChooserTableEditor.cpp
+++ b/Engine/Plugins/Experimental/Chooser/Source/ChooserEditor/Private/ChooserTableEditor.cpp
@@ -582,13 +582,16 @@ int FChooserTableEditor::MoveRow(int SourceRowIndex, int TargetRowIndex)
 	return TargetRowIndex;
 }
 	
-void FChooserTableEditor::SelectRow(int32 RowIndex)
+void FChooserTableEditor::SelectRow(int32 RowIndex, bool bClear)
 {
 	if (TableRows.IsValidIndex(RowIndex))
 	{
 		if (!TableView->IsItemSelected(TableRows[RowIndex]))
 		{
-			TableView->ClearSelection();
+			if (bClear)
+			{
+				TableView->ClearSelection();
+			}
 			TableView->SetItemSelection(TableRows[RowIndex], true, ESelectInfo::OnMouseClick);
 		}
 	}
diff --git a/Engine/Plugins/Experimental/Chooser/Source/ChooserEditor/Private/ChooserTableEditor.h b/Engine/Plugins/Experimental/Chooser/Source/ChooserEditor/Private/ChooserTableEditor.h
--- a/Engine/Plugins/Experimental/Chooser/Source/ChooserEditor/Private/ChooserTableEditor.h
+++ b/Engine/Plugins/Experimental/Chooser/Source/ChooserEditor/Private/ChooserTableEditor.h
@@ -92,15 +92,13 @@ namespace UE::ChooserEditor
 		/** Can be used to disable the details view making it read-only */
 		void SetPropertyEditingEnabledDelegate(FIsPropertyEditingEnabled InPropertyEditingDelegate);
 	
-
-
 		void UpdateTableRows();
 		void SelectColumn(UChooserTable* Chooser, int Index);
 		void ClearSelectedColumn();
 		void DeleteColumn(int Index);
 		void AddColumn(const UScriptStruct* ColumnType);
 		int MoveRow(int SourceRowIndex, int TargetIndex);
-		void SelectRow(int32 RowIndex);
+		void SelectRow(int32 RowIndex, bool bClear = true);
 		void ClearSelectedRows(); 
 	private:
 
diff --git a/Engine/Plugins/Experimental/Chooser/Source/ChooserEditor/Private/SChooserTableRow.cpp b/Engine/Plugins/Experimental/Chooser/Source/ChooserEditor/Private/SChooserTableRow.cpp
--- a/Engine/Plugins/Experimental/Chooser/Source/ChooserEditor/Private/SChooserTableRow.cpp
+++ b/Engine/Plugins/Experimental/Chooser/Source/ChooserEditor/Private/SChooserTableRow.cpp
@@ -4,17 +4,18 @@
 #include "Chooser.h"
 #include "ChooserEditorStyle.h"
 #include "ChooserTableEditor.h"
+#include "IContentBrowserSingleton.h"
 #include "IObjectChooser.h"
 #include "ObjectChooserWidgetFactories.h"
+#include "ObjectChooser_Asset.h"
 #include "SChooserRowHandle.h"
 #include "ScopedTransaction.h"
+#include "DragAndDrop/AssetDragDropOp.h"
 #include "Widgets/Colors/SColorBlock.h"
 #include "Widgets/Input/SComboButton.h"
 #include "Widgets/Views/STableRow.h"
 #include "Widgets/Layout/SSeparator.h"
 
-
-
 #define LOCTEXT_NAMESPACE "ChooserTableRow"
 
 namespace UE::ChooserEditor
@@ -29,6 +30,46 @@ namespace UE::ChooserEditor
 			FSuperRowType::FArguments(),
 			OwnerTableView
 		);
+
+		if (RowIndex->RowIndex >=0)
+		SetContent(
+				SNew(SOverlay)
+						+ SOverlay::Slot()
+						[
+							Content.Pin().ToSharedRef()
+						]
+						+ SOverlay::Slot().VAlign(VAlign_Bottom)
+						[
+							SNew(SSeparator)
+							.SeparatorImage(FAppStyle::GetBrush("PropertyEditor.HorizontalDottedLine"))
+							.ColorAndOpacity_Lambda([this]() { return FSlateColor(bDropSupported ? EStyleColor::Select : EStyleColor::Error); })
+							.Visibility_Lambda([this]() { return bDragActive && !bDropAbove ? EVisibility::Visible : EVisibility::Hidden; })
+						]
+						+ SOverlay::Slot().VAlign(VAlign_Top)
+						[
+							SNew(SSeparator)
+							.SeparatorImage(FAppStyle::GetBrush("PropertyEditor.HorizontalDottedLine"))
+							.ColorAndOpacity_Lambda([this]() { return FSlateColor(bDropSupported ? EStyleColor::Select : EStyleColor::Error); }) 
+							.Visibility_Lambda([this]() { return bDragActive && bDropAbove ? EVisibility::Visible : EVisibility::Hidden; })
+						]
+			);
+		else if (RowIndex->RowIndex == SpecialIndex_Fallback)
+		{
+			SetContent(
+					SNew(SOverlay)
+							+ SOverlay::Slot()
+							[
+								Content.Pin().ToSharedRef()
+							]
+							+ SOverlay::Slot().VAlign(VAlign_Top)
+							[
+								SNew(SSeparator)
+								.SeparatorImage(FAppStyle::GetBrush("PropertyEditor.HorizontalDottedLine"))
+								.ColorAndOpacity_Lambda([this]() { return FSlateColor(bDropSupported ? EStyleColor::Select : EStyleColor::Error); }) 
+								.Visibility_Lambda([this]() { return bDragActive ? EVisibility::Visible : EVisibility::Hidden; })
+							]
+				);
+		}
 	}
 
 	/** Overridden from SMultiColumnTableRow.  Generates a widget for this column of the list view. */
@@ -59,21 +100,7 @@ namespace UE::ChooserEditor
 				&CacheBorder
 				);
 			
-				return SNew(SOverlay)
-						+ SOverlay::Slot()
-						[
-							ResultWidget.ToSharedRef()
-						]
-						+ SOverlay::Slot().VAlign(VAlign_Bottom)
-						[
-							SNew(SSeparator).SeparatorImage(FCoreStyle::Get().GetBrush("FocusRectangle"))
-							.Visibility_Lambda([this]() { return bDragActive && !bDropAbove ? EVisibility::Visible : EVisibility::Hidden; })
-						]
-						+ SOverlay::Slot().VAlign(VAlign_Top)
-						[
-							SNew(SSeparator).SeparatorImage(FCoreStyle::Get().GetBrush("FocusRectangle"))
-							.Visibility_Lambda([this]() { return bDragActive && bDropAbove ? EVisibility::Visible : EVisibility::Hidden; })
-						];
+				return ResultWidget.ToSharedRef();
 			}
 			else
 			{
@@ -143,16 +170,7 @@ namespace UE::ChooserEditor
 				,FChooserWidgetValueChanged(), LOCTEXT("Fallback Result", "Fallback Result: (None)")
 				);
 				
-				return SNew(SOverlay)
-						+ SOverlay::Slot()
-						[
-							ResultWidget.ToSharedRef()
-						]
-						+ SOverlay::Slot().VAlign(VAlign_Top)
-						[
-							SNew(SSeparator).SeparatorImage(FCoreStyle::Get().GetBrush("FocusRectangle"))
-							.Visibility_Lambda([this]() { return bDragActive  ? EVisibility::Visible : EVisibility::Hidden; })
-						];
+				return ResultWidget.ToSharedRef();
 			}
 			else
 			{
@@ -185,12 +203,50 @@ namespace UE::ChooserEditor
 
 	void SChooserTableRow::OnDragEnter(const FGeometry& MyGeometry, const FDragDropEvent& DragDropEvent)
 	{
+		bDropSupported = false;
 		if (TSharedPtr<FChooserRowDragDropOp> Operation = DragDropEvent.GetOperationAs<FChooserRowDragDropOp>())
 		{
-			bDragActive = true;
-			float Center = MyGeometry.Position.Y + MyGeometry.Size.Y;
-			bDropAbove = DragDropEvent.GetScreenSpacePosition().Y < Center;
+			bDropSupported = true;
+		}
+		else if (TSharedPtr<FAssetDragDropOp> ContentDragDropOp = DragDropEvent.GetOperationAs<FAssetDragDropOp>())
+		{
+			if (Chooser->ResultType == EObjectChooserResultType::ObjectResult)
+			{
+				bDropSupported = true;
+			
+				if (Chooser->OutputObjectType) // if OutputObjectType is null, then any kind of object is supported, don't need to check all of them
+				{
+					for (const FAssetData& Asset : ContentDragDropOp->GetAssets())
+					{
+						const UClass* AssetClass = Asset.GetClass();
+
+						if(AssetClass->IsChildOf(Chooser->OutputObjectType))
+						{
+							bDropSupported = false;
+							break;
+						}
+						
+						if(AssetClass->IsChildOf(UChooserTable::StaticClass()))
+						{
+							const UChooserTable* DraggedChooserTable = Cast<UChooserTable>(Asset.GetAsset());
+
+							// verify dragged chooser result type matches this chooser result type
+							if (DraggedChooserTable->ResultType == EObjectChooserResultType::ClassResult
+							    || DraggedChooserTable->OutputObjectType == nullptr
+							    || !DraggedChooserTable->OutputObjectType->IsChildOf(Chooser->OutputObjectType))
+							{
+								bDropSupported = false;
+								break;
+							}
+						}
+					}
+				}
+			}
 		}
+		
+		float Center = MyGeometry.Position.Y + MyGeometry.Size.Y;
+		bDropAbove = DragDropEvent.GetScreenSpacePosition().Y < Center;
+		bDragActive = true;
 	}
 	void SChooserTableRow::OnDragLeave(const FDragDropEvent& DragDropEvent)
 	{
@@ -199,7 +255,8 @@ namespace UE::ChooserEditor
 
 	FReply SChooserTableRow::OnDragOver(const FGeometry& MyGeometry, const FDragDropEvent& DragDropEvent)
 	{
-		if (TSharedPtr<FChooserRowDragDropOp> Operation = DragDropEvent.GetOperationAs<FChooserRowDragDropOp>())
+		if (DragDropEvent.GetOperationAs<FChooserRowDragDropOp>()
+			|| DragDropEvent.GetOperationAs<FAssetDragDropOp>())
 		{
 			float Center = MyGeometry.AbsolutePosition.Y + MyGeometry.Size.Y/2;
 			bDropAbove = DragDropEvent.GetScreenSpacePosition().Y < Center;
@@ -210,6 +267,13 @@ namespace UE::ChooserEditor
 
 	FReply SChooserTableRow::OnDrop(const FGeometry& MyGeometry, const FDragDropEvent& DragDropEvent)
 	{
+		bDragActive = false;
+		
+		if (!bDropSupported)
+		{
+			return FReply::Unhandled();
+		}
+		
 		if (TSharedPtr<FChooserRowDragDropOp> Operation = DragDropEvent.GetOperationAs<FChooserRowDragDropOp>())
 		{
 			if (Chooser == Operation->ChooserEditor->GetChooser())
@@ -229,11 +293,74 @@ namespace UE::ChooserEditor
 					NewRowIndex = Editor->MoveRow(Operation->RowIndex, RowIndex->RowIndex+1);
 				}
 				Editor->SelectRow(NewRowIndex);
+			}
+		}
+		else if (TSharedPtr<FAssetDragDropOp> ContentDragDropOp = DragDropEvent.GetOperationAs<FAssetDragDropOp>())
+		{
+			
+			FScopedTransaction ScopedTransaction(LOCTEXT("DragDropAssets","Drag and Drop Assets into Chooser"));
+			Chooser->Modify();
+			
+			if (Chooser->ResultType == EObjectChooserResultType::ObjectResult)
+			{
+				int InsertRowIndex = RowIndex->RowIndex;
+				if (!Chooser->ResultsStructs.IsValidIndex(RowIndex->RowIndex))
+				{
+					// for special (negative) indices, move to the end
+					InsertRowIndex = Chooser->ResultsStructs.Num();
+				}
+				else if (!bDropAbove)
+				{
+					InsertRowIndex = RowIndex->RowIndex + 1;
+				}
 				
-				return FReply::Handled();		
+				if (bDropSupported)
+				{
+					TArray<FInstancedStruct> NewResults;
+					const TArray<FAssetData>& AssetList = ContentDragDropOp->GetAssets();
+					NewResults.Reserve(AssetList.Num());
+					
+					for (const FAssetData& Asset : AssetList)
+					{
+						const UClass* AssetClass = Asset.GetClass();
+						
+						NewResults.SetNum(NewResults.Num()+1);
+						FInstancedStruct& NewResult = NewResults.Last();
+
+						if(AssetClass->IsChildOf(UChooserTable::StaticClass()))
+						{
+							NewResult.InitializeAs(FEvaluateChooser::StaticStruct());
+							NewResult.GetMutable<FEvaluateChooser>().Chooser = Cast<UChooserTable>(Asset.GetAsset());
+						}
+						else
+						{
+							check(AssetClass->IsChildOf(Chooser->OutputObjectType));
+							NewResult.InitializeAs(FAssetChooser::StaticStruct());
+							NewResult.GetMutable<FAssetChooser>().Asset = Asset.GetAsset();
+						}
+					}
+
+					Chooser->ResultsStructs.Insert(NewResults, InsertRowIndex);
+					
+					// Make sure each column has the same number of row datas as there are results
+					for(FInstancedStruct& ColumnData : Chooser->ColumnsStructs)
+					{
+						FChooserColumnBase& Column = ColumnData.GetMutable<FChooserColumnBase>();
+						Column.InsertRows(InsertRowIndex, NewResults.Num());
+					}
+
+					Editor->RefreshAll();
+					
+					Editor->ClearSelectedRows();
+					for(int Index = InsertRowIndex; Index < InsertRowIndex + NewResults.Num(); Index++)
+					{
+						Editor->SelectRow(Index, false);
+					}
+				}
 			}
 		}
-		return FReply::Unhandled();
+				
+		return FReply::Handled();		
 	}
 
 }
diff --git a/Engine/Plugins/Experimental/Chooser/Source/ChooserEditor/Private/SChooserTableRow.h b/Engine/Plugins/Experimental/Chooser/Source/ChooserEditor/Private/SChooserTableRow.h
--- a/Engine/Plugins/Experimental/Chooser/Source/ChooserEditor/Private/SChooserTableRow.h
+++ b/Engine/Plugins/Experimental/Chooser/Source/ChooserEditor/Private/SChooserTableRow.h
@@ -43,6 +43,7 @@ private:
 	FChooserTableEditor* Editor;
 	TSharedPtr<SBorder> CacheBorder;
 	bool bDragActive = false;
+	bool bDropSupported = false;
 	bool bDropAbove = false;
 };
 
