ProxyTable: Change the fallback path in FProxyEntry::GetGuid for old FName key content to use a deterministic hash of the key to fix indeterminism in the cook.
#rnx

diff --git a/Engine/Plugins/Experimental/Chooser/Source/ProxyTable/Private/ProxyTable.cpp b/Engine/Plugins/Experimental/Chooser/Source/ProxyTable/Private/ProxyTable.cpp
--- a/Engine/Plugins/Experimental/Chooser/Source/ProxyTable/Private/ProxyTable.cpp
+++ b/Engine/Plugins/Experimental/Chooser/Source/ProxyTable/Private/ProxyTable.cpp
@@ -1,7 +1,9 @@
 // Copyright Epic Games, Inc. All Rights Reserved.
 #include "ProxyTable.h"
+
 #include "AssetRegistry/AssetRegistryModule.h"
 #include "Logging/LogMacros.h"
+#include "Misc/StringBuilder.h"
 
 DEFINE_LOG_CATEGORY_STATIC(LogProxyTable,Log,All);
 
@@ -23,7 +25,7 @@ const FGuid FProxyEntry::GetGuid() const
 		FGuid Guid;
 		if (Key != NAME_None)
 		{
-			Guid.A = GetTypeHash(Key);
+			Guid.A = GetTypeHash(WriteToString<128>(Key).ToView());
 		}
 		return Guid;
 	}
