ProxyTable: Make the hash in FProxyEntry::GetGuid match the hash in UProxyTableFunctionLibrary::EvaluateProxyTable.
#rnx

diff --git a/Engine/Plugins/Experimental/Chooser/Source/ProxyTable/Private/ProxyTable.cpp b/Engine/Plugins/Experimental/Chooser/Source/ProxyTable/Private/ProxyTable.cpp
--- a/Engine/Plugins/Experimental/Chooser/Source/ProxyTable/Private/ProxyTable.cpp
+++ b/Engine/Plugins/Experimental/Chooser/Source/ProxyTable/Private/ProxyTable.cpp
@@ -25,7 +25,7 @@ const FGuid FProxyEntry::GetGuid() const
 		FGuid Guid;
 		if (Key != NAME_None)
 		{
-			Guid.A = GetTypeHash(WriteToString<128>(Key).ToView());
+			Guid.A = GetTypeHash(WriteToString<128>(Key).ToView()); // Make sure this matches UProxyTableFunctionLibrary::EvaluateProxyTable
 		}
 		return Guid;
 	}
diff --git a/Engine/Plugins/Experimental/Chooser/Source/ProxyTable/Private/ProxyTableFunctionLibrary.cpp b/Engine/Plugins/Experimental/Chooser/Source/ProxyTable/Private/ProxyTableFunctionLibrary.cpp
--- a/Engine/Plugins/Experimental/Chooser/Source/ProxyTable/Private/ProxyTableFunctionLibrary.cpp
+++ b/Engine/Plugins/Experimental/Chooser/Source/ProxyTable/Private/ProxyTableFunctionLibrary.cpp
@@ -2,6 +2,7 @@
 #include "ProxyTableFunctionLibrary.h"
 
 #include "LookupProxy.h"
+#include "Misc/StringBuilder.h"
 #include "ProxyTable.h"
 
 /////////////////////////////////////////////////////////////////////////////////////
@@ -30,7 +31,7 @@ UObject* UProxyTableFunctionLibrary::EvaluateProxyTable(const UObject* ContextOb
 	if (ProxyTable)
 	{
 		FGuid Guid;
-		Guid.A = GetTypeHash(Key);
+		Guid.A = GetTypeHash(WriteToString<128>(Key).ToView()); // Make sure this matches FProxyEntry::GetGuid
 		FChooserEvaluationContext Context(const_cast<UObject*>(ContextObject));
 		if (UObject* Value = ProxyTable->FindProxyObject(Guid, Context))
 		{
