ChooserTable Editor selection refactor for multi select support

diff --git a/Engine/Plugins/Experimental/Chooser/Source/Chooser/Internal/BoolColumn.h b/Engine/Plugins/Experimental/Chooser/Source/Chooser/Internal/BoolColumn.h
--- a/Engine/Plugins/Experimental/Chooser/Source/Chooser/Internal/BoolColumn.h
+++ b/Engine/Plugins/Experimental/Chooser/Source/Chooser/Internal/BoolColumn.h
@@ -10,7 +10,7 @@
 #include "BoolColumn.generated.h"
 
 UENUM()
-enum class EBoolColumnCellValue
+enum class EBoolColumnCellValue : uint8
 {
 	MatchFalse = 0,
 	MatchTrue = 1,
@@ -81,6 +81,9 @@ struct CHOOSER_API FBoolColumn : public FChooserColumnBase
 		FMemoryReaderView Reader(Value);
 		Reader << TestValue;
 	}
+	
+	virtual void AddToDetails(FInstancedPropertyBag& PropertyBag, int32 ColumnIndex, int32 RowIndex) override;
+	virtual void SetFromDetails(FInstancedPropertyBag& PropertyBag, int32 ColumnIndex, int32 RowIndex) override;
 #endif
 
 	virtual void PostLoad() override
diff --git a/Engine/Plugins/Experimental/Chooser/Source/Chooser/Internal/EnumColumn.h b/Engine/Plugins/Experimental/Chooser/Source/Chooser/Internal/EnumColumn.h
--- a/Engine/Plugins/Experimental/Chooser/Source/Chooser/Internal/EnumColumn.h
+++ b/Engine/Plugins/Experimental/Chooser/Source/Chooser/Internal/EnumColumn.h
@@ -115,6 +115,9 @@ public:
 		FMemoryReaderView Reader(Value);
 		Reader << TestValue;
 	}
+	
+	virtual void AddToDetails (FInstancedPropertyBag& PropertyBag, int32 ColumnIndex, int32 RowIndex);
+	virtual void SetFromDetails(FInstancedPropertyBag& PropertyBag, int32 ColumnIndex, int32 RowIndex);
 #endif
 	
 	CHOOSER_COLUMN_BOILERPLATE(FChooserParameterEnumBase);
diff --git a/Engine/Plugins/Experimental/Chooser/Source/Chooser/Internal/FloatRangeColumn.h b/Engine/Plugins/Experimental/Chooser/Source/Chooser/Internal/FloatRangeColumn.h
--- a/Engine/Plugins/Experimental/Chooser/Source/Chooser/Internal/FloatRangeColumn.h
+++ b/Engine/Plugins/Experimental/Chooser/Source/Chooser/Internal/FloatRangeColumn.h
@@ -84,6 +84,9 @@ struct CHOOSER_API FFloatRangeColumn : public FChooserColumnBase
 		FMemoryReaderView Reader(Value);
 		Reader << TestValue;
 	}
+
+	virtual void AddToDetails(FInstancedPropertyBag& PropertyBag, int32 ColumnIndex, int32 RowIndex) override;
+	virtual void SetFromDetails(FInstancedPropertyBag& PropertyBag, int32 ColumnIndex, int32 RowIndex) override;
 #endif
 	
 	virtual void PostLoad() override
diff --git a/Engine/Plugins/Experimental/Chooser/Source/Chooser/Internal/GameplayTagColumn.h b/Engine/Plugins/Experimental/Chooser/Source/Chooser/Internal/GameplayTagColumn.h
--- a/Engine/Plugins/Experimental/Chooser/Source/Chooser/Internal/GameplayTagColumn.h
+++ b/Engine/Plugins/Experimental/Chooser/Source/Chooser/Internal/GameplayTagColumn.h
@@ -89,6 +89,9 @@ struct CHOOSER_API FGameplayTagColumn : public FChooserColumnBase
     	Reader << Tags;
 		TestValue.FromExportString(Tags);
     }
+	
+	virtual void AddToDetails(FInstancedPropertyBag& PropertyBag, int32 ColumnIndex, int32 RowIndex) override;
+	virtual void SetFromDetails(FInstancedPropertyBag& PropertyBag, int32 ColumnIndex, int32 RowIndex) override;
 #endif
 
 	virtual void PostLoad() override
diff --git a/Engine/Plugins/Experimental/Chooser/Source/Chooser/Internal/ObjectColumn.h b/Engine/Plugins/Experimental/Chooser/Source/Chooser/Internal/ObjectColumn.h
--- a/Engine/Plugins/Experimental/Chooser/Source/Chooser/Internal/ObjectColumn.h
+++ b/Engine/Plugins/Experimental/Chooser/Source/Chooser/Internal/ObjectColumn.h
@@ -87,6 +87,9 @@ struct CHOOSER_API FObjectColumn : public FChooserColumnBase
 		Reader << Path;
 		TestValue.SetPath(Path);
 	}
+
+	virtual void AddToDetails(FInstancedPropertyBag& PropertyBag, int32 ColumnIndex, int32 RowIndex) override;
+	virtual void SetFromDetails(FInstancedPropertyBag& PropertyBag, int32 ColumnIndex, int32 RowIndex) override;
 #endif
 
 	CHOOSER_COLUMN_BOILERPLATE(FChooserParameterObjectBase);
diff --git a/Engine/Plugins/Experimental/Chooser/Source/Chooser/Internal/OutputBoolColumn.h b/Engine/Plugins/Experimental/Chooser/Source/Chooser/Internal/OutputBoolColumn.h
--- a/Engine/Plugins/Experimental/Chooser/Source/Chooser/Internal/OutputBoolColumn.h
+++ b/Engine/Plugins/Experimental/Chooser/Source/Chooser/Internal/OutputBoolColumn.h
@@ -23,6 +23,9 @@ struct CHOOSER_API FOutputBoolColumn : public FChooserColumnBase
 
 #if WITH_EDITOR
 	mutable bool TestValue=false;
+
+	virtual void AddToDetails(FInstancedPropertyBag& PropertyBag, int32 ColumnIndex, int32 RowIndex) override;
+	virtual void SetFromDetails(FInstancedPropertyBag& PropertyBag, int32 ColumnIndex, int32 RowIndex) override;
 #endif
 	
 	// FallbackValue will be used as the output value if the all rows in the chooser fail, and the FallbackResult from the chooser is used.
diff --git a/Engine/Plugins/Experimental/Chooser/Source/Chooser/Internal/OutputEnumColumn.h b/Engine/Plugins/Experimental/Chooser/Source/Chooser/Internal/OutputEnumColumn.h
--- a/Engine/Plugins/Experimental/Chooser/Source/Chooser/Internal/OutputEnumColumn.h
+++ b/Engine/Plugins/Experimental/Chooser/Source/Chooser/Internal/OutputEnumColumn.h
@@ -47,6 +47,9 @@ public:
 	
 #if WITH_EDITOR
 	mutable uint8 TestValue;
+
+	virtual void AddToDetails (FInstancedPropertyBag& PropertyBag, int32 ColumnIndex, int32 RowIndex) override;
+	virtual void SetFromDetails(FInstancedPropertyBag& PropertyBag, int32 ColumnIndex, int32 RowIndex) override;
 #endif
 	
 	CHOOSER_COLUMN_BOILERPLATE(FChooserParameterEnumBase);
diff --git a/Engine/Plugins/Experimental/Chooser/Source/Chooser/Internal/OutputFloatColumn.h b/Engine/Plugins/Experimental/Chooser/Source/Chooser/Internal/OutputFloatColumn.h
--- a/Engine/Plugins/Experimental/Chooser/Source/Chooser/Internal/OutputFloatColumn.h
+++ b/Engine/Plugins/Experimental/Chooser/Source/Chooser/Internal/OutputFloatColumn.h
@@ -23,6 +23,9 @@ struct CHOOSER_API FOutputFloatColumn : public FChooserColumnBase
 
 #if WITH_EDITOR
 	mutable double TestValue=false;
+
+	virtual void AddToDetails(FInstancedPropertyBag& PropertyBag, int32 ColumnIndex, int32 RowIndex) override;
+	virtual void SetFromDetails(FInstancedPropertyBag& PropertyBag, int32 ColumnIndex, int32 RowIndex) override;
 #endif
 
 	// FallbackValue will be used as the output value if the all rows in the chooser fail, and the FallbackResult from the chooser is used.
diff --git a/Engine/Plugins/Experimental/Chooser/Source/Chooser/Internal/OutputStructColumn.h b/Engine/Plugins/Experimental/Chooser/Source/Chooser/Internal/OutputStructColumn.h
--- a/Engine/Plugins/Experimental/Chooser/Source/Chooser/Internal/OutputStructColumn.h
+++ b/Engine/Plugins/Experimental/Chooser/Source/Chooser/Internal/OutputStructColumn.h
@@ -43,6 +43,9 @@ struct CHOOSER_API FOutputStructColumn : public FChooserColumnBase
 #if WITH_EDITOR
 	void StructTypeChanged();
 	mutable FInstancedStruct TestValue;
+
+	virtual void AddToDetails(FInstancedPropertyBag& PropertyBag, int32 ColumnIndex, int32 RowIndex) override;
+	virtual void SetFromDetails(FInstancedPropertyBag& PropertyBag, int32 ColumnIndex, int32 RowIndex) override;
 #endif
 	
 	// FallbackValue will be used as the output value if the all rows in the chooser fail, and the FallbackResult from the chooser is used.
diff --git a/Engine/Plugins/Experimental/Chooser/Source/Chooser/Internal/RandomizeColumn.h b/Engine/Plugins/Experimental/Chooser/Source/Chooser/Internal/RandomizeColumn.h
--- a/Engine/Plugins/Experimental/Chooser/Source/Chooser/Internal/RandomizeColumn.h
+++ b/Engine/Plugins/Experimental/Chooser/Source/Chooser/Internal/RandomizeColumn.h
@@ -55,6 +55,9 @@ struct CHOOSER_API FRandomizeColumn : public FChooserColumnBase
     	{
     		return true;
     	}
+
+		virtual void AddToDetails(FInstancedPropertyBag& PropertyBag, int32 ColumnIndex, int32 RowIndex) override;
+		virtual void SetFromDetails(FInstancedPropertyBag& PropertyBag, int32 ColumnIndex, int32 RowIndex) override;
     #endif
 
 	virtual void PostLoad() override
diff --git a/Engine/Plugins/Experimental/Chooser/Source/Chooser/Private/BoolColumn.cpp b/Engine/Plugins/Experimental/Chooser/Source/Chooser/Private/BoolColumn.cpp
--- a/Engine/Plugins/Experimental/Chooser/Source/Chooser/Private/BoolColumn.cpp
+++ b/Engine/Plugins/Experimental/Chooser/Source/Chooser/Private/BoolColumn.cpp
@@ -5,6 +5,10 @@
 #include "Chooser.h"
 #include "ChooserTrace.h"
 
+#if WITH_EDITOR
+#include "PropertyBag.h"
+#endif
+
 bool FBoolContextProperty::GetValue(FChooserEvaluationContext& Context, bool& OutResult) const
 {
 	return Binding.GetValue(Context, OutResult);
@@ -53,4 +57,28 @@ void FBoolColumn::Filter(FChooserEvaluationContext& Context, const FChooserIndex
 		// passthrough fallback (behaves better during live editing)
 		IndexListOut = IndexListIn;
 	}
-}
\ No newline at end of file
+}
+
+#if WITH_EDITOR
+	void FBoolColumn::AddToDetails(FInstancedPropertyBag& PropertyBag, int32 ColumnIndex, int32 RowIndex)
+	{
+		FText DisplayName;
+		InputValue.Get<FChooserParameterBoolBase>().GetDisplayName(DisplayName);
+		FName PropertyName("RowData",ColumnIndex);
+		FPropertyBagPropertyDesc PropertyDesc(PropertyName, EPropertyBagPropertyType::Enum, StaticEnum<EBoolColumnCellValue>());
+		PropertyDesc.MetaData.Add(FPropertyBagPropertyDescMetaData("DisplayName", DisplayName.ToString()));
+		PropertyBag.AddProperties({PropertyDesc});
+		PropertyBag.SetValueEnum(PropertyName, RowValuesWithAny[RowIndex]);
+	}
+
+	void FBoolColumn::SetFromDetails(FInstancedPropertyBag& PropertyBag, int32 ColumnIndex, int32 RowIndex)
+	{
+		FName PropertyName("RowData", ColumnIndex);
+		
+		TValueOrError<uint8, EPropertyBagResult> Result = PropertyBag.GetValueEnum(PropertyName, StaticEnum<EBoolColumnCellValue>());
+		if (uint8* Value = Result.TryGetValue())
+		{
+			RowValuesWithAny[RowIndex] = static_cast<EBoolColumnCellValue>(*Value);
+		}
+	}
+#endif
\ No newline at end of file
diff --git a/Engine/Plugins/Experimental/Chooser/Source/Chooser/Private/EnumColumn.cpp b/Engine/Plugins/Experimental/Chooser/Source/Chooser/Private/EnumColumn.cpp
--- a/Engine/Plugins/Experimental/Chooser/Source/Chooser/Private/EnumColumn.cpp
+++ b/Engine/Plugins/Experimental/Chooser/Source/Chooser/Private/EnumColumn.cpp
@@ -3,7 +3,9 @@
 #include "ChooserIndexArray.h"
 #include "ChooserPropertyAccess.h"
 #include "ChooserTrace.h"
-
+#if WITH_EDITOR
+#include "PropertyBag.h"
+#endif
 
 bool FEnumContextProperty::GetValue(FChooserEvaluationContext& Context, uint8& OutResult) const
 {
@@ -95,4 +97,49 @@ void FEnumColumn::Filter(FChooserEvaluationContext& Context, const FChooserIndex
 		// passthrough fallback (behaves better during live editing)
 		IndexListOut = IndexListIn;
 	}
-}
\ No newline at end of file
+}
+
+#if WITH_EDITOR
+	void FEnumColumn::AddToDetails (FInstancedPropertyBag& PropertyBag, int32 ColumnIndex, int32 RowIndex)
+	{
+		FText DisplayName;
+		InputValue.Get<FChooserParameterBase>().GetDisplayName(DisplayName);
+		FName PropertyName("RowData",ColumnIndex);
+
+		// make another property bag in place of our struct, just so that the value enum will be correctly typed in the details panel
+		FInstancedPropertyBag Struct;
+		Struct.AddProperty("Value", EPropertyBagPropertyType::Enum, const_cast<UEnum*>(InputValue.Get<FChooserParameterEnumBase>().GetEnum()));
+		Struct.SetValueEnum("Value", RowValues[RowIndex].Value, InputValue.Get<FChooserParameterEnumBase>().GetEnum());
+		Struct.AddProperty("Comparison", EPropertyBagPropertyType::Enum, StaticEnum<EEnumColumnCellValueComparison>());
+		Struct.SetValueEnum("Comparison", static_cast<uint8>(RowValues[RowIndex].Comparison), StaticEnum<EEnumColumnCellValueComparison>());
+		
+		FPropertyBagPropertyDesc PropertyDesc(PropertyName,  EPropertyBagPropertyType::Struct, const_cast<UPropertyBag*>(Struct.GetPropertyBagStruct()));
+		PropertyDesc.MetaData.Add(FPropertyBagPropertyDescMetaData("DisplayName", DisplayName.ToString()));
+		PropertyBag.AddProperties({PropertyDesc});
+		PropertyBag.SetValueStruct(PropertyName, Struct.GetValue());
+	}
+	
+	void FEnumColumn::SetFromDetails(FInstancedPropertyBag& PropertyBag, int32 ColumnIndex, int32 RowIndex)
+	{
+		FName PropertyName("RowData",ColumnIndex);
+
+		FInstancedPropertyBag Struct;
+		Struct.AddProperty("Value", EPropertyBagPropertyType::Enum, const_cast<UEnum*>(InputValue.Get<FChooserParameterEnumBase>().GetEnum()));
+		Struct.AddProperty("Comparison", EPropertyBagPropertyType::Enum, StaticEnum<EEnumColumnCellValueComparison>());
+		
+		TValueOrError<FStructView, EPropertyBagResult> Result = PropertyBag.GetValueStruct(PropertyName, const_cast<UPropertyBag*>(Struct.GetPropertyBagStruct()));
+		if (FStructView* StructView = Result.TryGetValue())
+		{
+			const UScriptStruct* StructDefinition = StructView->GetScriptStruct();
+			if (const FEnumProperty* ValueProperty = CastField<FEnumProperty>(StructDefinition->FindPropertyByName("Value")))
+			{
+				ValueProperty->GetValue_InContainer(StructView->GetMemory(), &RowValues[RowIndex].Value);
+			}
+			
+			if (const FEnumProperty* ComparisonProperty = CastField<FEnumProperty>(StructDefinition->FindPropertyByName("Comparison")))
+			{
+				ComparisonProperty->GetValue_InContainer(StructView->GetMemory(), &RowValues[RowIndex].Comparison);
+			}
+		}
+	}
+#endif
\ No newline at end of file
diff --git a/Engine/Plugins/Experimental/Chooser/Source/Chooser/Private/FloatRangeColumn.cpp b/Engine/Plugins/Experimental/Chooser/Source/Chooser/Private/FloatRangeColumn.cpp
--- a/Engine/Plugins/Experimental/Chooser/Source/Chooser/Private/FloatRangeColumn.cpp
+++ b/Engine/Plugins/Experimental/Chooser/Source/Chooser/Private/FloatRangeColumn.cpp
@@ -3,8 +3,9 @@
 #include "ChooserIndexArray.h"
 #include "ChooserPropertyAccess.h"
 #include "ChooserTrace.h"
-
-
+#if WITH_EDITOR
+#include "PropertyBag.h"
+#endif
 
 bool FFloatContextProperty::SetValue(FChooserEvaluationContext& Context, double InValue) const
 {
@@ -54,4 +55,29 @@ void FFloatRangeColumn::Filter(FChooserEvaluationContext& Context, const FChoose
 		// passthrough fallback (behaves better during live editing)
 		IndexListOut = IndexListIn;
 	}
-}
\ No newline at end of file
+}
+
+#if WITH_EDITOR
+	void FFloatRangeColumn::AddToDetails(FInstancedPropertyBag& PropertyBag, int32 ColumnIndex, int32 RowIndex)
+	{
+		FText DisplayName;
+		InputValue.Get<FChooserParameterFloatBase>().GetDisplayName(DisplayName);
+		FName PropertyName("RowData", ColumnIndex);
+		FPropertyBagPropertyDesc PropertyDesc(PropertyName,  EPropertyBagPropertyType::Struct, FChooserFloatRangeRowData::StaticStruct());
+		PropertyDesc.MetaData.Add(FPropertyBagPropertyDescMetaData("DisplayName", DisplayName.ToString()));
+		PropertyBag.AddProperties({PropertyDesc});
+		
+		PropertyBag.SetValueStruct(PropertyName, RowValues[RowIndex]);
+	}
+
+	void FFloatRangeColumn::SetFromDetails(FInstancedPropertyBag& PropertyBag, int32 ColumnIndex, int32 RowIndex)
+	{
+		FName PropertyName("RowData", ColumnIndex);
+		
+   		TValueOrError<FStructView, EPropertyBagResult> Result = PropertyBag.GetValueStruct(PropertyName, FChooserFloatRangeRowData::StaticStruct());
+		if (FStructView* StructView = Result.TryGetValue())
+		{
+			RowValues[RowIndex] = StructView->Get<FChooserFloatRangeRowData>();
+		}
+	}
+#endif
\ No newline at end of file
diff --git a/Engine/Plugins/Experimental/Chooser/Source/Chooser/Private/GameplayTagColumn.cpp b/Engine/Plugins/Experimental/Chooser/Source/Chooser/Private/GameplayTagColumn.cpp
--- a/Engine/Plugins/Experimental/Chooser/Source/Chooser/Private/GameplayTagColumn.cpp
+++ b/Engine/Plugins/Experimental/Chooser/Source/Chooser/Private/GameplayTagColumn.cpp
@@ -3,6 +3,10 @@
 #include "ChooserIndexArray.h"
 #include "ChooserPropertyAccess.h"
 #include "ChooserTrace.h"
+#if WITH_EDITOR
+#include "PropertyBag.h"
+#endif
+
 
 bool FGameplayTagContextProperty::GetValue(FChooserEvaluationContext& Context, const FGameplayTagContainer*& OutResult) const
 {
@@ -62,4 +66,28 @@ void FGameplayTagColumn::Filter(FChooserEvaluationContext& Context, const FChoos
 		// passthrough fallback (behaves better during live editing)
 		IndexListOut = IndexListIn;
 	}
-}
\ No newline at end of file
+}
+
+#if WITH_EDITOR
+	void FGameplayTagColumn::AddToDetails(FInstancedPropertyBag& PropertyBag, int32 ColumnIndex, int32 RowIndex)
+	{
+		FText DisplayName;
+		InputValue.Get<FChooserParameterGameplayTagBase>().GetDisplayName(DisplayName);
+		FName PropertyName("RowData",ColumnIndex);
+		FPropertyBagPropertyDesc PropertyDesc(PropertyName,  EPropertyBagPropertyType::Struct, FGameplayTagContainer::StaticStruct());
+		PropertyDesc.MetaData.Add(FPropertyBagPropertyDescMetaData("DisplayName", DisplayName.ToString()));
+		PropertyBag.AddProperties({PropertyDesc});
+		PropertyBag.SetValueStruct(PropertyName, RowValues[RowIndex]);
+	}
+
+	void FGameplayTagColumn::SetFromDetails(FInstancedPropertyBag& PropertyBag, int32 ColumnIndex, int32 RowIndex)
+	{
+		FName PropertyName("RowData", ColumnIndex);
+		
+		TValueOrError<FStructView, EPropertyBagResult> Result = PropertyBag.GetValueStruct(PropertyName, FGameplayTagContainer::StaticStruct());
+		if (FStructView* StructView = Result.TryGetValue())
+		{
+			RowValues[RowIndex] = StructView->Get<FGameplayTagContainer>();
+		}
+	}	
+#endif
\ No newline at end of file
diff --git a/Engine/Plugins/Experimental/Chooser/Source/Chooser/Private/ObjectColumn.cpp b/Engine/Plugins/Experimental/Chooser/Source/Chooser/Private/ObjectColumn.cpp
--- a/Engine/Plugins/Experimental/Chooser/Source/Chooser/Private/ObjectColumn.cpp
+++ b/Engine/Plugins/Experimental/Chooser/Source/Chooser/Private/ObjectColumn.cpp
@@ -5,7 +5,8 @@
 #include "ChooserTrace.h"
 
 #if WITH_EDITOR
-	#include "IPropertyAccessEditor.h"
+#include "IPropertyAccessEditor.h"
+#include "PropertyBag.h"
 #endif
 
 bool FObjectContextProperty::GetValue(FChooserEvaluationContext& Context, FSoftObjectPath& OutResult) const
@@ -91,3 +92,26 @@ void FObjectColumn::Filter(FChooserEvaluationContext& Context, const FChooserInd
 		IndexListOut = IndexListIn;
 	}
 }
+#if WITH_EDITOR
+	void FObjectColumn::AddToDetails(FInstancedPropertyBag& PropertyBag, int32 ColumnIndex, int32 RowIndex)
+	{
+		FText DisplayName;
+		InputValue.Get<FChooserParameterObjectBase>().GetDisplayName(DisplayName);
+		FName PropertyName("RowData",ColumnIndex);
+		FPropertyBagPropertyDesc PropertyDesc(PropertyName,  EPropertyBagPropertyType::Struct, FChooserObjectRowData::StaticStruct());
+		PropertyDesc.MetaData.Add(FPropertyBagPropertyDescMetaData("DisplayName", DisplayName.ToString()));
+		PropertyBag.AddProperties({PropertyDesc});
+		PropertyBag.SetValueStruct(PropertyName, RowValues[RowIndex]);
+	}
+
+	void FObjectColumn::SetFromDetails(FInstancedPropertyBag& PropertyBag, int32 ColumnIndex, int32 RowIndex)
+	{
+		FName PropertyName("RowData", ColumnIndex);
+		
+		TValueOrError<FStructView, EPropertyBagResult> Result = PropertyBag.GetValueStruct(PropertyName, FChooserObjectRowData::StaticStruct());
+		if (FStructView* StructView = Result.TryGetValue())
+		{
+			RowValues[RowIndex] = StructView->Get<FChooserObjectRowData>();
+		}
+	}
+#endif
\ No newline at end of file
diff --git a/Engine/Plugins/Experimental/Chooser/Source/Chooser/Private/OutputBoolColumn.cpp b/Engine/Plugins/Experimental/Chooser/Source/Chooser/Private/OutputBoolColumn.cpp
--- a/Engine/Plugins/Experimental/Chooser/Source/Chooser/Private/OutputBoolColumn.cpp
+++ b/Engine/Plugins/Experimental/Chooser/Source/Chooser/Private/OutputBoolColumn.cpp
@@ -2,6 +2,9 @@
 #include "OutputBoolColumn.h"
 #include "ChooserPropertyAccess.h"
 #include "BoolColumn.h"
+#if WITH_EDITOR
+#include "PropertyBag.h"
+#endif
 
 FOutputBoolColumn::FOutputBoolColumn()
 {
@@ -28,3 +31,28 @@ void FOutputBoolColumn::SetOutputs(FChooserEvaluationContext& Context, int RowIn
 	}
 #endif
 }
+
+#if WITH_EDITOR
+	void FOutputBoolColumn::AddToDetails(FInstancedPropertyBag& PropertyBag, int32 ColumnIndex, int32 RowIndex)
+	{
+		FText DisplayName;
+		InputValue.Get<FChooserParameterBoolBase>().GetDisplayName(DisplayName);
+		FName PropertyName("RowData",ColumnIndex);
+		FPropertyBagPropertyDesc PropertyDesc(PropertyName, EPropertyBagPropertyType::Bool);
+		PropertyDesc.MetaData.Add(FPropertyBagPropertyDescMetaData("DisplayName", DisplayName.ToString()));
+		PropertyBag.AddProperties({PropertyDesc});
+		PropertyBag.SetValueFloat(PropertyName, RowValues[RowIndex]);
+	}
+
+	void FOutputBoolColumn::SetFromDetails(FInstancedPropertyBag& PropertyBag, int32 ColumnIndex, int32 RowIndex)
+	{
+		FName PropertyName("RowData", ColumnIndex);
+			
+		TValueOrError<bool, EPropertyBagResult> Result = PropertyBag.GetValueBool(PropertyName);
+		if (bool* Value = Result.TryGetValue())
+		{
+			RowValues[RowIndex] = *Value;
+		}
+	} 
+
+#endif
diff --git a/Engine/Plugins/Experimental/Chooser/Source/Chooser/Private/OutputEnumColumn.cpp b/Engine/Plugins/Experimental/Chooser/Source/Chooser/Private/OutputEnumColumn.cpp
--- a/Engine/Plugins/Experimental/Chooser/Source/Chooser/Private/OutputEnumColumn.cpp
+++ b/Engine/Plugins/Experimental/Chooser/Source/Chooser/Private/OutputEnumColumn.cpp
@@ -2,6 +2,9 @@
 #include "OutputEnumColumn.h"
 #include "EnumColumn.h"
 #include "ChooserPropertyAccess.h"
+#if WITH_EDITOR
+#include "PropertyBag.h"
+#endif
 
 FOutputEnumColumn::FOutputEnumColumn()
 {
@@ -28,4 +31,31 @@ void FOutputEnumColumn::SetOutputs(FChooserEvaluationContext& Context, int RowIn
 		TestValue = RowValues[RowIndex].Value;
 	}
 #endif
-}
\ No newline at end of file
+}
+
+#if WITH_EDITOR
+	void FOutputEnumColumn::AddToDetails (FInstancedPropertyBag& PropertyBag, int32 ColumnIndex, int32 RowIndex)
+	{
+		FText DisplayName;
+		InputValue.Get<FChooserParameterBase>().GetDisplayName(DisplayName);
+		FName PropertyName("RowData",ColumnIndex);
+
+		// make another property bag in place of our struct, just so that the value enum will be correctly typed in the details panel
+		
+		FPropertyBagPropertyDesc PropertyDesc(PropertyName,  EPropertyBagPropertyType::Enum, InputValue.Get<FChooserParameterEnumBase>().GetEnum());
+		PropertyDesc.MetaData.Add(FPropertyBagPropertyDescMetaData("DisplayName", DisplayName.ToString()));
+		PropertyBag.AddProperties({PropertyDesc});
+		PropertyBag.SetValueEnum("Value", RowValues[RowIndex].Value, InputValue.Get<FChooserParameterEnumBase>().GetEnum());
+	}
+	
+	void FOutputEnumColumn::SetFromDetails(FInstancedPropertyBag& PropertyBag, int32 ColumnIndex, int32 RowIndex)
+	{
+		FName PropertyName("RowData",ColumnIndex);
+		
+		TValueOrError<uint8, EPropertyBagResult> Result = PropertyBag.GetValueEnum(PropertyName, InputValue.Get<FChooserParameterEnumBase>().GetEnum());
+		if (uint8* Value = Result.TryGetValue())
+		{
+			RowValues[RowIndex].Value = *Value;
+		}
+	}
+#endif
\ No newline at end of file
diff --git a/Engine/Plugins/Experimental/Chooser/Source/Chooser/Private/OutputFloatColumn .cpp b/Engine/Plugins/Experimental/Chooser/Source/Chooser/Private/OutputFloatColumn .cpp
--- a/Engine/Plugins/Experimental/Chooser/Source/Chooser/Private/OutputFloatColumn .cpp	
+++ b/Engine/Plugins/Experimental/Chooser/Source/Chooser/Private/OutputFloatColumn .cpp	
@@ -2,6 +2,9 @@
 #include "OutputFloatColumn.h"
 #include "ChooserPropertyAccess.h"
 #include "FloatRangeColumn.h"
+#if WITH_EDITOR
+#include "PropertyBag.h"
+#endif
 
 FOutputFloatColumn::FOutputFloatColumn()
 {
@@ -28,3 +31,26 @@ void FOutputFloatColumn::SetOutputs(FChooserEvaluationContext& Context, int RowI
 	}
 #endif
 }
+#if WITH_EDITOR
+	void FOutputFloatColumn::AddToDetails(FInstancedPropertyBag& PropertyBag, int32 ColumnIndex, int32 RowIndex)
+	{
+		FText DisplayName;
+		InputValue.Get<FChooserParameterBase>().GetDisplayName(DisplayName);
+		FName PropertyName("RowData",ColumnIndex);
+		FPropertyBagPropertyDesc PropertyDesc(PropertyName, EPropertyBagPropertyType::Float);
+		PropertyDesc.MetaData.Add(FPropertyBagPropertyDescMetaData("DisplayName", DisplayName.ToString()));
+		PropertyBag.AddProperties({PropertyDesc});
+		PropertyBag.SetValueFloat(PropertyName, RowValues[RowIndex]);
+	}
+
+	void FOutputFloatColumn::SetFromDetails(FInstancedPropertyBag& PropertyBag, int32 ColumnIndex, int32 RowIndex)
+	{
+		FName PropertyName("RowData", ColumnIndex);
+			
+		TValueOrError<float, EPropertyBagResult> Result = PropertyBag.GetValueFloat(PropertyName);
+		if (float* Value = Result.TryGetValue())
+		{
+			RowValues[RowIndex] = *Value;
+		}
+	}
+#endif
\ No newline at end of file
diff --git a/Engine/Plugins/Experimental/Chooser/Source/Chooser/Private/OutputStructColumn.cpp b/Engine/Plugins/Experimental/Chooser/Source/Chooser/Private/OutputStructColumn.cpp
--- a/Engine/Plugins/Experimental/Chooser/Source/Chooser/Private/OutputStructColumn.cpp
+++ b/Engine/Plugins/Experimental/Chooser/Source/Chooser/Private/OutputStructColumn.cpp
@@ -4,7 +4,8 @@
 #include "OutputStructColumn.h"
 
 #if WITH_EDITOR
-	#include "IPropertyAccessEditor.h"
+#include "IPropertyAccessEditor.h"
+#include "PropertyBag.h"
 #endif
 
 bool FStructContextProperty::SetValue(FChooserEvaluationContext& Context, const FInstancedStruct& InValue) const
@@ -46,6 +47,28 @@ void FOutputStructColumn::StructTypeChanged()
 	}
 }
 
+void FOutputStructColumn::AddToDetails(FInstancedPropertyBag& PropertyBag, int32 ColumnIndex, int32 RowIndex)
+{
+	FText DisplayName;
+	InputValue.Get<FChooserParameterStructBase>().GetDisplayName(DisplayName);
+	FName PropertyName("RowData",ColumnIndex);
+	FPropertyBagPropertyDesc PropertyDesc(PropertyName,  EPropertyBagPropertyType::Struct, RowValues[RowIndex].GetScriptStruct());
+	PropertyDesc.MetaData.Add(FPropertyBagPropertyDescMetaData("DisplayName", DisplayName.ToString()));
+	PropertyBag.AddProperties({PropertyDesc});
+	PropertyBag.SetValueStruct(PropertyName, FConstStructView(RowValues[RowIndex].GetScriptStruct(), RowValues[RowIndex].GetMemory()));
+}
+
+void FOutputStructColumn::SetFromDetails(FInstancedPropertyBag& PropertyBag, int32 ColumnIndex, int32 RowIndex)
+{
+	FName PropertyName("RowData", ColumnIndex);
+	
+	TValueOrError<FStructView, EPropertyBagResult> Result = PropertyBag.GetValueStruct(PropertyName, RowValues[RowIndex].GetScriptStruct());
+	if (FStructView* StructView = Result.TryGetValue())
+	{
+		RowValues[RowIndex].GetScriptStruct()->CopyScriptStruct(RowValues[RowIndex].GetMutableMemory(), StructView->GetMemory());
+	}
+}
+
 #endif // WITH_EDITOR
 
 FOutputStructColumn::FOutputStructColumn()
diff --git a/Engine/Plugins/Experimental/Chooser/Source/Chooser/Private/RandomizeColumn.cpp b/Engine/Plugins/Experimental/Chooser/Source/Chooser/Private/RandomizeColumn.cpp
--- a/Engine/Plugins/Experimental/Chooser/Source/Chooser/Private/RandomizeColumn.cpp
+++ b/Engine/Plugins/Experimental/Chooser/Source/Chooser/Private/RandomizeColumn.cpp
@@ -2,6 +2,9 @@
 #include "RandomizeColumn.h"
 #include "ChooserIndexArray.h"
 #include "ChooserPropertyAccess.h"
+#if WITH_EDITOR
+#include "PropertyBag.h"
+#endif
 
 bool FRandomizeContextProperty::GetValue(FChooserEvaluationContext& Context, const FChooserRandomizationContext*& OutResult) const
 {
@@ -101,3 +104,26 @@ void FRandomizeColumn::SetOutputs(FChooserEvaluationContext& Context, int RowInd
 		}
 	}
 }
+	#if WITH_EDITOR
+
+	void FRandomizeColumn::AddToDetails(FInstancedPropertyBag& PropertyBag, int32 ColumnIndex, int32 RowIndex)
+	{
+		FName PropertyName("RowData",ColumnIndex);
+		FPropertyBagPropertyDesc PropertyDesc(PropertyName, EPropertyBagPropertyType::Float);
+		PropertyDesc.MetaData.Add(FPropertyBagPropertyDescMetaData("DisplayName", "Randomize"));
+		PropertyBag.AddProperties({PropertyDesc});
+		PropertyBag.SetValueFloat(PropertyName, RowValues[RowIndex]);
+	}
+
+	void FRandomizeColumn::SetFromDetails(FInstancedPropertyBag& PropertyBag, int32 ColumnIndex, int32 RowIndex)
+	{
+		FName PropertyName("RowData", ColumnIndex);
+		
+		TValueOrError<float, EPropertyBagResult> Result = PropertyBag.GetValueFloat(PropertyName);
+		if (float* Value = Result.TryGetValue())
+		{
+			RowValues[RowIndex] = *Value;
+		}
+	}
+
+    #endif
\ No newline at end of file
diff --git a/Engine/Plugins/Experimental/Chooser/Source/Chooser/Public/IChooserColumn.h b/Engine/Plugins/Experimental/Chooser/Source/Chooser/Public/IChooserColumn.h
--- a/Engine/Plugins/Experimental/Chooser/Source/Chooser/Public/IChooserColumn.h
+++ b/Engine/Plugins/Experimental/Chooser/Source/Chooser/Public/IChooserColumn.h
@@ -28,6 +28,8 @@ class FVariant;
 
 class FChooserIndexArray;
 
+struct FInstancedPropertyBag;
+
 USTRUCT()
 struct CHOOSER_API FChooserColumnBase
 {
@@ -69,6 +71,9 @@ public:
 	virtual bool IsRandomizeColumn() const { return false; }
 
 	virtual bool EditorTestFilter(int32 RowIndex) const { return false; }
+
+	virtual void AddToDetails(FInstancedPropertyBag& PropertyBag, int32 ColumnIndex, int32 RowIndex) {}
+	virtual void SetFromDetails(FInstancedPropertyBag& PropertyBag, int32 ColumnIndex, int32 RowIndex) {}
 #endif
 };
 
diff --git a/Engine/Plugins/Experimental/Chooser/Source/ChooserEditor/Private/ChooserDetails.cpp b/Engine/Plugins/Experimental/Chooser/Source/ChooserEditor/Private/ChooserDetails.cpp
--- a/Engine/Plugins/Experimental/Chooser/Source/ChooserEditor/Private/ChooserDetails.cpp
+++ b/Engine/Plugins/Experimental/Chooser/Source/ChooserEditor/Private/ChooserDetails.cpp
@@ -48,43 +48,8 @@ void FChooserRowDetails::CustomizeDetails(IDetailLayoutBuilder& DetailBuilder)
 	UChooserRowDetails* Row = Cast<UChooserRowDetails>(Objects[0]);
 	UChooserTable* Chooser = Row->Chooser;
 	
-	if (Chooser->ResultsStructs.IsValidIndex(Row->Row))
-	{
-		IDetailCategoryBuilder& PropertiesCategory = DetailBuilder.EditCategory("Row Properties");
-
-		TSharedPtr<IPropertyHandle> ChooserProperty = DetailBuilder.GetProperty("Chooser", Row->StaticClass());
-		DetailBuilder.HideProperty(ChooserProperty);
-	
-		TSharedPtr<IPropertyHandle> ResultsArrayProperty = ChooserProperty->GetChildHandle("ResultsStructs");
-		TSharedPtr<IPropertyHandle> CurrentResultProperty = ResultsArrayProperty->AsArray()->GetElement(Row->Row);
-		IDetailPropertyRow& NewResultProperty = PropertiesCategory.AddProperty(CurrentResultProperty);
-		NewResultProperty.DisplayName(LOCTEXT("ResultColumnName","Result"));
-		NewResultProperty.ShowPropertyButtons(false); // hide array add button
-		NewResultProperty.ShouldAutoExpand(true);
-	
-		for(int ColumnIndex=0; ColumnIndex<Chooser->ColumnsStructs.Num(); ColumnIndex++)
-		{
-			FChooserColumnBase& Column = Chooser->ColumnsStructs[ColumnIndex].GetMutable<FChooserColumnBase>();
-			TSharedRef<FStructOnScope> StructOnScope = MakeShared<FStructOnScope>(Chooser->ColumnsStructs[ColumnIndex].GetScriptStruct(), reinterpret_cast<uint8*>(&Column));
-			TSharedPtr<IPropertyHandle> ColumnDataProperty = DetailBuilder.AddStructurePropertyData({StructOnScope}, Column.RowValuesPropertyName());
-			uint32 NumElements = 0;
-			ColumnDataProperty->AsArray()->GetNumElements(NumElements);
-			if (Row->Row < (int)NumElements)
-			{
-				TSharedRef<IPropertyHandle> CellData = ColumnDataProperty->AsArray()->GetElement(Row->Row);
-	
-				IDetailPropertyRow& NewColumnProperty = PropertiesCategory.AddProperty(CellData);
-				FText DisplayName = LOCTEXT("No Input Value", "No Input Value");
-				if (FChooserParameterBase* InputValue = Column.GetInputValue())
-				{
-					InputValue->GetDisplayName(DisplayName); 
-				}
-				NewColumnProperty.DisplayName(DisplayName);
-				NewColumnProperty.ShowPropertyButtons(false); // hide array add button
-				NewColumnProperty.ShouldAutoExpand(true);
-			}
-		}
-	}
+	TSharedPtr<IPropertyHandle> ChooserProperty = DetailBuilder.GetProperty("Chooser", Row->StaticClass());
+	DetailBuilder.HideProperty(ChooserProperty);
 }
 
 // Make the details panel show the values for the selected row, showing each column value
diff --git a/Engine/Plugins/Experimental/Chooser/Source/ChooserEditor/Private/ChooserDetails.h b/Engine/Plugins/Experimental/Chooser/Source/ChooserEditor/Private/ChooserDetails.h
--- a/Engine/Plugins/Experimental/Chooser/Source/ChooserEditor/Private/ChooserDetails.h
+++ b/Engine/Plugins/Experimental/Chooser/Source/ChooserEditor/Private/ChooserDetails.h
@@ -4,6 +4,7 @@
 
 #include "CoreMinimal.h"
 #include "IDetailCustomization.h"
+#include "PropertyBag.h"
 #include "ChooserDetails.generated.h"
 
 class UChooserTable;
@@ -14,6 +15,10 @@ class UChooserRowDetails : public UObject
 {
 	GENERATED_BODY()
 public:
+
+	UPROPERTY(EditAnywhere, Category="Properties", meta=(FixedLayout, ShowOnlyInnerProperties))
+    FInstancedPropertyBag Properties;
+	
 	UPROPERTY(EditAnywhere, Instanced, Category="Hidden")
 	TObjectPtr<UChooserTable> Chooser;
 	int Row;
diff --git a/Engine/Plugins/Experimental/Chooser/Source/ChooserEditor/Private/ChooserTableEditor.cpp b/Engine/Plugins/Experimental/Chooser/Source/ChooserEditor/Private/ChooserTableEditor.cpp
--- a/Engine/Plugins/Experimental/Chooser/Source/ChooserEditor/Private/ChooserTableEditor.cpp
+++ b/Engine/Plugins/Experimental/Chooser/Source/ChooserEditor/Private/ChooserTableEditor.cpp
@@ -18,6 +18,7 @@
 #include "ObjectChooser_Asset.h"
 #include "ObjectChooser_Class.h"
 #include "PersonaModule.h"
+#include "PropertyBag.h"
 #include "PropertyCustomizationHelpers.h"
 #include "PropertyEditorModule.h"
 #include "SAssetDropTarget.h"
@@ -31,6 +32,7 @@
 #include "Modules/ModuleManager.h"
 #include "Styling/AppStyle.h"
 #include "SChooserTableRow.h"
+#include "Elements/Framework/TypedElementQueryBuilder.h"
 #include "Widgets/Docking/SDockTab.h"
 #include "Widgets/Input/SComboButton.h"
 #include "Widgets/Input/SHyperlink.h"
@@ -109,8 +111,10 @@ FChooserTableEditor::~FChooserTableEditor()
 	}
 	
 	FCoreUObjectDelegates::OnObjectsReplaced.RemoveAll(this);
+	FCoreUObjectDelegates::OnObjectTransacted.RemoveAll(this);
 
 	DetailsView.Reset();
+	
 }
 
 
@@ -277,6 +281,44 @@ void FChooserTableEditor::BindCommands()
 		FExecuteAction::CreateSP(this, &FChooserTableEditor::SelectRootProperties));
 }
 
+void FChooserTableEditor::OnObjectsTransacted(UObject* Object, const FTransactionObjectEvent& Event)
+{
+	if (UChooserTable* ChooserTable = Cast<UChooserTable>(Object))
+	{
+		// if this is the chooser we're editing
+		if (GetChooser() == ChooserTable)
+		{
+			if (CurrentSelectionType == ESelectionType::Rows)
+			{
+				// refresh details if we have rows selected
+				RefreshRowSelectionDetails();
+			}
+		}
+	}
+	
+	if (UChooserRowDetails* RowDetails = Cast<UChooserRowDetails>(Object))
+	{
+		// if this is for the chooser we're editing
+		if (GetChooser() == RowDetails->Chooser)
+		{
+			// copy all the values over
+			TValueOrError<FStructView, EPropertyBagResult> Result = RowDetails->Properties.GetValueStruct("Result", FInstancedStruct::StaticStruct());
+			if (Result.IsValid())
+			{
+				RowDetails->Chooser->ResultsStructs[RowDetails->Row] = Result.GetValue().Get<FInstancedStruct>();
+			}
+
+			int ColumnIndex = 0;
+			for (FInstancedStruct& ColumnData : RowDetails->Chooser->ColumnsStructs)
+			{
+				FChooserColumnBase& Column = ColumnData.GetMutable<FChooserColumnBase>();
+				Column.SetFromDetails(RowDetails->Properties, ColumnIndex, RowDetails->Row);
+				ColumnIndex++;
+			}
+		}
+	}
+}
+
 void FChooserTableEditor::InitEditor( const EToolkitMode::Type Mode, const TSharedPtr< class IToolkitHost >& InitToolkitHost, const TArray<UObject*>& ObjectsToEdit, FGetDetailsViewObjects GetDetailsViewObjects )
 {
 	EditingObjects = ObjectsToEdit;
@@ -341,6 +383,8 @@ void FChooserTableEditor::InitEditor( const EToolkitMode::Type Mode, const TShar
 		
 	FAnimAssetFindReplaceConfig FindReplaceConfig;
 	FindReplaceConfig.InitialProcessorClass = UChooserFindProperties::StaticClass();
+	
+	FCoreUObjectDelegates::OnObjectTransacted.AddSP(this, &FChooserTableEditor::OnObjectsTransacted);
 }
 
 FName FChooserTableEditor::GetToolkitFName() const
@@ -357,10 +401,50 @@ void FChooserTableEditor::RefreshAll()
 {
 	if (HeaderRow)
 	{
+		// Cache Selection state
+		ESelectionType CachedSelectionType = CurrentSelectionType;
+		int SelectedColumnIndex = -1;
+		UChooserTable* SelectedChooser = nullptr;
+		TArray<int> CachedSelectedRows;
+
+		if (CachedSelectionType == ESelectionType::Column)
+		{
+			SelectedColumnIndex = SelectedColumn->Column;
+			SelectedChooser = SelectedColumn->Chooser;
+		}
+		else if (CachedSelectionType == ESelectionType::Rows)
+		{
+			if (!SelectedRows.IsEmpty())
+			{
+				SelectedChooser = SelectedRows[0]->Chooser;
+			}
+			for(const TObjectPtr<UChooserRowDetails>& SelectedRow : SelectedRows)
+			{
+				CachedSelectedRows.Add(SelectedRow->Row);
+			}
+		}
+		
 		UpdateTableColumns();
 		UpdateTableRows();
-		ClearSelectedColumn();
-		SelectRootProperties();
+
+		// reapply cached selection state
+		if (CachedSelectionType == ESelectionType::Root)
+		{
+			SelectRootProperties();
+		}
+		else if (CachedSelectionType == ESelectionType::Column)
+		{
+			SelectColumn(SelectedChooser, SelectedColumnIndex);
+		}
+		else if (CachedSelectionType == ESelectionType::Rows)
+		{
+			ClearSelectedRows();
+			for(int Row : CachedSelectedRows)
+			{
+				SelectRow(Row, false);
+			}
+		}
+		
 	}
 }
 
@@ -381,10 +465,26 @@ void FChooserTableEditor::NotifyPreChange(FProperty* PropertyAboutToChange)
 
 void FChooserTableEditor::NotifyPostChange(const FPropertyChangedEvent& PropertyChangedEvent, FProperty* PropertyThatChanged)
 {
-	// was previously only refreshing when known properties changed, however, changing the contents of an FInstancedStruct in the Details Panel (for example DefaultRowValue on a Column)
-	// will invalidate previously cached pointers to that FInstancedStruct, and so we need to refresh the columns every time anything is edited in the Details Panel
-	UpdateTableColumns();
-	UpdateTableRows();
+	// Called on details panel edits
+	
+	if (CurrentSelectionType==ESelectionType::Root)
+	{
+		// Editing the root in the details panel can change ContextData that means all wigets need to be refreshed
+		UpdateTableColumns();
+		UpdateTableRows();
+		SelectRootProperties();
+	}
+	if (CurrentSelectionType==ESelectionType::Column)
+	{
+		check(SelectedColumn);
+		int SelectedColumnIndex = SelectedColumn->Column;
+		UChooserTable* SelectedColumnChooser = SelectedColumn->Chooser;
+		// Editing column properties can change the column type, which requires refreshing everything
+		UpdateTableColumns();
+		UpdateTableRows();
+		SelectColumn(SelectedColumnChooser, SelectedColumnIndex);
+	}
+	// editing row data should not require any refreshing
 }
 
 
@@ -535,8 +635,6 @@ TSharedRef<SDockTab> FChooserTableEditor::SpawnFindReplaceTab( const FSpawnTabAr
 	];
 }
 
-
-
 TSharedRef<ITableRow> FChooserTableEditor::GenerateTableRow(TSharedPtr<FChooserTableRow> InItem, const TSharedRef<STableViewBase>& OwnerTable)
 {
 	UChooserTable* Chooser = GetChooser();
@@ -551,6 +649,7 @@ void FChooserTableEditor::SelectRootProperties()
 	{
 		// point the details view to the main table
 		DetailsView->SetObject( GetRootChooser() );
+		CurrentSelectionType = ESelectionType::Root;
 	}
 }
 
@@ -604,6 +703,18 @@ void FChooserTableEditor::ClearSelectedRows()
 	SelectRootProperties();
 }
 
+bool FChooserTableEditor::IsRowSelected(int32 RowIndex)
+{
+	for(auto& SelectedRow:SelectedRows)
+ 	{
+ 		if (SelectedRow->Row == RowIndex)
+ 		{
+ 			return true;
+ 		}
+ 	}
+	return false;
+}
+
 void FChooserTableEditor::UpdateTableColumns()
 {
 	UChooserTable* Chooser = GetChooser();
@@ -806,6 +917,57 @@ void FChooserTableEditor::AddColumn(const UScriptStruct* ColumnType)
 	SelectColumn(Chooser, InsertIndex);
 }
 
+void FChooserTableEditor::RefreshRowSelectionDetails()
+{
+	for (UObject* SelectedRow : SelectedRows)
+	{
+		SelectedRow->ClearFlags(RF_Standalone);
+	}
+	SelectedRows.SetNum(0);
+	UChooserTable* Chooser = GetChooser();
+	
+	// Get the list of objects to edit the details of
+	TArray<TSharedPtr<FChooserTableRow>> SelectedItems = TableView->GetSelectedItems();
+	for(TSharedPtr<FChooserTableRow>& SelectedItem : SelectedItems)
+	{
+		if (Chooser->ResultsStructs.IsValidIndex(SelectedItem->RowIndex))
+		{
+			TObjectPtr<UChooserRowDetails> Selection = NewObject<UChooserRowDetails>();
+			Selection->Chooser = Chooser;
+			Selection->Row = SelectedItem->RowIndex;
+			Selection->SetFlags(RF_Standalone | RF_Transactional);
+
+			FInstancedStruct& Result = Chooser->ResultsStructs[SelectedItem->RowIndex];
+			Selection->Properties.AddProperty("Result", EPropertyBagPropertyType::Struct, FInstancedStruct::StaticStruct());
+			Selection->Properties.SetValueStruct("Result", FConstStructView(FInstancedStruct::StaticStruct(), reinterpret_cast<uint8*>(&Result)));
+
+			int ColumnIndex = 0;
+			for (FInstancedStruct& ColumnData : Chooser->ColumnsStructs)
+			{
+				FChooserColumnBase& Column = ColumnData.GetMutable<FChooserColumnBase>();
+				Column.AddToDetails(Selection->Properties, ColumnIndex, SelectedItem->RowIndex);
+				ColumnIndex++;
+			}
+		
+			SelectedRows.Add(Selection);
+			
+		}
+	}
+	
+	TArray<UObject*> DetailsObjects;
+	for(auto& Item : SelectedRows)
+	{
+		DetailsObjects.Add(Item.Get());
+	}
+
+	if( DetailsView.IsValid() )
+	{
+		// Make sure details window is pointing to our object
+		DetailsView->SetObjects( DetailsObjects );
+	}
+}
+    										
+
 TSharedRef<SDockTab> FChooserTableEditor::SpawnTableTab( const FSpawnTabArgs& Args )
 {
 	check( Args.GetTabId() == TableTabId );
@@ -908,45 +1070,14 @@ TSharedRef<SDockTab> FChooserTableEditor::SpawnTableTab( const FSpawnTabArgs& Ar
 					return FReply::Unhandled();
 				}
 				)
-				.OnSelectionChanged_Lambda([this](TSharedPtr<FChooserTableRow> SelectedItem,  ESelectInfo::Type SelectInfo)
+				.OnSelectionChanged_Lambda([this](TSharedPtr<FChooserTableRow>,  ESelectInfo::Type SelectInfo)
 				{
 					// deselect any selected column
 					ClearSelectedColumn();
-					
-					if (SelectedItem)
-					{
-						for (UObject* SelectedRow : SelectedRows)
-						{
-							SelectedRow->ClearFlags(RF_Standalone);
-						}
-						SelectedRows.SetNum(0);
-						UChooserTable* Chooser = GetChooser();
-						
-						if (!Chooser->ResultsStructs.IsValidIndex(SelectedItem->RowIndex))
-						{
-							SelectRootProperties();
-							return;
-						}
-						
-						// Get the list of objects to edit the details of
-						TObjectPtr<UChooserRowDetails> Selection = NewObject<UChooserRowDetails>();
-						Selection->Chooser = Chooser;
-						Selection->Row = SelectedItem->RowIndex;
-						Selection->SetFlags(RF_Standalone);
-						SelectedRows.Add(Selection);
-											
-						TArray<UObject*> DetailsObjects;
-						for(auto& Item : SelectedRows)
-						{
-							DetailsObjects.Add(Item.Get());
-						}
 
-						if( DetailsView.IsValid() )
-						{
-							// Make sure details window is pointing to our object
-							DetailsView->SetObjects( DetailsObjects );
-						}
-					}
+					CurrentSelectionType = ESelectionType::Rows;
+
+					RefreshRowSelectionDetails();
 				})
     			.OnGenerateRow_Raw(this, &FChooserTableEditor::GenerateTableRow)
 				.HeaderRow(HeaderRow);
@@ -1109,6 +1240,7 @@ void FChooserTableEditor::SelectColumn(UChooserTable* ChooserEditor, int Index)
    		SelectedColumn->Chooser = Chooser;
    		SelectedColumn->Column = Index;
    		DetailsView->SetObject(SelectedColumn, true);
+   		CurrentSelectionType = ESelectionType::Column;
    	}
    	else
    	{
diff --git a/Engine/Plugins/Experimental/Chooser/Source/ChooserEditor/Private/ChooserTableEditor.h b/Engine/Plugins/Experimental/Chooser/Source/ChooserEditor/Private/ChooserTableEditor.h
--- a/Engine/Plugins/Experimental/Chooser/Source/ChooserEditor/Private/ChooserTableEditor.h
+++ b/Engine/Plugins/Experimental/Chooser/Source/ChooserEditor/Private/ChooserTableEditor.h
@@ -97,14 +97,24 @@ namespace UE::ChooserEditor
 		void ClearSelectedColumn();
 		void DeleteColumn(int Index);
 		void AddColumn(const UScriptStruct* ColumnType);
+		void RefreshRowSelectionDetails();
 		int MoveRow(int SourceRowIndex, int TargetIndex);
 		void SelectRow(int32 RowIndex, bool bClear = true);
 		void ClearSelectedRows(); 
+		bool IsRowSelected(int32 RowIndex);
+
+		enum class ESelectionType
+		{
+			Root, Rows, Column
+		};
+		
+		ESelectionType GetCurrentSelectionType() const { return CurrentSelectionType; }
 	private:
 
 		void SelectRootProperties();
 		void RegisterToolbar();
 		void BindCommands();
+		void OnObjectsTransacted(UObject* Object, const FTransactionObjectEvent& Event);
 		void MakeDebugTargetMenu(UToolMenu* InToolMenu);
 	
 		/** Create the properties tab and its content */
@@ -147,6 +157,8 @@ namespace UE::ChooserEditor
 		TSharedPtr<SHeaderRow> HeaderRow;
 		TSharedPtr<SListView<TSharedPtr<FChooserTableRow>>> TableView;
 
+		ESelectionType CurrentSelectionType = ESelectionType::Root;
+
 	public:
 
 		TSharedPtr<SComboButton>& GetCreateRowComboButton() { return CreateRowComboButton; };
diff --git a/Engine/Plugins/Experimental/Chooser/Source/ChooserEditor/Private/SChooserRowHandle.cpp b/Engine/Plugins/Experimental/Chooser/Source/ChooserEditor/Private/SChooserRowHandle.cpp
--- a/Engine/Plugins/Experimental/Chooser/Source/ChooserEditor/Private/SChooserRowHandle.cpp
+++ b/Engine/Plugins/Experimental/Chooser/Source/ChooserEditor/Private/SChooserRowHandle.cpp
@@ -68,8 +68,16 @@ namespace UE::ChooserEditor
 
 	FReply SChooserRowHandle::OnMouseButtonDown(const FGeometry& MyGeometry, const FPointerEvent& MouseEvent)
 	{
-		ChooserEditor->SelectRow(RowIndex);
-		return FReply::Handled().DetectDrag(SharedThis(this), EKeys::LeftMouseButton);
+		// act as a move handle if the row is already selected, and if there are multiselect modifiers pressed
+		if (!MouseEvent.IsControlDown() && !MouseEvent.IsShiftDown()
+			&& ChooserEditor->IsRowSelected(RowIndex))
+		{
+			return FReply::Handled().DetectDrag(SharedThis(this), EKeys::LeftMouseButton);
+		}
+		else
+		{
+			return SCompoundWidget::OnMouseButtonDown(MyGeometry, MouseEvent);
+		}
 	};
 
 	FReply SChooserRowHandle::OnDragDetected(const FGeometry& MyGeometry, const FPointerEvent& MouseEvent)
