Fix ChooserTable crash on asset reload

diff --git a/Engine/Plugins/Experimental/Chooser/Source/Chooser/Private/Chooser.cpp b/Engine/Plugins/Experimental/Chooser/Source/Chooser/Private/Chooser.cpp
--- a/Engine/Plugins/Experimental/Chooser/Source/Chooser/Private/Chooser.cpp
+++ b/Engine/Plugins/Experimental/Chooser/Source/Chooser/Private/Chooser.cpp
@@ -184,26 +184,34 @@ void UChooserTable::PostEditChangeProperty(FPropertyChangedEvent& PropertyChange
 {
 	UObject::PostEditChangeProperty(PropertyChangedEvent);
 	
-	static FName OutputObjectTypeName = "OutputObjectType";
-	static FName ResultTypeName = "ResultType";
-	if (PropertyChangedEvent.Property->GetName() == OutputObjectTypeName)
+	if (PropertyChangedEvent.Property)
 	{
-		if (CachedPreviousOutputObjectType != OutputObjectType)
+		static FName OutputObjectTypeName = "OutputObjectType";
+		static FName ResultTypeName = "ResultType";
+		if (PropertyChangedEvent.Property->GetName() == OutputObjectTypeName)
 		{
-			OnOutputObjectTypeChanged.Broadcast(OutputObjectType);
-			CachedPreviousOutputObjectType = OutputObjectType;
+			if (CachedPreviousOutputObjectType != OutputObjectType)
+			{
+				OnOutputObjectTypeChanged.Broadcast(OutputObjectType);
+				CachedPreviousOutputObjectType = OutputObjectType;
+			}
 		}
-	}
-	else if (PropertyChangedEvent.Property->GetName() == ResultTypeName)
-	{
-		if (CachedPreviousResultType != ResultType)
+		else if (PropertyChangedEvent.Property->GetName() == ResultTypeName)
 		{
-			OnOutputObjectTypeChanged.Broadcast(OutputObjectType);
-			CachedPreviousResultType = ResultType;
+			if (CachedPreviousResultType != ResultType)
+			{
+				OnOutputObjectTypeChanged.Broadcast(OutputObjectType);
+				CachedPreviousResultType = ResultType;
+			}
+		}
+		else
+		{
+			OnContextClassChanged.Broadcast();
 		}
 	}
 	else
 	{
+		OnOutputObjectTypeChanged.Broadcast(OutputObjectType);
 		OnContextClassChanged.Broadcast();
 	}
 }
