fix for evaluate chooser blueprint nodes which have an object context entry that isn't the first one

diff --git a/Engine/Plugins/Experimental/Chooser/Source/Chooser/Private/ChooserTrace.cpp b/Engine/Plugins/Experimental/Chooser/Source/Chooser/Private/ChooserTrace.cpp
--- a/Engine/Plugins/Experimental/Chooser/Source/Chooser/Private/ChooserTrace.cpp
+++ b/Engine/Plugins/Experimental/Chooser/Source/Chooser/Private/ChooserTrace.cpp
@@ -24,25 +24,35 @@ UE_TRACE_EVENT_BEGIN(Chooser, ChooserValue)
 	UE_TRACE_EVENT_FIELD(UE::Trace::WideString, Key)
 UE_TRACE_EVENT_END()
 
-void FChooserTrace::OutputChooserValueArchive(const FChooserEvaluationContext& Context, const TCHAR* InKey, const FBufferArchive& InValueArchive)
+namespace
 {
-	UObject* ContextObject = nullptr;
-	if (!Context.Params.IsEmpty())
+	UObject* GetContextObject(const FChooserEvaluationContext& Context)
 	{
-		if (const FChooserEvaluationInputObject* ContextObjectInput = Context.Params[0].GetPtr<FChooserEvaluationInputObject>())
+		for (const FStructView& ContextEntry : Context.Params)
 		{
-			ContextObject = ContextObjectInput->Object;
+			if (const FChooserEvaluationInputObject* ContextObjectInput = ContextEntry.GetPtr<FChooserEvaluationInputObject>())
+			{
+				return ContextObjectInput->Object;
+			}
 		}
+		return nullptr;
 	}
-	
+}
+
+
+
+void FChooserTrace::OutputChooserValueArchive(const FChooserEvaluationContext& Context, const TCHAR* InKey, const FBufferArchive& InValueArchive)
+{
 	bool bChannelEnabled = UE_TRACE_CHANNELEXPR_IS_ENABLED(ChooserChannel);
 
-	if (!bChannelEnabled || ContextObject == nullptr || Context.DebuggingInfo.CurrentChooser == nullptr)
+	if (!bChannelEnabled  || Context.DebuggingInfo.CurrentChooser == nullptr)
 	{
 		return;
 	}
 
-	if (CANNOT_TRACE_OBJECT(ContextObject->GetWorld()))
+	const UObject* ContextObject = GetContextObject(Context);
+
+	if (ContextObject == nullptr || CANNOT_TRACE_OBJECT(ContextObject->GetWorld()))
 	{
 		return;
 	}
@@ -60,22 +70,15 @@ void FChooserTrace::OutputChooserValueArchive(const FChooserEvaluationContext& C
 
 void FChooserTrace::OutputChooserEvaluation(const UObject* InChooser, const FChooserEvaluationContext& Context, uint32 InSelectedIndex)
 {
-	UObject* ContextObject = nullptr;
-	if (!Context.Params.IsEmpty())
-	{
-		if (const FChooserEvaluationInputObject* ContextObjectInput = Context.Params[0].GetPtr<FChooserEvaluationInputObject>())
-		{
-			ContextObject = ContextObjectInput->Object;
-		}
-	}
-
 	bool bChannelEnabled = UE_TRACE_CHANNELEXPR_IS_ENABLED(ChooserChannel);
-	if (!bChannelEnabled || ContextObject == nullptr || InChooser == nullptr)
+	if (!bChannelEnabled || InChooser == nullptr)
 	{
 		return;
 	}
+	
+	const UObject* ContextObject = GetContextObject(Context);
 
-	if (CANNOT_TRACE_OBJECT(ContextObject->GetWorld()))
+	if (ContextObject == nullptr || CANNOT_TRACE_OBJECT(ContextObject->GetWorld()))
 	{
 		return;
 	}
diff --git a/Engine/Plugins/Experimental/Chooser/Source/ChooserUncooked/Private/EvaluateChooserNode.cpp b/Engine/Plugins/Experimental/Chooser/Source/ChooserUncooked/Private/EvaluateChooserNode.cpp
--- a/Engine/Plugins/Experimental/Chooser/Source/ChooserUncooked/Private/EvaluateChooserNode.cpp
+++ b/Engine/Plugins/Experimental/Chooser/Source/ChooserUncooked/Private/EvaluateChooserNode.cpp
@@ -615,6 +615,7 @@ void UK2Node_EvaluateChooser2::ExpandNode(class FKismetCompilerContext& Compiler
 		
 		if (Chooser)
 		{
+			bool bFoundObject = false;
 			int ContextDataCount = Chooser->ContextData.Num();
 			for(int ContextDataIndex = 0; ContextDataIndex < ContextDataCount; ContextDataIndex++)
 			{
@@ -624,6 +625,7 @@ void UK2Node_EvaluateChooser2::ExpandNode(class FKismetCompilerContext& Compiler
 					const UScriptStruct* EntryType = ContextDataEntry.GetScriptStruct();
 					if (EntryType == FContextObjectTypeClass::StaticStruct())
 					{
+						bFoundObject = true;
 						const FContextObjectTypeClass& ClassContext = ContextDataEntry.Get<FContextObjectTypeClass>();
 						if (ClassContext.Class)
 						{
@@ -634,7 +636,6 @@ void UK2Node_EvaluateChooser2::ExpandNode(class FKismetCompilerContext& Compiler
 								AddObjectFunction->SetFromFunction(UChooserFunctionLibrary::StaticClass()->FindFunctionByName(GET_MEMBER_NAME_CHECKED(UChooserFunctionLibrary, AddChooserObjectInput)));
 								AddObjectFunction->AllocateDefaultPins();
 								ContextStructPin->MakeLinkTo(AddObjectFunction->FindPin(FName("Context")));
-								PreviousNodeExecOutput = ContextStructNode->GetThenPin();
 
 								PreviousNodeExecOutput->MakeLinkTo(AddObjectFunction->GetExecPin());
 								PreviousNodeExecOutput = AddObjectFunction->GetThenPin();
@@ -725,6 +726,23 @@ void UK2Node_EvaluateChooser2::ExpandNode(class FKismetCompilerContext& Compiler
 				}
 			}
 
+			if (!bFoundObject)
+			{
+				// add Self reference to the end of the context, for debugging purposes
+				UK2Node_CallFunction* AddObjectFunction = CompilerContext.SpawnIntermediateNode<UK2Node_CallFunction>(this, SourceGraph);
+				CompilerContext.MessageLog.NotifyIntermediateObjectCreation(AddObjectFunction, this);
+				AddObjectFunction->SetFromFunction(UChooserFunctionLibrary::StaticClass()->FindFunctionByName(GET_MEMBER_NAME_CHECKED(UChooserFunctionLibrary, AddChooserObjectInput)));
+				AddObjectFunction->AllocateDefaultPins();
+				ContextStructPin->MakeLinkTo(AddObjectFunction->FindPin(FName("Context")));
+	
+				PreviousNodeExecOutput->MakeLinkTo(AddObjectFunction->GetExecPin());
+				PreviousNodeExecOutput = AddObjectFunction->GetThenPin();
+									
+				UEdGraphPin* AddObjectPin = AddObjectFunction->FindPin(FName("Object"));
+				SelfPin->MakeLinkTo(AddObjectPin);
+			}
+			
+
 			if (PreviousNodeExecOutput)
 			{
 				PreviousNodeExecOutput->MakeLinkTo(CallFunction->GetExecPin());
