Fix crash dragging/dropping into nested chooser

diff --git a/Engine/Plugins/Experimental/Chooser/Source/ChooserEditor/Private/SChooserTableRow.cpp b/Engine/Plugins/Experimental/Chooser/Source/ChooserEditor/Private/SChooserTableRow.cpp
--- a/Engine/Plugins/Experimental/Chooser/Source/ChooserEditor/Private/SChooserTableRow.cpp
+++ b/Engine/Plugins/Experimental/Chooser/Source/ChooserEditor/Private/SChooserTableRow.cpp
@@ -214,18 +214,13 @@ namespace UE::ChooserEditor
 			{
 				bDropSupported = true;
 			
-				if (Chooser->OutputObjectType) // if OutputObjectType is null, then any kind of object is supported, don't need to check all of them
+				UChooserTable* ContextOwner = Chooser->GetContextOwner();
+				if (ContextOwner->OutputObjectType) // if OutputObjectType is null, then any kind of object is supported, don't need to check all of them
 				{
 					for (const FAssetData& Asset : ContentDragDropOp->GetAssets())
 					{
 						const UClass* AssetClass = Asset.GetClass();
 
-						if(AssetClass->IsChildOf(Chooser->OutputObjectType))
-						{
-							bDropSupported = false;
-							break;
-						}
-						
 						if(AssetClass->IsChildOf(UChooserTable::StaticClass()))
 						{
 							const UChooserTable* DraggedChooserTable = Cast<UChooserTable>(Asset.GetAsset());
@@ -233,12 +228,17 @@ namespace UE::ChooserEditor
 							// verify dragged chooser result type matches this chooser result type
 							if (DraggedChooserTable->ResultType == EObjectChooserResultType::ClassResult
 							    || DraggedChooserTable->OutputObjectType == nullptr
-							    || !DraggedChooserTable->OutputObjectType->IsChildOf(Chooser->OutputObjectType))
+							    || !DraggedChooserTable->OutputObjectType->IsChildOf(ContextOwner->OutputObjectType))
 							{
 								bDropSupported = false;
 								break;
 							}
 						}
+						else if(!AssetClass->IsChildOf(ContextOwner->OutputObjectType))
+                        {
+							bDropSupported = false;
+							break;
+                        }
 					}
 				}
 			}
@@ -255,14 +255,17 @@ namespace UE::ChooserEditor
 
 	FReply SChooserTableRow::OnDragOver(const FGeometry& MyGeometry, const FDragDropEvent& DragDropEvent)
 	{
-		if (DragDropEvent.GetOperationAs<FChooserRowDragDropOp>()
-			|| DragDropEvent.GetOperationAs<FAssetDragDropOp>())
+		if (bDropSupported)
 		{
+			bDragActive = true;
 			float Center = MyGeometry.AbsolutePosition.Y + MyGeometry.Size.Y/2;
 			bDropAbove = DragDropEvent.GetScreenSpacePosition().Y < Center;
 			return FReply::Handled();
 		}
-		return FReply::Unhandled();
+		else
+		{
+			return FReply::Unhandled();
+		}
 	}
 
 	FReply SChooserTableRow::OnDrop(const FGeometry& MyGeometry, const FDragDropEvent& DragDropEvent)
@@ -297,9 +300,8 @@ namespace UE::ChooserEditor
 		}
 		else if (TSharedPtr<FAssetDragDropOp> ContentDragDropOp = DragDropEvent.GetOperationAs<FAssetDragDropOp>())
 		{
-			
 			FScopedTransaction ScopedTransaction(LOCTEXT("DragDropAssets","Drag and Drop Assets into Chooser"));
-			Chooser->Modify();
+			Chooser->Modify(true);
 			
 			if (Chooser->ResultType == EObjectChooserResultType::ObjectResult)
 			{
@@ -334,9 +336,12 @@ namespace UE::ChooserEditor
 						}
 						else
 						{
-							check(AssetClass->IsChildOf(Chooser->OutputObjectType));
-							NewResult.InitializeAs(FAssetChooser::StaticStruct());
-							NewResult.GetMutable<FAssetChooser>().Asset = Asset.GetAsset();
+							UChooserTable* ContextOwner = Chooser->GetContextOwner();
+							if (ContextOwner->OutputObjectType == nullptr || ensure(AssetClass->IsChildOf(ContextOwner->OutputObjectType)))
+							{
+								NewResult.InitializeAs(FAssetChooser::StaticStruct());
+								NewResult.GetMutable<FAssetChooser>().Asset = Asset.GetAsset();
+							}
 						}
 					}
 
