Moved data fixup from PostLoad to Serialize.

UK2Node_EvaluateProxy2::AllocateDefaultPins was calling postload on its proxy, which was causing issues with editor async loading. Specifically, we don't expect any postload calls to be threadsafe unless the UObject returns true for IsPostLoadThreadSafe. Additionally, we've never postloaded Blueprints and their dependencies during a compilation pass (only preloaded). Calling postload here would be a violation of that implicit rule.

We can easily resolve this by moving the data fixup code in UProxyAsset::PostLoad into Serialize, which is where we normally do that anyway. This alleviates the need for a manual postload call since the proxy is guaranteed to be fully preloaded with data fixup applied.

#rnx

diff --git a/Engine/Plugins/Experimental/Chooser/Source/ProxyTable/Private/ProxyAsset.cpp b/Engine/Plugins/Experimental/Chooser/Source/ProxyTable/Private/ProxyAsset.cpp
--- a/Engine/Plugins/Experimental/Chooser/Source/ProxyTable/Private/ProxyAsset.cpp
+++ b/Engine/Plugins/Experimental/Chooser/Source/ProxyTable/Private/ProxyAsset.cpp
@@ -71,6 +71,33 @@ void UProxyAsset::PostEditChangeProperty(FPropertyChangedEvent& PropertyChangedE
 /////////////////////////////////////////////////////////////////////////////////////////
 // Proxy Asset
 
+void UProxyAsset::Serialize(FArchive& Ar)
+{
+	Super::Serialize(Ar);
+
+#if WITH_EDITORONLY_DATA
+	if (Ar.IsLoading())
+	{
+		if (ContextClass_DEPRECATED)
+		{
+			ContextData.SetNum(1);
+			ContextData[0].InitializeAs<FContextObjectTypeClass>();
+			FContextObjectTypeClass& Context = ContextData[0].GetMutable<FContextObjectTypeClass>();
+			Context.Class = ContextClass_DEPRECATED;
+			Context.Direction = EContextObjectDirection::ReadWrite;
+			ContextClass_DEPRECATED = nullptr;
+		}
+
+		if (!Guid.IsValid())
+		{
+			// if we load a ProxyAsset that was created before the Guid, assign it a deterministic guid based on the name and path.
+			Guid.A = GetTypeHash(GetName());
+			Guid.B = GetTypeHash(GetPackage()->GetPathName());
+		}
+	}
+#endif
+}
+
 void UProxyAsset::PostLoad()
 {
 	Super::PostLoad();
@@ -79,24 +106,6 @@ void UProxyAsset::PostLoad()
 	CachedPreviousType = Type;
 	CachedPreviousResultType = ResultType;
 #endif
-	
-	if (ContextClass_DEPRECATED)
-	{
-		
-		ContextData.SetNum(1);
-		ContextData[0].InitializeAs<FContextObjectTypeClass>();
-		FContextObjectTypeClass& Context = ContextData[0].GetMutable<FContextObjectTypeClass>();
-		Context.Class = ContextClass_DEPRECATED;
-		Context.Direction = EContextObjectDirection::ReadWrite;
-		ContextClass_DEPRECATED = nullptr;
-	}
-
-	if (!Guid.IsValid())
-	{
-		// if we load a ProxyAsset that was created before the Guid, assign it a deterministic guid based on the name and path.
-		Guid.A = GetTypeHash(GetName());
-		Guid.B = GetTypeHash(GetPackage()->GetPathName());
-	}
 
 	if (ProxyTable.IsValid())
 	{
diff --git a/Engine/Plugins/Experimental/Chooser/Source/ProxyTable/Public/ProxyAsset.h b/Engine/Plugins/Experimental/Chooser/Source/ProxyTable/Public/ProxyAsset.h
--- a/Engine/Plugins/Experimental/Chooser/Source/ProxyTable/Public/ProxyAsset.h
+++ b/Engine/Plugins/Experimental/Chooser/Source/ProxyTable/Public/ProxyAsset.h
@@ -56,6 +56,7 @@ public:
 	UPROPERTY()
 	TObjectPtr<UClass> ContextClass_DEPRECATED;
 	
+	virtual void Serialize(FArchive& Ar) override;
 	virtual void PostLoad() override;
 	virtual void PostDuplicate(EDuplicateMode::Type DuplicateMode) override;
 #endif
diff --git a/Engine/Plugins/Experimental/Chooser/Source/ProxyTableUncooked/Private/EvaluateProxyNode.cpp b/Engine/Plugins/Experimental/Chooser/Source/ProxyTableUncooked/Private/EvaluateProxyNode.cpp
--- a/Engine/Plugins/Experimental/Chooser/Source/ProxyTableUncooked/Private/EvaluateProxyNode.cpp
+++ b/Engine/Plugins/Experimental/Chooser/Source/ProxyTableUncooked/Private/EvaluateProxyNode.cpp
@@ -394,10 +394,7 @@ void UK2Node_EvaluateProxy2::AllocateDefaultPins()
 
 	if (Proxy)
 	{
-		// ensure any data upgrades have been applied to Proxy before generating pins
-		Proxy->ConditionalPostLoad();
-
-		for(FInstancedStruct& ContextDataEntry : Proxy->ContextData)
+		for (FInstancedStruct& ContextDataEntry : Proxy->ContextData)
 		{
 			if (ContextDataEntry.IsValid())
 			{
