Wrapping support in chooser float range column

diff --git a/Engine/Plugins/Experimental/Chooser/Source/Chooser/Internal/FloatRangeColumn.h b/Engine/Plugins/Experimental/Chooser/Source/Chooser/Internal/FloatRangeColumn.h
--- a/Engine/Plugins/Experimental/Chooser/Source/Chooser/Internal/FloatRangeColumn.h
+++ b/Engine/Plugins/Experimental/Chooser/Source/Chooser/Internal/FloatRangeColumn.h
@@ -59,7 +59,17 @@ struct CHOOSER_API FFloatRangeColumn : public FChooserColumnBase
 		
 	UPROPERTY(EditAnywhere, NoClear, Meta = (ExcludeBaseStruct, BaseStruct = "/Script/Chooser.ChooserParameterFloatBase"), Category = "Data")
 	FInstancedStruct InputValue;
-	
+
+	// Wrap input, and comparisons for numbers such as angles which 
+	UPROPERTY(EditAnywhere, Category = "Data");
+	bool bWrapInput = false;
+	// Minimum value (for WrapInput)
+	UPROPERTY(EditAnywhere, Category = "Data", meta=(DisplayAfter="bWrapInput", EditCondition="bWrapInput"));
+	double MinValue = -180;
+	// Maximum value (for WrapInput)
+	UPROPERTY(EditAnywhere, Category = "Data", meta=(DisplayAfter="bWrapInput", EditCondition="bWrapInput"));
+	double MaxValue = 180;
+
 #if WITH_EDITORONLY_DATA
 	UPROPERTY(EditAnywhere, Category="Data")
 	FChooserFloatRangeRowData DefaultRowValue;
@@ -74,10 +84,7 @@ struct CHOOSER_API FFloatRangeColumn : public FChooserColumnBase
 
 #if WITH_EDITOR
 	mutable double TestValue = 0.0;
-	virtual bool EditorTestFilter(int32 RowIndex) const override
-	{
-		return RowValues.IsValidIndex(RowIndex) && TestValue >= RowValues[RowIndex].Min && TestValue <= RowValues[RowIndex].Max;
-	}
+	virtual bool EditorTestFilter(int32 RowIndex) const override;
 	
 	virtual void SetTestValue(TArrayView<const uint8> Value) override
 	{
diff --git a/Engine/Plugins/Experimental/Chooser/Source/Chooser/Private/FloatRangeColumn.cpp b/Engine/Plugins/Experimental/Chooser/Source/Chooser/Private/FloatRangeColumn.cpp
--- a/Engine/Plugins/Experimental/Chooser/Source/Chooser/Private/FloatRangeColumn.cpp
+++ b/Engine/Plugins/Experimental/Chooser/Source/Chooser/Private/FloatRangeColumn.cpp
@@ -22,6 +22,33 @@ FFloatRangeColumn::FFloatRangeColumn()
 	InputValue.InitializeAs(FFloatContextProperty::StaticStruct());
 }
 
+#if WITH_EDITOR
+bool FFloatRangeColumn::EditorTestFilter(int32 RowIndex) const
+{
+	if (RowValues.IsValidIndex(RowIndex))
+	{
+		if (bWrapInput)
+		{
+			double TestResult = FMath::Wrap(TestValue, MinValue, MaxValue);
+			const FChooserFloatRangeRowData& RowValue = RowValues[RowIndex];
+			if (RowValue.Max < RowValue.Min) // eg for an angle range from  135 to -135  (135 to 180 or -180 to -135)
+			{
+				return TestResult >= RowValue.Min || TestResult <= RowValue.Max;
+			}
+			else
+			{
+				return TestResult >= RowValue.Min && TestResult <= RowValue.Max;
+			}
+		}
+		else
+		{
+			return TestValue >= RowValues[RowIndex].Min && TestValue <= RowValues[RowIndex].Max;
+		}
+	}
+	return false;
+}
+#endif
+
 void FFloatRangeColumn::Filter(FChooserEvaluationContext& Context, const FChooserIndexArray& IndexListIn, FChooserIndexArray& IndexListOut) const
 {
 	if (InputValue.IsValid())
@@ -38,17 +65,43 @@ void FFloatRangeColumn::Filter(FChooserEvaluationContext& Context, const FChoose
 		}
 #endif
 
-		for(uint32 Index : IndexListIn)
+		if (bWrapInput)
 		{
-			if (RowValues.Num() > (int)Index)
+			Result = FMath::Wrap(Result, MinValue, MaxValue);
+			for(uint32 Index : IndexListIn)
 			{
-				const FChooserFloatRangeRowData& RowValue = RowValues[Index];
-				if (Result >= RowValue.Min && Result <= RowValue.Max)
+				if (RowValues.Num() > static_cast<int>(Index))
 				{
-					IndexListOut.Push(Index);
+					const FChooserFloatRangeRowData& RowValue = RowValues[Index];
+					if (RowValue.Max < RowValue.Min) // eg for an angle range from  135 to -135  (135 to 180 or -180 to -135)
+					{
+						if (Result >= RowValue.Min || Result <= RowValue.Max)
+						{
+							IndexListOut.Push(Index);
+						}
+					}
+					else if (Result >= RowValue.Min && Result <= RowValue.Max)
+					{
+						IndexListOut.Push(Index);
+					}
 				}
 			}
 		}
+		else
+		{
+			for(uint32 Index : IndexListIn)
+			{
+				if (RowValues.Num() > static_cast<int>(Index))
+				{
+					const FChooserFloatRangeRowData& RowValue = RowValues[Index];
+					if (Result >= RowValue.Min && Result <= RowValue.Max)
+					{
+						IndexListOut.Push(Index);
+					}
+				}
+			}
+		}
+		
 	}
 	else
 	{
