Asset Drag/Drop improvements in Chooser Table

diff --git a/Engine/Plugins/Experimental/Chooser/Source/ChooserEditor/Private/ChooserEditorWidgets.cpp b/Engine/Plugins/Experimental/Chooser/Source/ChooserEditor/Private/ChooserEditorWidgets.cpp
--- a/Engine/Plugins/Experimental/Chooser/Source/ChooserEditor/Private/ChooserEditorWidgets.cpp
+++ b/Engine/Plugins/Experimental/Chooser/Source/ChooserEditor/Private/ChooserEditorWidgets.cpp
@@ -35,6 +35,7 @@ TSharedRef<SWidget> CreateAssetWidget(bool bReadOnly, UObject* TransactionObject
 		.IsEnabled(!bReadOnly)
 		.AllowedClass(ResultBaseClass ? ResultBaseClass : UObject::StaticClass())
 		.ObjectPath_Lambda([DIAsset](){ return DIAsset->Asset ? DIAsset->Asset.GetPath() : "";})
+		.bOnlyRecognizeOnDragEnter(true)
 		.OnObjectChanged_Lambda([TransactionObject, DIAsset](const FAssetData& AssetData)
 		{
 			const FScopedTransaction Transaction(LOCTEXT("Edit Asset", "Edit Asset"));
@@ -71,6 +72,7 @@ TSharedRef<SWidget> CreateEvaluateChooserWidget(bool bReadOnly, UObject* Transac
 	return SNew(SObjectPropertyEntryBox)
 		.IsEnabled(!bReadOnly)
 		.AllowedClass(UChooserTable::StaticClass())
+		.bOnlyRecognizeOnDragEnter(true)
 		.ObjectPath_Lambda([EvaluateChooser](){ return EvaluateChooser->Chooser ? EvaluateChooser->Chooser.GetPath() : "";})
 		.OnShouldFilterAsset_Lambda([ResultBaseClass](const FAssetData& InAssetData)
 		{
diff --git a/Engine/Plugins/Experimental/Chooser/Source/ChooserEditor/Private/SChooserTableRow.cpp b/Engine/Plugins/Experimental/Chooser/Source/ChooserEditor/Private/SChooserTableRow.cpp
--- a/Engine/Plugins/Experimental/Chooser/Source/ChooserEditor/Private/SChooserTableRow.cpp
+++ b/Engine/Plugins/Experimental/Chooser/Source/ChooserEditor/Private/SChooserTableRow.cpp
@@ -257,6 +257,7 @@ namespace UE::ChooserEditor
 	{
 		if (bDropSupported)
 		{
+			DragActiveCounter = 2;
 			bDragActive = true;
 			float Center = MyGeometry.AbsolutePosition.Y + MyGeometry.Size.Y/2;
 			bDropAbove = DragDropEvent.GetScreenSpacePosition().Y < Center;
@@ -368,6 +369,18 @@ namespace UE::ChooserEditor
 		return FReply::Handled();		
 	}
 
+	
+	void SChooserTableRow::Tick(const FGeometry& AllottedGeometry, const double InCurrentTime, const float InDeltaTime)
+	{
+		DragActiveCounter --;
+		if (DragActiveCounter < 0)
+		{
+			DragActiveCounter = 0;
+			bDragActive = false;
+		}
+		
+		SMultiColumnTableRow<TSharedPtr<FChooserTableRow, ESPMode::ThreadSafe>>::Tick(AllottedGeometry, InCurrentTime, InDeltaTime);
+	}
 }
 
 #undef LOCTEXT_NAMESPACE
diff --git a/Engine/Plugins/Experimental/Chooser/Source/ChooserEditor/Private/SChooserTableRow.h b/Engine/Plugins/Experimental/Chooser/Source/ChooserEditor/Private/SChooserTableRow.h
--- a/Engine/Plugins/Experimental/Chooser/Source/ChooserEditor/Private/SChooserTableRow.h
+++ b/Engine/Plugins/Experimental/Chooser/Source/ChooserEditor/Private/SChooserTableRow.h
@@ -37,12 +37,14 @@ SLATE_BEGIN_ARGS(SChooserTableRow) {}
 	virtual void OnDragLeave(const FDragDropEvent& DragDropEvent) override;
 	virtual FReply OnDragOver(const FGeometry& MyGeometry, const FDragDropEvent& DragDropEvent) override;
 	virtual FReply OnDrop(const FGeometry& MyGeometry, const FDragDropEvent& DragDropEvent) override;
+	virtual void Tick(const FGeometry& AllottedGeometry, const double InCurrentTime, const float InDeltaTime) override;
 
 private:
 	TSharedPtr<FChooserTableRow> RowIndex;
 	UChooserTable* Chooser;
 	FChooserTableEditor* Editor;
 	TSharedPtr<SBorder> CacheBorder;
+	int DragActiveCounter = 0;
 	bool bDragActive = false;
 	bool bDropSupported = false;
 	bool bDropAbove = false;
diff --git a/Engine/Source/Editor/EditorWidgets/Private/SAssetDropTarget.cpp b/Engine/Source/Editor/EditorWidgets/Private/SAssetDropTarget.cpp
--- a/Engine/Source/Editor/EditorWidgets/Private/SAssetDropTarget.cpp
+++ b/Engine/Source/Editor/EditorWidgets/Private/SAssetDropTarget.cpp
@@ -27,6 +27,7 @@ void SAssetDropTarget::Construct(const FArguments& InArgs )
 
 	SDropTarget::Construct(
 		SDropTarget::FArguments()
+		.bOnlyRecognizeOnDragEnter(InArgs._bOnlyRecognizeOnDragEnter)
 		.OnDropped(this, &SAssetDropTarget::OnDropped)
 		[
 			InArgs._Content.Widget
diff --git a/Engine/Source/Editor/EditorWidgets/Public/SAssetDropTarget.h b/Engine/Source/Editor/EditorWidgets/Public/SAssetDropTarget.h
--- a/Engine/Source/Editor/EditorWidgets/Public/SAssetDropTarget.h
+++ b/Engine/Source/Editor/EditorWidgets/Public/SAssetDropTarget.h
@@ -57,6 +57,8 @@ public:
 		SLATE_EVENT( FAreAssetsAcceptableForDropWithReason, OnAreAssetsAcceptableForDropWithReason )
 		/** Sets if this drop target can support multiple assets dropped, or only supports a single asset dropped at a time. False by default for legacy behavior. */
 		SLATE_ARGUMENT( bool, bSupportsMultiDrop )
+		/** When this is true, the drop target will only get recognized when entering while drag & dropping. */
+		SLATE_ATTRIBUTE(bool, bOnlyRecognizeOnDragEnter)
 
 		FOnAssetsDropped ConvertObjectDropDelegate(const FOnAssetDropped& LegacyDelegate)
 		{
diff --git a/Engine/Source/Editor/PropertyEditor/Private/PropertyCustomizationHelpers.cpp b/Engine/Source/Editor/PropertyEditor/Private/PropertyCustomizationHelpers.cpp
--- a/Engine/Source/Editor/PropertyEditor/Private/PropertyCustomizationHelpers.cpp
+++ b/Engine/Source/Editor/PropertyEditor/Private/PropertyCustomizationHelpers.cpp
@@ -582,6 +582,7 @@ void SObjectPropertyEntryBox::Construct( const FArguments& InArgs )
 				.ThumbnailSize(ThumbnailSize)
 				.DisplayCompactSize(InArgs._DisplayCompactSize)
 				.OnShouldFilterActor(InArgs._OnShouldFilterActor)
+				.bOnlyRecognizeOnDragEnter(InArgs._bOnlyRecognizeOnDragEnter)
 				.CustomContentSlot()
 				[
 					InArgs._CustomContentSlot.Widget
diff --git a/Engine/Source/Editor/PropertyEditor/Private/UserInterface/PropertyEditor/SPropertyEditorAsset.cpp b/Engine/Source/Editor/PropertyEditor/Private/UserInterface/PropertyEditor/SPropertyEditorAsset.cpp
--- a/Engine/Source/Editor/PropertyEditor/Private/UserInterface/PropertyEditor/SPropertyEditorAsset.cpp
+++ b/Engine/Source/Editor/PropertyEditor/Private/UserInterface/PropertyEditor/SPropertyEditorAsset.cpp
@@ -438,6 +438,7 @@ void SPropertyEditorAsset::Construct(const FArguments& InArgs, const TSharedPtr<
 	ChildSlot
 	[
 		SNew( SAssetDropTarget )
+		.bOnlyRecognizeOnDragEnter(InArgs._bOnlyRecognizeOnDragEnter)
 		.OnAreAssetsAcceptableForDropWithReason( this, &SPropertyEditorAsset::OnAssetDraggedOver )
 		.OnAssetsDropped( this, &SPropertyEditorAsset::OnAssetDropped )
 		[
diff --git a/Engine/Source/Editor/PropertyEditor/Private/UserInterface/PropertyEditor/SPropertyEditorAsset.h b/Engine/Source/Editor/PropertyEditor/Private/UserInterface/PropertyEditor/SPropertyEditorAsset.h
--- a/Engine/Source/Editor/PropertyEditor/Private/UserInterface/PropertyEditor/SPropertyEditorAsset.h
+++ b/Engine/Source/Editor/PropertyEditor/Private/UserInterface/PropertyEditor/SPropertyEditorAsset.h
@@ -83,6 +83,8 @@ public:
 		SLATE_ARGUMENT(TArray<FAssetData>, OwnerAssetDataArray)
 		SLATE_EVENT(FOnShouldFilterActor, OnShouldFilterActor)
 		SLATE_ARGUMENT(TOptional<FDetailWidgetRow*>, InWidgetRow)
+		/** When this is true, the drop target will only get recognized when entering while drag & dropping. */
+		SLATE_ATTRIBUTE(bool, bOnlyRecognizeOnDragEnter)
 
 	SLATE_END_ARGS()
 
diff --git a/Engine/Source/Editor/PropertyEditor/Public/PropertyCustomizationHelpers.h b/Engine/Source/Editor/PropertyEditor/Public/PropertyCustomizationHelpers.h
--- a/Engine/Source/Editor/PropertyEditor/Public/PropertyCustomizationHelpers.h
+++ b/Engine/Source/Editor/PropertyEditor/Public/PropertyCustomizationHelpers.h
@@ -227,6 +227,8 @@ public:
 		SLATE_ATTRIBUTE(FIntPoint, ThumbnailSizeOverride)
 		/** Called to check if an actor is valid to use */
 		SLATE_EVENT(FOnShouldFilterActor, OnShouldFilterActor)
+		/** When this is true, the drop target will only get recognized when entering while drag & dropping. */
+		SLATE_ATTRIBUTE(bool, bOnlyRecognizeOnDragEnter)
 	SLATE_END_ARGS()
 
 	PROPERTYEDITOR_API void Construct( const FArguments& InArgs );
