Multi value enum column for choosers

diff --git a/Engine/Plugins/Experimental/Chooser/Source/Chooser/Internal/EnumColumn.h b/Engine/Plugins/Experimental/Chooser/Source/Chooser/Internal/EnumColumn.h
--- a/Engine/Plugins/Experimental/Chooser/Source/Chooser/Internal/EnumColumn.h
+++ b/Engine/Plugins/Experimental/Chooser/Source/Chooser/Internal/EnumColumn.h
@@ -122,9 +122,7 @@ public:
 	
 	CHOOSER_COLUMN_BOILERPLATE(FChooserParameterEnumBase);
 
-#if WITH_EDITORONLY_DATA
 	virtual void PostLoad() override;
-#endif
 };
 
 // deprecated class version for converting old data
diff --git a/Engine/Plugins/Experimental/Chooser/Source/Chooser/Internal/MultiEnumColumn.h b/Engine/Plugins/Experimental/Chooser/Source/Chooser/Internal/MultiEnumColumn.h
new file mode 100644
--- /dev/null
+++ b/Engine/Plugins/Experimental/Chooser/Source/Chooser/Internal/MultiEnumColumn.h
@@ -0,0 +1,72 @@
+// Copyright Epic Games, Inc. All Rights Reserved.
+#pragma once
+
+#include "CoreMinimal.h"
+#include "EnumColumn.h"
+#include "IChooserColumn.h"
+#include "IChooserParameterEnum.h"
+#include "InstancedStruct.h"
+#include "ChooserPropertyAccess.h"
+#include "Serialization/MemoryReader.h"
+#include "MultiEnumColumn.generated.h"
+
+struct FBindingChainElement;
+
+USTRUCT()
+struct FChooserMultiEnumRowData
+{
+	GENERATED_BODY()
+
+	UPROPERTY(EditAnywhere, Category = Runtime)
+	uint32 Value = 0;
+
+	bool Evaluate(const uint32 LeftHandSide) const
+	{
+		return Value == 0 || Value & LeftHandSide;
+	}
+};
+
+
+USTRUCT()
+struct CHOOSER_API FMultiEnumColumn : public FChooserColumnBase
+{
+	GENERATED_BODY()
+public:
+	FMultiEnumColumn();
+
+	UPROPERTY(EditAnywhere, NoClear, Meta = (ExcludeBaseStruct, BaseStruct = "/Script/Chooser.ChooserParameterEnumBase"), Category = "Data")
+	FInstancedStruct InputValue;
+
+#if WITH_EDITORONLY_DATA
+	UPROPERTY(EditAnywhere, Category = "Data")
+	FChooserMultiEnumRowData DefaultRowValue;
+#endif
+	
+	UPROPERTY(EditAnywhere, Category = "Data")
+	// array of results (cells for this column for each row in the table)
+	// should match the length of the Results array
+	TArray<FChooserMultiEnumRowData> RowValues;
+
+	virtual void Filter(FChooserEvaluationContext& Context, const FChooserIndexArray& IndexListIn, FChooserIndexArray& IndexListOut) const override;
+	
+#if WITH_EDITOR
+	mutable uint8 TestValue = 0;
+	virtual bool EditorTestFilter(int32 RowIndex) const override
+	{
+		return RowValues.IsValidIndex(RowIndex) && RowValues[RowIndex].Evaluate(1 << TestValue);
+	}
+	
+	virtual void SetTestValue(TArrayView<const uint8> Value) override
+	{
+		FMemoryReaderView Reader(Value);
+		Reader << TestValue;
+	}
+	
+	virtual void AddToDetails (FInstancedPropertyBag& PropertyBag, int32 ColumnIndex, int32 RowIndex);
+	virtual void SetFromDetails(FInstancedPropertyBag& PropertyBag, int32 ColumnIndex, int32 RowIndex);
+#endif
+	
+	CHOOSER_COLUMN_BOILERPLATE(FChooserParameterEnumBase);
+
+	virtual void PostLoad() override;
+};
\ No newline at end of file
diff --git a/Engine/Plugins/Experimental/Chooser/Source/Chooser/Private/MultiEnumColumn.cpp b/Engine/Plugins/Experimental/Chooser/Source/Chooser/Private/MultiEnumColumn.cpp
new file mode 100644
--- /dev/null
+++ b/Engine/Plugins/Experimental/Chooser/Source/Chooser/Private/MultiEnumColumn.cpp
@@ -0,0 +1,85 @@
+// Copyright Epic Games, Inc. All Rights Reserved.
+#include "MultiEnumColumn.h"
+#include "ChooserIndexArray.h"
+#include "ChooserTrace.h"
+#if WITH_EDITOR
+#include "PropertyBag.h"
+#endif
+
+FMultiEnumColumn::FMultiEnumColumn()
+{
+#if WITH_EDITOR
+	InputValue.InitializeAs(FEnumContextProperty::StaticStruct());
+#endif
+}
+
+void FMultiEnumColumn::PostLoad()
+{
+	Super::PostLoad();
+	
+	if (InputValue.IsValid())
+	{
+		InputValue.GetMutable<FChooserParameterBase>().PostLoad();
+	}
+}
+
+void FMultiEnumColumn::Filter(FChooserEvaluationContext& Context, const FChooserIndexArray& IndexListIn, FChooserIndexArray& IndexListOut) const
+{
+	uint8 Result = 0;
+	if (InputValue.IsValid() &&
+		InputValue.Get<FChooserParameterEnumBase>().GetValue(Context, Result))
+	{
+#if WITH_EDITOR
+		if (Context.DebuggingInfo.bCurrentDebugTarget)
+		{
+			TestValue = Result;
+		}
+#endif
+
+		TRACE_CHOOSER_VALUE(Context, ToCStr(InputValue.Get<FChooserParameterBase>().GetDebugName()), Result);
+
+		// log if Result > 31
+		uint32 ResultBit = 1 << Result;
+		for (const uint32 Index : IndexListIn)
+		{
+			if (RowValues.IsValidIndex(Index))
+			{
+				const FChooserMultiEnumRowData& RowValue = RowValues[Index];
+				if (RowValue.Evaluate(ResultBit))
+				{
+					IndexListOut.Push(Index);
+				}
+			}
+		}
+	}
+	else
+	{
+		// passthrough fallback (behaves better during live editing)
+		IndexListOut = IndexListIn;
+	}
+}
+
+#if WITH_EDITOR
+	void FMultiEnumColumn::AddToDetails (FInstancedPropertyBag& PropertyBag, int32 ColumnIndex, int32 RowIndex)
+	{
+		FText DisplayName;
+		InputValue.Get<FChooserParameterBase>().GetDisplayName(DisplayName);
+		FName PropertyName("RowData",ColumnIndex);
+
+		FPropertyBagPropertyDesc PropertyDesc(PropertyName,  EPropertyBagPropertyType::Struct, FChooserMultiEnumRowData::StaticStruct());
+		PropertyDesc.MetaData.Add(FPropertyBagPropertyDescMetaData("DisplayName", DisplayName.ToString()));
+		PropertyBag.AddProperties({PropertyDesc});
+		PropertyBag.SetValueStruct(PropertyName, RowValues[RowIndex]);
+	}
+	
+	void FMultiEnumColumn::SetFromDetails(FInstancedPropertyBag& PropertyBag, int32 ColumnIndex, int32 RowIndex)
+	{
+		FName PropertyName("RowData",ColumnIndex);
+
+		TValueOrError<FStructView, EPropertyBagResult> Result = PropertyBag.GetValueStruct(PropertyName, FChooserMultiEnumRowData::StaticStruct());
+		if (const FStructView* StructView = Result.TryGetValue())
+		{
+			RowValues[RowIndex] = StructView->Get<FChooserMultiEnumRowData>();
+		}
+	}
+#endif
\ No newline at end of file
diff --git a/Engine/Plugins/Experimental/Chooser/Source/ChooserEditor/Private/ChooserEditorModule.cpp b/Engine/Plugins/Experimental/Chooser/Source/ChooserEditor/Private/ChooserEditorModule.cpp
--- a/Engine/Plugins/Experimental/Chooser/Source/ChooserEditor/Private/ChooserEditorModule.cpp
+++ b/Engine/Plugins/Experimental/Chooser/Source/ChooserEditor/Private/ChooserEditorModule.cpp
@@ -15,6 +15,7 @@
 #include "FrameTimeCustomization.h"
 #include "GameplayTagColumnEditor.h"
 #include "IAssetTools.h"
+#include "MultiEnumColumnEditor.h"
 #include "ObjectColumnEditor.h"
 #include "OutputFloatColumnEditor.h"
 #include "OutputStructColumnEditor.h"
@@ -40,6 +41,7 @@ void FModule::StartupModule()
 	RegisterOutputFloatWidgets();
 	RegisterBoolWidgets();
 	RegisterEnumWidgets();
+	RegisterMultiEnumWidgets();
 	RegisterObjectWidgets();
 	RegisterStructWidgets();
 	RegisterRandomizeWidgets();
diff --git a/Engine/Plugins/Experimental/Chooser/Source/ChooserEditor/Private/MultiEnumColumnEditor.cpp b/Engine/Plugins/Experimental/Chooser/Source/ChooserEditor/Private/MultiEnumColumnEditor.cpp
new file mode 100644
--- /dev/null
+++ b/Engine/Plugins/Experimental/Chooser/Source/ChooserEditor/Private/MultiEnumColumnEditor.cpp
@@ -0,0 +1,208 @@
+// Copyright Epic Games, Inc. All Rights Reserved.
+
+#include "MultiEnumColumnEditor.h"
+#include "MultiEnumColumn.h"
+#include "SPropertyAccessChainWidget.h"
+#include "ObjectChooserWidgetFactories.h"
+#include "ChooserTableEditor.h"
+#include "GraphEditorSettings.h"
+#include "SEnumCombo.h"
+#include "TransactionCommon.h"
+#include "Widgets/Input/SButton.h"
+#include "ScopedTransaction.h"
+
+#define LOCTEXT_NAMESPACE "MultiEnumColumnEditor"
+
+namespace UE::ChooserEditor
+{
+
+	
+// Wrapper widget for EnumComboBox which will reconstruct the combo box when the Enum has changed
+template <typename ColumnType>
+class SMultiEnumCell : public SCompoundWidget
+{
+public:
+
+	DECLARE_DELEGATE_OneParam(FOnValueSet, int);
+	
+	SLATE_BEGIN_ARGS(SMultiEnumCell)
+	{}
+
+	SLATE_ARGUMENT(UObject*, TransactionObject)
+	SLATE_ARGUMENT(ColumnType*, MultiEnumColumn)
+	SLATE_ATTRIBUTE(int32, EnumValue);
+	SLATE_EVENT(FOnValueSet, OnValueSet)
+            
+	SLATE_END_ARGS()
+
+	TSharedRef<SWidget> CreateEnumComboBox()
+	{
+		if (const ColumnType* MultiEnumColumnPointer = MultiEnumColumn)
+		{
+			if (MultiEnumColumnPointer->InputValue.IsValid())
+			{
+				if (const UEnum* Enum = MultiEnumColumnPointer->InputValue.template Get<FChooserParameterEnumBase>().GetEnum())
+				{
+					return SNew(SEnumComboBox, Enum)
+						.bForceBitFlags(true)
+						.OverrideNoFlagsSetText(LOCTEXT("(Any)", "(Any)"))
+						.IsEnabled_Lambda([this](){ return IsEnabled(); } )
+						.CurrentValue(EnumValue)
+						.OnEnumSelectionChanged_Lambda([this](int32 InEnumValue, ESelectInfo::Type)
+						{
+							const FScopedTransaction Transaction(LOCTEXT("Edit RHS", "Edit Enum Value"));
+							TransactionObject->Modify(true);
+							OnValueSet.ExecuteIfBound(InEnumValue);
+						});
+				}
+			}
+		}
+		
+		return SNullWidget::NullWidget;
+	}
+
+	void UpdateEnumComboBox()
+	{
+		ChildSlot[ CreateEnumComboBox()	];
+	}
+
+	virtual void Tick(const FGeometry& AllottedGeometry, const double InCurrentTime, const float InDeltaTime) override
+	{
+		const UEnum* CurrentEnumSource = nullptr;
+		if (MultiEnumColumn->InputValue.IsValid())
+		{
+			CurrentEnumSource = MultiEnumColumn->InputValue.template Get<FChooserParameterEnumBase>().GetEnum(); 
+		}
+		if (EnumSource != CurrentEnumSource)
+		{
+			EnumComboBorder->SetContent(CreateEnumComboBox());
+			EnumSource = CurrentEnumSource;
+		}
+	}
+    					
+
+	void Construct( const FArguments& InArgs)
+	{
+		SetEnabled(InArgs._IsEnabled);
+		
+		SetCanTick(true);
+		MultiEnumColumn = InArgs._MultiEnumColumn;
+		TransactionObject = InArgs._TransactionObject;
+		EnumValue = InArgs._EnumValue;
+		OnValueSet = InArgs._OnValueSet;
+
+		if (MultiEnumColumn)
+		{
+			if (MultiEnumColumn->InputValue.IsValid())
+			{
+				EnumSource = MultiEnumColumn->InputValue.template Get<FChooserParameterEnumBase>().GetEnum();
+			}
+		}
+
+		UpdateEnumComboBox();
+
+		int Row = RowIndex.Get();
+
+		ChildSlot
+		[
+			SAssignNew(EnumComboBorder, SBorder).Padding(0).BorderBackgroundColor(FLinearColor(0,0,0,0))
+			[
+				CreateEnumComboBox()
+			]
+		];
+		
+	}
+
+	~SMultiEnumCell()
+	{
+	}
+
+private:
+	UObject* TransactionObject = nullptr;
+	ColumnType* MultiEnumColumn = nullptr;
+	const UEnum* EnumSource = nullptr;
+	TSharedPtr<SBorder> EnumComboBorder;
+	TAttribute<int> RowIndex;
+	FDelegateHandle EnumChangedHandle;
+	
+	FOnValueSet OnValueSet;
+	TAttribute<int32> EnumValue;
+};
+
+TSharedRef<SWidget> CreateMultiEnumColumnWidget(UChooserTable* Chooser, FChooserColumnBase* Column, int Row)
+{
+	FMultiEnumColumn* MultiEnumColumn = static_cast<FMultiEnumColumn*>(Column);
+	
+	if (Row == ColumnWidget_SpecialIndex_Fallback)
+	{
+		return SNullWidget::NullWidget;
+	}
+	if (Row == ColumnWidget_SpecialIndex_Header)
+	{
+		// create column header widget
+		TSharedPtr<SWidget> InputValueWidget = nullptr;
+		if (FChooserParameterBase* InputValue = Column->GetInputValue())
+		{
+			InputValueWidget = FObjectChooserWidgetFactories::CreateWidget(false, Chooser, InputValue, Column->GetInputType(), Chooser->OutputObjectType);
+		}
+		
+		const FSlateBrush* ColumnIcon = FCoreStyle::Get().GetBrush("Icons.Filter");
+		
+		TSharedRef<SWidget> ColumnHeaderWidget = SNew(SHorizontalBox)
+			+ SHorizontalBox::Slot().AutoWidth()
+			[
+				SNew(SBorder)
+				.BorderBackgroundColor(FLinearColor(0,0,0,0))
+				.Content()
+				[
+					SNew(SImage).Image(ColumnIcon)
+				]
+			]
+			+ SHorizontalBox::Slot()
+			[
+				InputValueWidget ? InputValueWidget.ToSharedRef() : SNullWidget::NullWidget
+			];
+	
+		if (Chooser->bEnableDebugTesting)
+		{
+			ColumnHeaderWidget = SNew(SVerticalBox)
+			+ SVerticalBox::Slot()
+			[
+				ColumnHeaderWidget
+			]
+			+ SVerticalBox::Slot()
+			[
+				SNew(SMultiEnumCell<FMultiEnumColumn>).TransactionObject(Chooser).MultiEnumColumn(MultiEnumColumn)
+					.OnValueSet_Lambda([MultiEnumColumn](int Value) { MultiEnumColumn->TestValue = Value; })
+					.EnumValue_Lambda([MultiEnumColumn]() { return MultiEnumColumn->TestValue; })
+					.IsEnabled_Lambda([Chooser] { return !Chooser->HasDebugTarget(); })
+			];
+		}
+
+		return ColumnHeaderWidget;
+	}
+
+	// create cell widget
+	
+	return SNew(SMultiEnumCell<FMultiEnumColumn>).TransactionObject(Chooser).MultiEnumColumn(MultiEnumColumn)
+			.OnValueSet_Lambda([MultiEnumColumn, Row](int Value)
+			{
+				if (MultiEnumColumn->RowValues.IsValidIndex(Row))
+				{
+					MultiEnumColumn->RowValues[Row].Value = static_cast<uint8>(Value);
+				}
+			})
+			.EnumValue_Lambda([MultiEnumColumn, Row]()
+			{
+				return MultiEnumColumn->RowValues.IsValidIndex(Row) ? static_cast<int32>(MultiEnumColumn->RowValues[Row].Value) : 0;
+			});
+}
+
+void RegisterMultiEnumWidgets()
+{
+	FObjectChooserWidgetFactories::RegisterColumnWidgetCreator(FMultiEnumColumn::StaticStruct(), CreateMultiEnumColumnWidget);
+}
+	
+}
+
+#undef LOCTEXT_NAMESPACE
diff --git a/Engine/Plugins/Experimental/Chooser/Source/ChooserEditor/Private/MultiEnumColumnEditor.h b/Engine/Plugins/Experimental/Chooser/Source/ChooserEditor/Private/MultiEnumColumnEditor.h
new file mode 100644
--- /dev/null
+++ b/Engine/Plugins/Experimental/Chooser/Source/ChooserEditor/Private/MultiEnumColumnEditor.h
@@ -0,0 +1,10 @@
+// Copyright Epic Games, Inc. All Rights Reserved.
+
+#pragma once
+
+#include "CoreMinimal.h"
+
+namespace UE::ChooserEditor
+{
+	void RegisterMultiEnumWidgets();
+}
\ No newline at end of file
diff --git a/Engine/Source/Editor/EditorWidgets/Private/SEnumCombo.cpp b/Engine/Source/Editor/EditorWidgets/Private/SEnumCombo.cpp
--- a/Engine/Source/Editor/EditorWidgets/Private/SEnumCombo.cpp
+++ b/Engine/Source/Editor/EditorWidgets/Private/SEnumCombo.cpp
@@ -31,8 +31,9 @@ void SEnumComboBox::Construct(const FArguments& InArgs, const UEnum* InEnum)
 	OnEnumSelectionChangedDelegate = InArgs._OnEnumSelectionChanged;
 	OnGetToolTipForValue = InArgs._OnGetToolTipForValue;
 	Font = InArgs._Font;
+	OverrideNoFlagsSetText = InArgs._OverrideNoFlagsSetText;
 	bUpdatingSelectionInternally = false;
-	bIsBitflagsEnum = Enum->HasMetaData(TEXT("Bitflags"));
+	bIsBitflagsEnum = InArgs._bForceBitFlags || Enum->HasMetaData(TEXT("Bitflags"));
 
 	TArray<int32> EnumValueSubset = InArgs._EnumValueSubset;
 	if(EnumValueSubset.IsEmpty())
@@ -141,7 +142,14 @@ FText SEnumComboBox::GetCurrentValueText() const
 
 				return FText::Join(FText::FromString(" | "), SetFlags);
 			}
-			return LOCTEXT("BitmaskButtonContentNoFlagsSet", "(No Flags Set)");
+			if (OverrideNoFlagsSetText.IsSet())
+			{
+				return OverrideNoFlagsSetText.Get();
+			}
+			else
+			{
+				return LOCTEXT("BitmaskButtonContentNoFlagsSet", "(No Flags Set)");
+			}
 		}
 		return FText::GetEmpty();
 	}
diff --git a/Engine/Source/Editor/EditorWidgets/Public/SEnumCombo.h b/Engine/Source/Editor/EditorWidgets/Public/SEnumCombo.h
--- a/Engine/Source/Editor/EditorWidgets/Public/SEnumCombo.h
+++ b/Engine/Source/Editor/EditorWidgets/Public/SEnumCombo.h
@@ -32,17 +32,20 @@ public:
 		, _ComboButtonStyle(&FAppStyle::Get().GetWidgetStyle<FComboButtonStyle>("ComboButton"))
 		, _ButtonStyle(nullptr)
 		, _TextStyle(&FCoreStyle::Get().GetWidgetStyle<FTextBlockStyle>("NormalText"))
+		, _bForceBitFlags(false)
 	{}
 
 		SLATE_ATTRIBUTE(int32, CurrentValue)
 		SLATE_ATTRIBUTE(FMargin, ContentPadding)
 		SLATE_ATTRIBUTE(FSlateFontInfo, Font)
+		SLATE_ATTRIBUTE(FText, OverrideNoFlagsSetText)
 		SLATE_EVENT(FOnEnumSelectionChanged, OnEnumSelectionChanged)
 		SLATE_EVENT(FOnGetToolTipForValue, OnGetToolTipForValue)
 		SLATE_STYLE_ARGUMENT(FComboButtonStyle, ComboButtonStyle)
 		SLATE_STYLE_ARGUMENT(FButtonStyle, ButtonStyle)
 		SLATE_STYLE_ARGUMENT(FTextBlockStyle, TextStyle)
 		SLATE_ARGUMENT(TArray<int32>, EnumValueSubset)
+		SLATE_ARGUMENT(bool, bForceBitFlags)
 	SLATE_END_ARGS()
 
 	void Construct(const FArguments& InArgs, const UEnum* InEnum);
@@ -72,6 +75,7 @@ private:
 	TAttribute<int32> CurrentValue;
 
 	TAttribute<FSlateFontInfo> Font;
+	TAttribute<FText> OverrideNoFlagsSetText;
 
 	TArray<FEnumInfo> VisibleEnums;
 
