NoMin / NoMax settings for Float Range Column cells

diff --git a/Engine/Plugins/Experimental/Chooser/Source/Chooser/Internal/FloatRangeColumn.h b/Engine/Plugins/Experimental/Chooser/Source/Chooser/Internal/FloatRangeColumn.h
--- a/Engine/Plugins/Experimental/Chooser/Source/Chooser/Internal/FloatRangeColumn.h
+++ b/Engine/Plugins/Experimental/Chooser/Source/Chooser/Internal/FloatRangeColumn.h
@@ -42,11 +42,19 @@ struct FChooserFloatRangeRowData
 {
 	GENERATED_BODY()
 
-	UPROPERTY(EditAnywhere, Category=Runtime)
+	UPROPERTY(EditAnywhere, meta=(DisplayAfter="bNoMin", EditCondition="bNoMin==false"), Category=Runtime)
 	float Min=0;
 	
-	UPROPERTY(EditAnywhere, Category=Runtime)
+	UPROPERTY(EditAnywhere, meta=(DisplayAfter="bNoMax", EditCondition="bNoMax==false"), Category=Runtime)
 	float Max=0;
+
+	// Infinite minimum range
+	UPROPERTY(EditAnywhere, Category=Runtime)
+   	bool bNoMin=false;
+	
+	// Infinite maximum range
+	UPROPERTY(EditAnywhere, Category=Runtime)
+   	bool bNoMax=false;
 };
 
 
diff --git a/Engine/Plugins/Experimental/Chooser/Source/Chooser/Private/FloatRangeColumn.cpp b/Engine/Plugins/Experimental/Chooser/Source/Chooser/Private/FloatRangeColumn.cpp
--- a/Engine/Plugins/Experimental/Chooser/Source/Chooser/Private/FloatRangeColumn.cpp
+++ b/Engine/Plugins/Experimental/Chooser/Source/Chooser/Private/FloatRangeColumn.cpp
@@ -27,22 +27,23 @@ bool FFloatRangeColumn::EditorTestFilter(int32 RowIndex) const
 {
 	if (RowValues.IsValidIndex(RowIndex))
 	{
+		const FChooserFloatRangeRowData& RowValue = RowValues[RowIndex];
+		
 		if (bWrapInput)
 		{
 			double TestResult = FMath::Wrap(TestValue, MinValue, MaxValue);
-			const FChooserFloatRangeRowData& RowValue = RowValues[RowIndex];
 			if (RowValue.Max < RowValue.Min) // eg for an angle range from  135 to -135  (135 to 180 or -180 to -135)
 			{
-				return TestResult >= RowValue.Min || TestResult <= RowValue.Max;
+				return (RowValue.bNoMin || TestResult >= RowValue.Min) || (RowValue.bNoMax || TestResult <= RowValue.Max);
 			}
 			else
 			{
-				return TestResult >= RowValue.Min && TestResult <= RowValue.Max;
+				return (RowValue.bNoMin || TestResult >= RowValue.Min) && (RowValue.bNoMax || TestResult <= RowValue.Max);
 			}
 		}
 		else
 		{
-			return TestValue >= RowValues[RowIndex].Min && TestValue <= RowValues[RowIndex].Max;
+			return (RowValue.bNoMin || TestValue >= RowValue.Min) && (RowValue.bNoMax || TestValue <= RowValue.Max);
 		}
 	}
 	return false;
diff --git a/Engine/Plugins/Experimental/Chooser/Source/ChooserEditor/Private/FloatRangeColumnEditor.cpp b/Engine/Plugins/Experimental/Chooser/Source/ChooserEditor/Private/FloatRangeColumnEditor.cpp
--- a/Engine/Plugins/Experimental/Chooser/Source/ChooserEditor/Private/FloatRangeColumnEditor.cpp
+++ b/Engine/Plugins/Experimental/Chooser/Source/ChooserEditor/Private/FloatRangeColumnEditor.cpp
@@ -4,11 +4,13 @@
 #include "FloatRangeColumn.h"
 #include "ObjectChooserWidgetFactories.h"
 #include "ChooserTableEditor.h"
+#include "Widgets/Input/SEditableTextBox.h"
 #include "Widgets/Input/SNumericEntryBox.h"
 #include "Widgets/Images/SImage.h"
 #include "GraphEditorSettings.h"
 #include "SPropertyAccessChainWidget.h"
 #include "ScopedTransaction.h"
+#include "Widgets/Layout/SWidgetSwitcher.h"
 
 #define LOCTEXT_NAMESPACE "FloatRangeColumnEditor"
 
@@ -86,24 +88,40 @@ TSharedRef<SWidget> CreateFloatRangeColumnWidget(UChooserTable* Chooser, FChoose
 	]
 	+ SHorizontalBox::Slot().FillWidth(0.5f)
 	[
-		SNew(SNumericEntryBox<float>)
-		.MaxValue_Lambda([FloatRangeColumn, Row]()
-		{
-			return (Row < FloatRangeColumn->RowValues.Num()) ? FloatRangeColumn->RowValues[Row].Max : 0;
-		})
-		.Value_Lambda([FloatRangeColumn, Row]()
-		{
-			return (Row < FloatRangeColumn->RowValues.Num()) ? FloatRangeColumn->RowValues[Row].Min : 0;
-		})
-		.OnValueCommitted_Lambda([Chooser, FloatRangeColumn, Row](float NewValue, ETextCommit::Type CommitType)
-		{
-			if (Row < FloatRangeColumn->RowValues.Num())
+		SNew(SOverlay)
+		+ SOverlay::Slot()
+		[
+			SNew(SNumericEntryBox<float>)
+			.Visibility_Lambda([FloatRangeColumn, Row]
+			{
+				return FloatRangeColumn->RowValues[Row].bNoMin ? EVisibility::Hidden : EVisibility::Visible;
+			})
+			.MaxValue_Lambda([FloatRangeColumn, Row]()
+			{
+				return (Row < FloatRangeColumn->RowValues.Num()) ? FloatRangeColumn->RowValues[Row].Max : 0;
+			})
+			.Value_Lambda([FloatRangeColumn, Row]()
+			{
+				return (Row < FloatRangeColumn->RowValues.Num()) ? FloatRangeColumn->RowValues[Row].Min : 0;
+			})
+			.OnValueCommitted_Lambda([Chooser, FloatRangeColumn, Row](float NewValue, ETextCommit::Type CommitType)
 			{
-				const FScopedTransaction Transaction(LOCTEXT("Edit Min Value", "Edit Min Value"));
-				Chooser->Modify(true);
-				FloatRangeColumn->RowValues[Row].Min = NewValue;
-			}
-		})
+				if (Row < FloatRangeColumn->RowValues.Num())
+				{
+					const FScopedTransaction Transaction(LOCTEXT("Edit Min Value", "Edit Min Value"));
+					Chooser->Modify(true);
+					FloatRangeColumn->RowValues[Row].Min = NewValue;
+				}
+			})
+		]
+		+ SOverlay::Slot()
+		[
+			SNew(SEditableTextBox).IsReadOnly(true).Text(LOCTEXT("negativeinfinity","-∞"))
+				.Visibility_Lambda([FloatRangeColumn, Row]
+				{
+					return FloatRangeColumn->RowValues[Row].bNoMin ? EVisibility::Visible : EVisibility::Hidden;
+				})
+		]
 	]
 	+ SHorizontalBox::Slot().AutoWidth().VAlign(VAlign_Center)
 	[
@@ -111,24 +129,40 @@ TSharedRef<SWidget> CreateFloatRangeColumnWidget(UChooserTable* Chooser, FChoose
 	]
 	+ SHorizontalBox::Slot().FillWidth(0.5f)
 	[
-		SNew(SNumericEntryBox<float>)
-		.MinValue_Lambda([FloatRangeColumn, Row]()
-		{
-			return (Row < FloatRangeColumn->RowValues.Num()) ? FloatRangeColumn->RowValues[Row].Min : 0;
-		})
-		.Value_Lambda([FloatRangeColumn, Row]()
-		{
-			return (Row < FloatRangeColumn->RowValues.Num()) ? FloatRangeColumn->RowValues[Row].Max : 0;
-		})
-		.OnValueCommitted_Lambda([Chooser, FloatRangeColumn, Row](float NewValue, ETextCommit::Type CommitType)
-		{
-			if (Row < FloatRangeColumn->RowValues.Num())
+		SNew(SOverlay)
+		+ SOverlay::Slot()
+		[
+			SNew(SNumericEntryBox<float>)
+			.Visibility_Lambda([FloatRangeColumn, Row]
+			{
+				return FloatRangeColumn->RowValues[Row].bNoMax ? EVisibility::Hidden : EVisibility::Visible;
+			})
+			.MinValue_Lambda([FloatRangeColumn, Row]()
+			{
+				return (Row < FloatRangeColumn->RowValues.Num()) ? FloatRangeColumn->RowValues[Row].Min : 0;
+			})
+			.Value_Lambda([FloatRangeColumn, Row]()
+			{
+				return (Row < FloatRangeColumn->RowValues.Num()) ? FloatRangeColumn->RowValues[Row].Max : 0;
+			})
+			.OnValueCommitted_Lambda([Chooser, FloatRangeColumn, Row](float NewValue, ETextCommit::Type CommitType)
 			{
-				const FScopedTransaction Transaction(LOCTEXT("Edit Max", "Edit Max Value"));
-				Chooser->Modify(true);
-				FloatRangeColumn->RowValues[Row].Max = NewValue;
-			}
-		})
+				if (Row < FloatRangeColumn->RowValues.Num())
+				{
+					const FScopedTransaction Transaction(LOCTEXT("Edit Max", "Edit Max Value"));
+					Chooser->Modify(true);
+					FloatRangeColumn->RowValues[Row].Max = NewValue;
+				}
+			})
+		]
+		+ SOverlay::Slot()
+		[
+			SNew(SEditableTextBox).IsReadOnly(true).Text(LOCTEXT("infinity","∞"))
+				.Visibility_Lambda([FloatRangeColumn, Row]
+				{
+					return FloatRangeColumn->RowValues[Row].bNoMax ? EVisibility::Visible : EVisibility::Hidden;
+				})
+		]
 	]
 	+ SHorizontalBox::Slot().AutoWidth().VAlign(VAlign_Center)
 	[
