Support for disabled Rows and Columns in Chooser Tables.  They display greyed out in editor, and are removed from cooked builds (including assets referenced from row results)

diff --git a/Engine/Plugins/Experimental/Chooser/Source/Chooser/Internal/Chooser.h b/Engine/Plugins/Experimental/Chooser/Source/Chooser/Internal/Chooser.h
--- a/Engine/Plugins/Experimental/Chooser/Source/Chooser/Internal/Chooser.h
+++ b/Engine/Plugins/Experimental/Chooser/Source/Chooser/Internal/Chooser.h
@@ -24,6 +24,21 @@ public:
 
 	virtual void PostLoad() override;
 	virtual void Compile(bool bForce = false) override;
+	virtual void Serialize(FArchive& Ar) override;
+
+#if WITH_EDITORONLY_DATA
+	void RemoveDisabledData();
+	void CookData();
+#endif
+
+	bool IsRowDisabled(int32 RowIndex) const
+	{
+#if WITH_EDITORONLY_DATA
+		return CookedResults.IsEmpty() && DisabledRows.IsValidIndex(RowIndex) && DisabledRows[RowIndex];
+#else
+		return false;
+#endif
+	}
 #if WITH_EDITOR
 	virtual void GetAssetRegistryTags(FAssetRegistryTagsContext Context) const override;
 	UE_DEPRECATED(5.4, "Implement the version that takes FAssetRegistryTagsContext instead.")
@@ -99,11 +114,19 @@ public:
 	// FallbackResult will be used as the Result if there are no rows in the chooser which pass all filters.  If FallbackResult is not assigned, then the Chooser will return null in that case.
 	UPROPERTY(EditAnywhere, Meta = (ExcludeBaseStruct, BaseStruct = "/Script/Chooser.ObjectChooserBase"), Category = "Fallback")
 	FInstancedStruct FallbackResult;
-	
+
+#if WITH_EDITORONLY_DATA
 	// Each possible Result (Rows of chooser table)
 	UPROPERTY(EditAnywhere, NoClear, DisplayName = "Results", Meta = (ExcludeBaseStruct, BaseStruct = "/Script/Chooser.ObjectChooserBase"), Category = "Hidden")
 	TArray<FInstancedStruct> ResultsStructs;
 
+	UPROPERTY()
+	TArray<bool> DisabledRows;
+#endif
+	
+	UPROPERTY()
+	TArray<FInstancedStruct> CookedResults;
+
 	// Columns which filter Results
 	UPROPERTY(EditAnywhere, NoClear, DisplayName = "Columns", Category = Hidden, meta = (ExcludeBaseStruct, BaseStruct = "/Script/Chooser.ChooserColumnBase"))
 	TArray<FInstancedStruct> ColumnsStructs;
@@ -117,7 +140,8 @@ public:
 	EObjectChooserResultType ResultType = EObjectChooserResultType::ObjectResult;
 
 	virtual TConstArrayView<FInstancedStruct> GetContextData() const override { return ContextData; }
-	
+
+	bool IsCookedData() const { return !CookedResults.IsEmpty(); }
 	static FObjectChooserBase::EIteratorStatus EvaluateChooser(FChooserEvaluationContext& Context, const UChooserTable* Chooser, FObjectChooserBase::FObjectChooserIteratorCallback Callback);
 };
 
@@ -181,4 +205,6 @@ public:
 	class FAssetEditorToolkit* Editor;
 	TWeakObjectPtr<UChooserTable> Chooser;
 	int ColumnIndex;
-};
\ No newline at end of file
+};
+
+
diff --git a/Engine/Plugins/Experimental/Chooser/Source/Chooser/Private/Chooser.cpp b/Engine/Plugins/Experimental/Chooser/Source/Chooser/Private/Chooser.cpp
--- a/Engine/Plugins/Experimental/Chooser/Source/Chooser/Private/Chooser.cpp
+++ b/Engine/Plugins/Experimental/Chooser/Source/Chooser/Private/Chooser.cpp
@@ -10,6 +10,10 @@
 #include "IChooserParameterGameplayTag.h"
 #include "UObject/AssetRegistryTagsContext.h"
 
+#if WITH_EDITORONLY_DATA
+#include "UObject/UObjectIterator.h"
+#endif
+
 #include UE_INLINE_GENERATED_CPP_BY_NAME(Chooser)
 
 DEFINE_LOG_CATEGORY(LogChooser)
@@ -85,7 +89,10 @@ void UChooserTable::PostLoad()
 void UChooserTable::BeginDestroy()
 {
 	ColumnsStructs.Empty();
+#if WITH_EDITORONLY_DATA
 	ResultsStructs.Empty();
+#endif
+	CookedResults.Empty();
 	Super::BeginDestroy();
 }
 
@@ -156,6 +163,7 @@ void UChooserTable::Compile(bool bForce)
 		}
 	}
 
+#if WITH_EDITORONLY_DATA
 	for(FInstancedStruct& ResultData : ResultsStructs)
 	{
 		if (ResultData.IsValid())
@@ -164,8 +172,74 @@ void UChooserTable::Compile(bool bForce)
 			Result.Compile(ContextOwner, bForce);
 		}
 	}
+#endif
+
+	if (IsCookedData())
+	{
+		for(FInstancedStruct& ResultData : CookedResults)
+		{
+			if (ResultData.IsValid())
+			{
+				FObjectChooserBase& Result = ResultData.GetMutable<FObjectChooserBase>();
+				Result.Compile(ContextOwner, bForce);
+			}
+		}
+	}
+}
+
+void UChooserTable::Serialize(FArchive& Ar)
+{
+#if WITH_EDITORONLY_DATA
+	if (Ar.IsCooking() && !Ar.IsObjectReferenceCollector())
+	{
+		CookData();
+	}
+#endif
+	
+	UObject::Serialize(Ar);
+}
+
+#if WITH_EDITORONLY_DATA
+void UChooserTable::RemoveDisabledData()
+{
+	// remove disabled or invalid columns
+	ColumnsStructs.RemoveAll([](const FInstancedStruct& ColumnStruct)
+	{
+		return !ColumnStruct.IsValid() || ColumnStruct.Get<FChooserColumnBase>().bDisabled;
+	});
+	
+	// remove disabled rows and corresponding row data from columns
+	TArray<uint32> RowsToDelete;
+	const int NumResults = ResultsStructs.Num();
+	for (int ResultIndex = NumResults - 1; ResultIndex>=0; ResultIndex--)
+	{
+		if (IsRowDisabled(ResultIndex))
+		{
+			RowsToDelete.Add(ResultIndex);
+		}
+	}
+	
+	DisabledRows.SetNum(0);
+
+	for (uint32 Index : RowsToDelete)
+	{
+		ResultsStructs.RemoveAt(Index);
+	}
+	for(FInstancedStruct& Column : ColumnsStructs)
+	{
+		Column.GetMutable<FChooserColumnBase>().DeleteRows(RowsToDelete);
+	}
 }
 
+void UChooserTable::CookData()
+{
+	RemoveDisabledData();
+
+	// copy stripped results struct into CookedResults array
+	CookedResults = ResultsStructs;
+}
+#endif
+
 #if WITH_EDITOR
 void UChooserTable::PostEditUndo()
 {
@@ -275,17 +349,28 @@ FObjectChooserBase::EIteratorStatus UChooserTable::EvaluateChooser(FChooserEvalu
 	Context.DebuggingInfo.CurrentChooser = Chooser;
 #endif
 
-	uint32 Count = Chooser->ResultsStructs.Num();
+	const TArray<FInstancedStruct>* ResultsArray = &Chooser->CookedResults;
+
+#if WITH_EDITORONLY_DATA
+	if (!Chooser->IsCookedData())
+	{
+		ResultsArray = &Chooser->ResultsStructs;
+	}
+#endif
+	
+
+	uint32 Count = ResultsArray->Num();
 	uint32 BufferSize = Count * sizeof(uint32);
 
 	FChooserIndexArray Indices1(static_cast<uint32*>(FMemory_Alloca(BufferSize)), Count);
 	FChooserIndexArray Indices2(static_cast<uint32*>(FMemory_Alloca(BufferSize)), Count);
 
-	int RowCount = Chooser->ResultsStructs.Num();
-	Indices1.SetNum(RowCount);
-	for(int i=0;i<RowCount;i++)
+	for(uint32 i=0;i<Count;i++)
 	{
-		Indices1[i]=i;
+		if (!Chooser->IsRowDisabled(i))
+		{
+			Indices1.Push(i);
+		}
 	}
 	FChooserIndexArray* IndicesOut = &Indices1;
 	FChooserIndexArray* IndicesIn = &Indices2;
@@ -293,6 +378,14 @@ FObjectChooserBase::EIteratorStatus UChooserTable::EvaluateChooser(FChooserEvalu
 	for (const FInstancedStruct& ColumnData : Chooser->ColumnsStructs)
 	{
 		const FChooserColumnBase& Column = ColumnData.Get<FChooserColumnBase>();
+
+#if WITH_EDITORONLY_DATA
+		if (Column.bDisabled)
+		{
+			continue;
+		}
+#endif
+		
 		if (Column.HasFilters())
 		{
 			Swap(IndicesIn, IndicesOut);
@@ -339,9 +432,9 @@ FObjectChooserBase::EIteratorStatus UChooserTable::EvaluateChooser(FChooserEvalu
 		// of the rows that passed all column filters, iterate through them calling the callback until it returns Stop
 		for (uint32 SelectedIndex : *IndicesOut)
 		{
-			if (Chooser->ResultsStructs.Num() > (int32)SelectedIndex)
+			if (ResultsArray->Num() > (int32)SelectedIndex)
 			{
-				const FObjectChooserBase& SelectedResult = Chooser->ResultsStructs[SelectedIndex].Get<FObjectChooserBase>();
+				const FObjectChooserBase& SelectedResult = (*ResultsArray)[SelectedIndex].Get<FObjectChooserBase>();
 				FObjectChooserBase::EIteratorStatus Status = SelectedResult.ChooseMulti(Context, Callback);
 				if (Status != FObjectChooserBase::EIteratorStatus::Continue)
 				{
@@ -421,4 +514,22 @@ FObjectChooserBase::EIteratorStatus FNestedChooser::ChooseMulti(FChooserEvaluati
 void FNestedChooser::GetDebugName(FString& OutDebugName) const
 {
 	OutDebugName  = GetNameSafe(Chooser);
-}
\ No newline at end of file
+}
+
+#if WITH_EDITORONLY_DATA
+namespace {
+	void TestCook()
+	{
+		for (TObjectIterator<UChooserTable> It; It; ++It)
+		{
+			It->CookData();
+		}
+	}
+}
+
+static FAutoConsoleCommand CCmdTestCookChoosers(
+	TEXT("Chooser.TestCook"),
+	TEXT(""),
+	FConsoleCommandDelegate::CreateStatic(TestCook));
+
+#endif
\ No newline at end of file
diff --git a/Engine/Plugins/Experimental/Chooser/Source/Chooser/Public/IChooserColumn.h b/Engine/Plugins/Experimental/Chooser/Source/Chooser/Public/IChooserColumn.h
--- a/Engine/Plugins/Experimental/Chooser/Source/Chooser/Public/IChooserColumn.h
+++ b/Engine/Plugins/Experimental/Chooser/Source/Chooser/Public/IChooserColumn.h
@@ -75,6 +75,11 @@ public:
 	virtual void AddToDetails(FInstancedPropertyBag& PropertyBag, int32 ColumnIndex, int32 RowIndex) {}
 	virtual void SetFromDetails(FInstancedPropertyBag& PropertyBag, int32 ColumnIndex, int32 RowIndex) {}
 #endif
+
+#if WITH_EDITORONLY_DATA
+	UPROPERTY(EditAnywhere, Category=Base)
+	bool bDisabled = false;
+#endif
 };
 
 
diff --git a/Engine/Plugins/Experimental/Chooser/Source/ChooserEditor/Private/ChooserTableEditor.cpp b/Engine/Plugins/Experimental/Chooser/Source/ChooserEditor/Private/ChooserTableEditor.cpp
--- a/Engine/Plugins/Experimental/Chooser/Source/ChooserEditor/Private/ChooserTableEditor.cpp
+++ b/Engine/Plugins/Experimental/Chooser/Source/ChooserEditor/Private/ChooserTableEditor.cpp
@@ -114,12 +114,12 @@ FChooserTableEditor::~FChooserTableEditor()
 	FCoreUObjectDelegates::OnObjectTransacted.RemoveAll(this);
 
 	DetailsView.Reset();
-	
 }
 
-
 FName FChooserTableEditor::EditorName = "ChooserTableEditor";
 	
+FName FChooserTableEditor::ContextMenuName("ChooserEditorContextMenu"); // todo: for this to actually be extensible this needs to be somewhere public
+	
 FName FChooserTableEditor::GetEditorName() const
 {
 	return EditorName;
@@ -203,9 +203,16 @@ void FChooserTableEditor::MakeDebugTargetMenu(UToolMenu* InToolMenu)
 						));
 		}
 	);
-	
 
 }
+	
+TSharedPtr<SWidget> FChooserTableEditor::GenerateRowContextMenu()
+{
+	UToolMenus* ToolMenus = UToolMenus::Get();
+	FToolMenuContext ToolMenuContext;
+	InitToolMenuContext(ToolMenuContext);
+	return ToolMenus->GenerateWidget(ContextMenuName, ToolMenuContext);
+}
 
 void FChooserTableEditor::RegisterToolbar()
 {
@@ -263,12 +270,73 @@ void FChooserTableEditor::RegisterToolbar()
 
 }
 
+void FChooserTableEditor::RegisterMenus()
+{
+	UToolMenus* ToolMenus = UToolMenus::Get();
+	const FChooserTableEditorCommands& Commands = FChooserTableEditorCommands::Get();
+
+	// Table Context Menu
+	UToolMenu* ToolMenu;
+	if (ToolMenus->IsMenuRegistered(ContextMenuName))
+	{
+		ToolMenu = ToolMenus->ExtendMenu(ContextMenuName);
+	}
+	else
+	{
+		ToolMenu = UToolMenus::Get()->RegisterMenu(ContextMenuName, NAME_None, EMultiBoxType::Menu);
+	}
+
+	if (ToolMenu)
+	{
+		FToolMenuSection& Section = ToolMenu->AddSection("ChooserTableRow", TAttribute<FText>());
+	
+		Section.AddEntry(FToolMenuEntry::InitMenuEntry(
+			Commands.Delete,
+			TAttribute<FText>(),
+			TAttribute<FText>(),
+			FSlateIcon()));
+	
+		Section.AddEntry(FToolMenuEntry::InitMenuEntry(
+			Commands.Disable,
+			TAttribute<FText>(),
+			TAttribute<FText>(),
+			FSlateIcon()));
+	}
+
+	
+	struct Local
+	{
+		static void FillEditMenu(FMenuBuilder& MenuBuilder)
+		{
+			MenuBuilder.BeginSection("ChooserEditing", LOCTEXT("Chooser Table Editing", "Chooser Table"));
+			{
+				MenuBuilder.AddMenuEntry(FChooserTableEditorCommands::Get().RemoveDisabledData, NAME_None);
+				MenuBuilder.AddMenuEntry(FChooserTableEditorCommands::Get().Delete, NAME_None, LOCTEXT("Delete Selection", "Delete Selection"));
+				MenuBuilder.AddMenuEntry(FChooserTableEditorCommands::Get().Disable, NAME_None, LOCTEXT("Disable Selection", "Disable Selection"));
+			}
+			MenuBuilder.EndSection();
+		}
+	};
+
+	TSharedPtr<FExtender> MenuExtender = MakeShareable(new FExtender);
+
+	// Extend the Edit menu
+	MenuExtender->AddMenuExtension(
+		"EditHistory",
+		EExtensionHook::After,
+		GetToolkitCommands(),
+		FMenuExtensionDelegate::CreateStatic(&Local::FillEditMenu));
+
+	AddMenuExtender(MenuExtender);
+}
+
 void FChooserTableEditor::InitToolMenuContext(FToolMenuContext& MenuContext)
 {
 	FAssetEditorToolkit::InitToolMenuContext(MenuContext);
 
 	UChooserEditorToolMenuContext* Context = NewObject<UChooserEditorToolMenuContext>();
 	Context->ChooserEditor = SharedThis(this);
+	MenuContext.AppendCommandList(GetToolkitCommands());
 	MenuContext.AddObject(Context);
 }
 
@@ -279,6 +347,24 @@ void FChooserTableEditor::BindCommands()
 	ToolkitCommands->MapAction(
 		Commands.EditChooserSettings,
 		FExecuteAction::CreateSP(this, &FChooserTableEditor::SelectRootProperties));
+	
+	ToolkitCommands->MapAction(
+		Commands.RemoveDisabledData,
+		FExecuteAction::CreateSP(this, &FChooserTableEditor::RemoveDisabledData));
+	
+	ToolkitCommands->MapAction(
+		Commands.Delete,
+		FExecuteAction::CreateSP(this, &FChooserTableEditor::DeleteSelection),
+		FCanExecuteAction::CreateSP(this, &FChooserTableEditor::HasSelection)
+		);
+	
+	ToolkitCommands->MapAction(
+		Commands.Disable,
+		FExecuteAction::CreateSP(this, &FChooserTableEditor::ToggleDisableSelection),
+		FCanExecuteAction::CreateSP(this, &FChooserTableEditor::HasSelection),
+		FIsActionChecked::CreateSP(this, &FChooserTableEditor::IsSelectionDisabled)
+		);
+
 }
 
 void FChooserTableEditor::OnObjectsTransacted(UObject* Object, const FTransactionObjectEvent& Event)
@@ -315,6 +401,9 @@ void FChooserTableEditor::OnObjectsTransacted(UObject* Object, const FTransactio
 				Column.SetFromDetails(RowDetails->Properties, ColumnIndex, RowDetails->Row);
 				ColumnIndex++;
 			}
+			
+			TValueOrError<bool, EPropertyBagResult> DisabledResult = RowDetails->Properties.GetValueBool("Disabled");
+			RowDetails->Chooser->DisabledRows[RowDetails->Row] = DisabledResult.GetValue();
 		}
 	}
 }
@@ -375,11 +464,14 @@ void FChooserTableEditor::InitEditor( const EToolkitMode::Type Mode, const TShar
 	FAssetEditorToolkit::InitAssetEditor( Mode, InitToolkitHost, FChooserTableEditor::ChooserEditorAppIdentifier, StandaloneDefaultLayout, bCreateDefaultStandaloneMenu, bCreateDefaultToolbar, ObjectsToEdit );
 
 	BindCommands();
-	RegenerateMenusAndToolbars();
+	
+	// todo: should these be triggered once per session by the module?
 	RegisterToolbar();
+	RegisterMenus();
+	
+	RegenerateMenusAndToolbars();
 
 	SelectRootProperties();
-
 		
 	FAnimAssetFindReplaceConfig FindReplaceConfig;
 	FindReplaceConfig.InitialProcessorClass = UChooserFindProperties::StaticClass();
@@ -607,6 +699,7 @@ void FChooserTableEditor::SetPropertyEditingEnabledDelegate(FIsPropertyEditingEn
 	DetailsView->ForceRefresh();
 }
 
+
 TSharedRef<SDockTab> FChooserTableEditor::SpawnPropertiesTab( const FSpawnTabArgs& Args )
 {
 	check( Args.GetTabId() == PropertiesTabId );
@@ -653,6 +746,16 @@ void FChooserTableEditor::SelectRootProperties()
 	}
 }
 
+void FChooserTableEditor::RemoveDisabledData()
+{
+	UChooserTable* Chooser = GetChooser();
+	const FScopedTransaction Transaction(LOCTEXT("Move Row", "Move Row"));
+
+	Chooser->Modify(true);
+	Chooser->RemoveDisabledData();
+	RefreshAll();
+}
+
 int FChooserTableEditor::MoveRow(int SourceRowIndex, int TargetRowIndex)
 {
 	UChooserTable* Chooser = GetChooser();
@@ -670,12 +773,14 @@ int FChooserTableEditor::MoveRow(int SourceRowIndex, int TargetRowIndex)
 
 	FInstancedStruct Result = Chooser->ResultsStructs[SourceRowIndex];
 	Chooser->ResultsStructs.RemoveAt(SourceRowIndex);
+	bool bDisabled = Chooser->DisabledRows[SourceRowIndex];
+	Chooser->DisabledRows.RemoveAt(SourceRowIndex);
 	if (SourceRowIndex < TargetRowIndex)
 	{
 		TargetRowIndex--;
 	}
 	Chooser->ResultsStructs.Insert(Result, TargetRowIndex);
-
+	Chooser->DisabledRows.Insert(bDisabled, TargetRowIndex);
 	UpdateTableRows();
 
 	return TargetRowIndex;
@@ -948,7 +1053,13 @@ void FChooserTableEditor::RefreshRowSelectionDetails()
 				Column.AddToDetails(Selection->Properties, ColumnIndex, SelectedItem->RowIndex);
 				ColumnIndex++;
 			}
-		
+
+			if (Chooser->DisabledRows.IsValidIndex(SelectedItem->RowIndex))
+			{
+				Selection->Properties.AddProperty("Disabled", EPropertyBagPropertyType::Bool);
+				Selection->Properties.SetValueBool("Disabled", Chooser->DisabledRows[SelectedItem->RowIndex]);
+			}
+
 			SelectedRows.Add(Selection);
 			
 		}
@@ -966,7 +1077,7 @@ void FChooserTableEditor::RefreshRowSelectionDetails()
 		DetailsView->SetObjects( DetailsObjects );
 	}
 }
-    										
+
 
 TSharedRef<SDockTab> FChooserTableEditor::SpawnTableTab( const FSpawnTabArgs& Args )
 {
@@ -1032,39 +1143,13 @@ TSharedRef<SDockTab> FChooserTableEditor::SpawnTableTab( const FSpawnTabArgs& Ar
 
 	TableView = SNew(SListView<TSharedPtr<FChooserTableRow>>)
     			.ListItemsSource(&TableRows)
+				.OnContextMenuOpening_Raw(this, &FChooserTableEditor::GenerateRowContextMenu)
 				.OnKeyDownHandler_Lambda([this](const FGeometry&, const FKeyEvent& Event)
 				{
 					
 					if (Event.GetKey() == EKeys::Delete)
 					{
-						const FScopedTransaction Transaction(LOCTEXT("Delete Row Transaction", "Delete Row"));
-						UChooserTable* Chooser = GetChooser();
-						Chooser->Modify(true);
-						// delete selected rows.
-						TArray<uint32> RowsToDelete;
-						for(auto& SelectedRow:SelectedRows)
-						{
-							RowsToDelete.Add(SelectedRow->Row);
-						}
-
-						SelectedRows.SetNum(0);
-						SelectRootProperties();
-
-						// sort indices in reverse
-						RowsToDelete.Sort([](int32 A, int32 B){ return A>B; });
-						for(uint32 RowIndex : RowsToDelete)
-						{
-							Chooser->ResultsStructs.RemoveAt(RowIndex);
-						}
-
-						for(FInstancedStruct& ColumnData : Chooser->ColumnsStructs)
-						{
-							FChooserColumnBase& Column = ColumnData.GetMutable<FChooserColumnBase>();
-							Column.DeleteRows(RowsToDelete);
-						}
-
-						UpdateTableRows();
-						
+						DeleteSelectedRows();
 						return FReply::Handled();
 					}
 					return FReply::Unhandled();
@@ -1154,6 +1239,7 @@ void FChooserTableEditor::UpdateTableRows()
 {
 	UChooserTable* Chooser = GetChooser();
 	int32 NewNum = Chooser->ResultsStructs.Num();
+	Chooser->DisabledRows.SetNum(NewNum);
 
 	// Sync the TableRows array which drives the ui table to match the number of results.
 	TableRows.SetNum(0, EAllowShrinking::No);
@@ -1275,6 +1361,120 @@ void FChooserTableEditor::DeleteColumn(int Index)
 		UpdateTableColumns();
 	}
 }
+	
+void FChooserTableEditor::DeleteSelectedRows()
+{
+	const FScopedTransaction Transaction(LOCTEXT("Delete Row Transaction", "Delete Row"));
+	UChooserTable* Chooser = GetChooser();
+	Chooser->Modify(true);
+	// delete selected rows.
+	TArray<uint32> RowsToDelete;
+	for(auto& SelectedRow:SelectedRows)
+	{
+		RowsToDelete.Add(SelectedRow->Row);
+	}
+
+	SelectedRows.SetNum(0);
+	SelectRootProperties();
+
+	// sort indices in reverse
+	RowsToDelete.Sort([](int32 A, int32 B){ return A>B; });
+	for(uint32 RowIndex : RowsToDelete)
+	{
+		Chooser->ResultsStructs.RemoveAt(RowIndex);
+		Chooser->DisabledRows.RemoveAt(RowIndex);
+	}
+
+	for(FInstancedStruct& ColumnData : Chooser->ColumnsStructs)
+	{
+		FChooserColumnBase& Column = ColumnData.GetMutable<FChooserColumnBase>();
+		Column.DeleteRows(RowsToDelete);
+	}
+	UpdateTableRows();
+}
+	
+inline bool FChooserTableEditor::HasSelection()
+{
+	if (CurrentSelectionType == ESelectionType::Column)
+	{
+		return true;
+	}
+	else if (CurrentSelectionType == ESelectionType::Rows)
+	{
+		return !SelectedRows.IsEmpty();
+	}
+	return false;
+}
+	
+inline bool FChooserTableEditor::IsSelectionDisabled()
+{
+	if (UChooserTable* Chooser = GetChooser())
+	{
+		if (CurrentSelectionType == ESelectionType::Column)
+		{
+			if (SelectedColumn && Chooser->ColumnsStructs.IsValidIndex(SelectedColumn->Column))
+			{
+				FChooserColumnBase& Column = Chooser->ColumnsStructs[SelectedColumn->Column].GetMutable<FChooserColumnBase>();
+				return Column.bDisabled;
+			}
+		}
+		else if (CurrentSelectionType == ESelectionType::Rows)
+		{
+			if (SelectedRows.IsEmpty())
+			{
+				return false;
+			}
+
+			bool bSomethingEnabled = false;
+			for(auto& Row : SelectedRows)
+			{
+				if (!Chooser->IsRowDisabled(Row->Row))
+				{
+					bSomethingEnabled = true;
+					break;
+				}
+			}
+			return !bSomethingEnabled;
+		}
+	}
+	return false;
+}
+
+void FChooserTableEditor::ToggleDisableSelection()
+{
+	bool bDisabled = IsSelectionDisabled();
+	if (UChooserTable* Chooser = GetChooser())
+	{
+		if (CurrentSelectionType == ESelectionType::Column)
+		{
+			if (SelectedColumn && Chooser->ColumnsStructs.IsValidIndex(SelectedColumn->Column))
+			{
+				FChooserColumnBase& Column = Chooser->ColumnsStructs[SelectedColumn->Column].GetMutable<FChooserColumnBase>();
+				Column.bDisabled = !Column.bDisabled;
+			}
+		}
+		else if (CurrentSelectionType == ESelectionType::Rows)
+		{
+			for (auto& Row : SelectedRows)
+			{
+				Chooser->DisabledRows[Row->Row] = !bDisabled;
+			}
+			RefreshRowSelectionDetails();
+		}
+	}
+}
+
+void FChooserTableEditor::DeleteSelection()
+{
+	if (CurrentSelectionType == ESelectionType::Column)
+	{
+		DeleteColumn(SelectedColumn->Column);
+	}
+	else if (CurrentSelectionType == ESelectionType::Rows)
+	{
+		DeleteSelectedRows();
+	}
+}
 
 void FChooserTableEditor::RegisterWidgets()
 {
diff --git a/Engine/Plugins/Experimental/Chooser/Source/ChooserEditor/Private/ChooserTableEditor.h b/Engine/Plugins/Experimental/Chooser/Source/ChooserEditor/Private/ChooserTableEditor.h
--- a/Engine/Plugins/Experimental/Chooser/Source/ChooserEditor/Private/ChooserTableEditor.h
+++ b/Engine/Plugins/Experimental/Chooser/Source/ChooserEditor/Private/ChooserTableEditor.h
@@ -91,13 +91,19 @@ namespace UE::ChooserEditor
 		void SetPropertyVisibilityDelegate(FIsPropertyVisible InVisibilityDelegate);
 		/** Can be used to disable the details view making it read-only */
 		void SetPropertyEditingEnabledDelegate(FIsPropertyEditingEnabled InPropertyEditingDelegate);
-	
+
+		bool HasSelection();
+		bool IsSelectionDisabled();
+		void ToggleDisableSelection();
+		void DeleteSelection();
+		
 		void UpdateTableRows();
 		void SelectColumn(UChooserTable* Chooser, int Index);
 		void ClearSelectedColumn();
 		void DeleteColumn(int Index);
 		void AddColumn(const UScriptStruct* ColumnType);
 		void RefreshRowSelectionDetails();
+		void DeleteSelectedRows();
 		int MoveRow(int SourceRowIndex, int TargetIndex);
 		void SelectRow(int32 RowIndex, bool bClear = true);
 		void ClearSelectedRows(); 
@@ -112,10 +118,13 @@ namespace UE::ChooserEditor
 	private:
 
 		void SelectRootProperties();
+		void RemoveDisabledData();
 		void RegisterToolbar();
+		void RegisterMenus();
 		void BindCommands();
 		void OnObjectsTransacted(UObject* Object, const FTransactionObjectEvent& Event);
 		void MakeDebugTargetMenu(UToolMenu* InToolMenu);
+		TSharedPtr<SWidget> GenerateRowContextMenu();
 	
 		/** Create the properties tab and its content */
 		TSharedRef<SDockTab> SpawnPropertiesTab( const FSpawnTabArgs& Args );
@@ -173,9 +182,13 @@ namespace UE::ChooserEditor
 		static void RegisterWidgets();
 		
 		static FName EditorName;
+		static FName ContextMenuName;
 	};
+
+
 }
 
+// todo: for menus to actually be extensible this needs to be somewhere public
 UCLASS()
 class UChooserEditorToolMenuContext : public UObject
 {
@@ -183,4 +196,4 @@ class UChooserEditorToolMenuContext : public UObject
 
 public:
 	TWeakPtr<UE::ChooserEditor::FChooserTableEditor> ChooserEditor;
-};
+};
\ No newline at end of file
diff --git a/Engine/Plugins/Experimental/Chooser/Source/ChooserEditor/Private/ChooserTableEditorCommands.h b/Engine/Plugins/Experimental/Chooser/Source/ChooserEditor/Private/ChooserTableEditorCommands.h
--- a/Engine/Plugins/Experimental/Chooser/Source/ChooserEditor/Private/ChooserTableEditorCommands.h
+++ b/Engine/Plugins/Experimental/Chooser/Source/ChooserEditor/Private/ChooserTableEditorCommands.h
@@ -18,11 +18,17 @@ public:
 	}
 
 	TSharedPtr<FUICommandInfo> EditChooserSettings;
+	TSharedPtr<FUICommandInfo> RemoveDisabledData;
+	TSharedPtr<FUICommandInfo> Delete;
+	TSharedPtr<FUICommandInfo> Disable;
 
 	/** Initialize commands */
 	virtual void RegisterCommands() override
 	{
 		UI_COMMAND(EditChooserSettings, "Table Settings", "Edit the root properties of the ChooserTable asset.", EUserInterfaceActionType::Button, FInputChord())
+		UI_COMMAND(RemoveDisabledData, "Remove Disabled Data", "Delete all data that's marked as disabled.", EUserInterfaceActionType::Button, FInputChord())
+		UI_COMMAND(Delete, "Delete", "Delete the selected Rows or Column.", EUserInterfaceActionType::Button, FInputChord())
+		UI_COMMAND(Disable, "Disable", "Disable the selected Rows or Column.", EUserInterfaceActionType::Check, FInputChord())
 	}
 };
 	
diff --git a/Engine/Plugins/Experimental/Chooser/Source/ChooserEditor/Private/SChooserRowHandle.cpp b/Engine/Plugins/Experimental/Chooser/Source/ChooserEditor/Private/SChooserRowHandle.cpp
--- a/Engine/Plugins/Experimental/Chooser/Source/ChooserEditor/Private/SChooserRowHandle.cpp
+++ b/Engine/Plugins/Experimental/Chooser/Source/ChooserEditor/Private/SChooserRowHandle.cpp
@@ -21,35 +21,38 @@ namespace UE::ChooserEditor
 		return Operation;
 	}
 	
-	void SChooserRowHandle::Construct(const FArguments& InArgs)
+	void SChooserRowHandle::Construct(const FArguments& InArgs, bool bShowImage)
 	{
 		ChooserEditor = InArgs._ChooserEditor;
 		RowIndex = InArgs._RowIndex;
 
-		ChildSlot
-		[
-			SNew(SOverlay)
-			+ SOverlay::Slot()
+		if (bShowImage)
+		{
+			ChildSlot
 			[
-				SNew(SBox) .Padding(0.0f) .HAlign(HAlign_Center) .VAlign(VAlign_Center) .WidthOverride(16.0f)
+				SNew(SOverlay)
+				+ SOverlay::Slot()
 				[
-					SNew(SImage)
-					.Image(FCoreStyle::Get().GetBrush("VerticalBoxDragIndicatorShort"))
+					SNew(SBox) .Padding(0.0f) .HAlign(HAlign_Center) .VAlign(VAlign_Center) .WidthOverride(16.0f)
+					[
+						SNew(SImage)
+						.Image(FCoreStyle::Get().GetBrush("VerticalBoxDragIndicatorShort"))
+					]
 				]
-			]
-			+ SOverlay::Slot()
-			[
-				SNew(SBox).Padding(0.0f) .HAlign(HAlign_Center) .VAlign(VAlign_Center) .WidthOverride(16.0f)
+				+ SOverlay::Slot()
 				[
-					SNew(SImage)
-					.Visibility_Lambda([this]()
-					{
-						return ChooserEditor->GetChooser()->bDebugTestValuesValid && RowIndex == ChooserEditor->GetChooser()->GetDebugSelectedRow() ? EVisibility::HitTestInvisible : EVisibility::Hidden;
-					})
-					.Image(FAppStyle::Get().GetBrush("Icons.ArrowRight"))
+					SNew(SBox).Padding(0.0f) .HAlign(HAlign_Center) .VAlign(VAlign_Center) .WidthOverride(16.0f)
+					[
+						SNew(SImage)
+						.Visibility_Lambda([this]()
+						{
+							return ChooserEditor->GetChooser()->bDebugTestValuesValid && RowIndex == ChooserEditor->GetChooser()->GetDebugSelectedRow() ? EVisibility::HitTestInvisible : EVisibility::Hidden;
+						})
+						.Image(FAppStyle::Get().GetBrush("Icons.ArrowRight"))
+					]
 				]
-			]
-		];
+			];
+		}
 	}
 	
 	FReply SChooserRowHandle::OnMouseButtonDoubleClick(const FGeometry& InMyGeometry, const FPointerEvent& InMouseEvent)
@@ -69,7 +72,7 @@ namespace UE::ChooserEditor
 	FReply SChooserRowHandle::OnMouseButtonDown(const FGeometry& MyGeometry, const FPointerEvent& MouseEvent)
 	{
 		// act as a move handle if the row is already selected, and if there are multiselect modifiers pressed
-		if (!MouseEvent.IsControlDown() && !MouseEvent.IsShiftDown()
+		if (MouseEvent.GetEffectingButton() != EKeys::RightMouseButton && !MouseEvent.IsControlDown() && !MouseEvent.IsShiftDown()
 			&& ChooserEditor->IsRowSelected(RowIndex))
 		{
 			return FReply::Handled().DetectDrag(SharedThis(this), EKeys::LeftMouseButton);
diff --git a/Engine/Plugins/Experimental/Chooser/Source/ChooserEditor/Private/SChooserRowHandle.h b/Engine/Plugins/Experimental/Chooser/Source/ChooserEditor/Private/SChooserRowHandle.h
--- a/Engine/Plugins/Experimental/Chooser/Source/ChooserEditor/Private/SChooserRowHandle.h
+++ b/Engine/Plugins/Experimental/Chooser/Source/ChooserEditor/Private/SChooserRowHandle.h
@@ -30,7 +30,7 @@ public:
 	SLATE_ARGUMENT(uint32, RowIndex)
 	SLATE_END_ARGS()
 
-	void Construct(const FArguments& InArgs);
+	void Construct(const FArguments& InArgs, bool bShowImage);
 	
 	virtual FReply OnMouseButtonDoubleClick(const FGeometry& InMyGeometry, const FPointerEvent& InMouseEvent) override;
 	FReply OnMouseButtonDown(const FGeometry& MyGeometry, const FPointerEvent& MouseEvent) override;
diff --git a/Engine/Plugins/Experimental/Chooser/Source/ChooserEditor/Private/SChooserTableRow.cpp b/Engine/Plugins/Experimental/Chooser/Source/ChooserEditor/Private/SChooserTableRow.cpp
--- a/Engine/Plugins/Experimental/Chooser/Source/ChooserEditor/Private/SChooserTableRow.cpp
+++ b/Engine/Plugins/Experimental/Chooser/Source/ChooserEditor/Private/SChooserTableRow.cpp
@@ -20,6 +20,10 @@
 
 namespace UE::ChooserEditor
 {
+	static const FLinearColor DisabledColor(0.0105,0.0105,0.0105,0.5);
+	static const FLinearColor TestPassedColor(0.0,1.0, 0.0,0.3);
+	static const FLinearColor TestFailedColor(1.0,0.0, 0.0,0.2);
+	
 	void SChooserTableRow::Construct(const FArguments& Args, const TSharedRef<STableViewBase>& OwnerTableView)
 	{
 		RowIndex = Args._Entry;
@@ -77,13 +81,34 @@ namespace UE::ChooserEditor
 	{
 		static FName Result = "Result";
 		static FName Handles = "Handles";
+		static FName AddColumn = "Add";
 	
 		if (Chooser->ResultsStructs.IsValidIndex(RowIndex->RowIndex))
 		{
-			if (ColumnName == Handles)
+			if (ColumnName == AddColumn || ColumnName == Handles)
 			{
+				bool bShowHandleImage = ColumnName == Handles;
+				
 				// row drag handle
-				return SNew(SChooserRowHandle).ChooserEditor(Editor).RowIndex(RowIndex->RowIndex);
+				return SNew(SOverlay)
+						+ SOverlay::Slot()
+						[
+							SNew(SChooserRowHandle, bShowHandleImage).ChooserEditor(Editor).RowIndex(RowIndex->RowIndex)
+						]
+						+ SOverlay::Slot()
+						[
+							SNew(SColorBlock).Color(DisabledColor)
+									.Visibility_Lambda(
+									[this]()
+									{
+										if (Chooser->IsRowDisabled(RowIndex->RowIndex))
+										{
+											return EVisibility::HitTestInvisible;
+										}
+										return EVisibility::Hidden;
+									})
+									
+						];
 			}
 			else if (ColumnName == Result) 
 			{
@@ -99,8 +124,26 @@ namespace UE::ChooserEditor
 				}),
 				&CacheBorder
 				);
-			
-				return ResultWidget.ToSharedRef();
+
+				return SNew(SOverlay)
+							+ SOverlay::Slot()
+    						[
+    							ResultWidget.ToSharedRef()
+    						]
+							+ SOverlay::Slot()
+							[
+								SNew(SColorBlock).Color(DisabledColor)
+										.Visibility_Lambda(
+										[this]()
+										{
+											if (Chooser->IsRowDisabled(RowIndex->RowIndex))
+											{
+												return EVisibility::HitTestInvisible;
+											}
+											return EVisibility::Hidden;
+										})
+										
+							];
 			}
 			else
 			{
@@ -128,16 +171,32 @@ namespace UE::ChooserEditor
 										{
 											if (Column->EditorTestFilter(RowIndex->RowIndex))
 											{
-												return FLinearColor(0.0,1.0,0.0,0.30);
+												return TestPassedColor;
 											}
 											else
 											{
-												return FLinearColor(1.0,0.0,0.0,0.20);
+												return TestFailedColor;
 											}
 										}
 										return FLinearColor::Transparent;
 									})
-						];
+									
+						]
+						+ SOverlay::Slot()
+						[
+							SNew(SColorBlock).Color(DisabledColor)
+									.Visibility_Lambda(
+									[this,Column]()
+									{
+										if (Chooser->IsRowDisabled(RowIndex->RowIndex) || Column->bDisabled)
+										{
+											return EVisibility::HitTestInvisible;
+										}
+										return EVisibility::Hidden;
+									})
+									
+						]
+						;
 					}
 				}
 			}
