Set Chooser plugin to beta

diff --git a/Engine/Plugins/Experimental/Chooser/Chooser.uplugin b/Engine/Plugins/Experimental/Chooser/Chooser.uplugin
--- a/Engine/Plugins/Experimental/Chooser/Chooser.uplugin
+++ b/Engine/Plugins/Experimental/Chooser/Chooser.uplugin
@@ -1,7 +1,7 @@
 {
 	"FileVersion" : 3,
 	"Version" : 1,
-	"VersionName" : "0.1",
+	"VersionName" : "1.0",
 	"FriendlyName" : "Chooser",
 	"Description" : "Chooser",
 	"Category" : "Animation",
@@ -12,8 +12,8 @@
 	"SupportURL" : "",
 	"EnabledByDefault" : false,
 	"CanContainContent" : true,
-	"IsBetaVersion" : false,
-	"IsExperimentalVersion" : true,
+	"IsBetaVersion" : true,
+	"IsExperimentalVersion" : false,
 	"Installed" : false,
 	"Plugins": [
 		{
