Gameplay Tags Column exact match and invert match features

diff --git a/Engine/Plugins/Experimental/Chooser/Source/Chooser/Internal/GameplayTagColumn.h b/Engine/Plugins/Experimental/Chooser/Source/Chooser/Internal/GameplayTagColumn.h
--- a/Engine/Plugins/Experimental/Chooser/Source/Chooser/Internal/GameplayTagColumn.h
+++ b/Engine/Plugins/Experimental/Chooser/Source/Chooser/Internal/GameplayTagColumn.h
@@ -51,16 +51,25 @@ struct CHOOSER_API FGameplayTagColumn : public FChooserColumnBase
 	UPROPERTY(EditAnywhere, Category="Data")
 	EGameplayContainerMatchType	TagMatchType = EGameplayContainerMatchType::Any;
 
+	//	If true, leaf tags must match exactly.
+	UPROPERTY(EditAnywhere, Category="Data")
+	bool bMatchExact = false;
+
+	//	If true, rows that pass the normal tag filter will be rejected, and vice versa
+	UPROPERTY(EditAnywhere, Category="Data")
+	bool bInvertMatchingLogic = false;
+
 #if WITH_EDITORONLY_DATA
 	UPROPERTY(EditAnywhere, Category=Runtime)
 	FGameplayTagContainer DefaultRowValue;
-
 #endif
 	
 	UPROPERTY(EditAnywhere, Category=Runtime)
 	// array of results (cells for this column for each row in the table)
 	// should match the length of the Results array 
 	TArray<FGameplayTagContainer> RowValues;
+
+	bool TestRow(int32 RowIndex, const FGameplayTagContainer& Value) const;
 	
 	virtual void Filter(FChooserEvaluationContext& Context, const FChooserIndexArray& IndexListIn, FChooserIndexArray& IndexListOut) const override;
 
@@ -68,18 +77,7 @@ struct CHOOSER_API FGameplayTagColumn : public FChooserColumnBase
 	mutable FGameplayTagContainer TestValue;
 	virtual bool EditorTestFilter(int32 RowIndex) const override
 	{
-		if(RowValues.IsValidIndex(RowIndex))
-		{
-			if (TagMatchType == EGameplayContainerMatchType::All)
-			{
-				return TestValue.HasAll(RowValues[RowIndex]);
-			}
-			else
-			{
-				return TestValue.HasAny(RowValues[RowIndex]);
-			}
-		}
-		return false;
+		return TestRow(RowIndex, TestValue);
 	}
 
 	virtual void SetTestValue(TArrayView<const uint8> Value) override
diff --git a/Engine/Plugins/Experimental/Chooser/Source/Chooser/Private/GameplayTagColumn.cpp b/Engine/Plugins/Experimental/Chooser/Source/Chooser/Private/GameplayTagColumn.cpp
--- a/Engine/Plugins/Experimental/Chooser/Source/Chooser/Private/GameplayTagColumn.cpp
+++ b/Engine/Plugins/Experimental/Chooser/Source/Chooser/Private/GameplayTagColumn.cpp
@@ -18,6 +18,45 @@ FGameplayTagColumn::FGameplayTagColumn()
 	InputValue.InitializeAs(FGameplayTagContextProperty::StaticStruct());
 }
 
+bool FGameplayTagColumn::TestRow(int32 RowIndex, const FGameplayTagContainer& Value) const
+{
+	if(RowValues.IsValidIndex(RowIndex))
+	{
+		bool bPasses = false;	
+
+		if (RowValues[RowIndex].IsEmpty())
+		{
+			//	An empty container should always pass in inverted mode, but we are going to invert later, so invert here.
+			bPasses = !bInvertMatchingLogic;
+		}
+		else
+		{
+			if (TagMatchType == EGameplayContainerMatchType::All)
+			{
+				if (bMatchExact ? Value.HasAllExact(RowValues[RowIndex]) : Value.HasAll(RowValues[RowIndex]))
+				{
+					bPasses = true;
+				}
+			}
+			else
+			{
+				if (bMatchExact ? Value.HasAnyExact(RowValues[RowIndex]) : Value.HasAny(RowValues[RowIndex]))
+				{
+					bPasses = true;
+				}
+			}
+		}
+
+		if (bInvertMatchingLogic)
+		{
+			bPasses = !bPasses;
+		}
+
+		return bPasses;	
+	}
+	return false;
+}
+
 void FGameplayTagColumn::Filter(FChooserEvaluationContext& Context, const FChooserIndexArray& IndexListIn, FChooserIndexArray& IndexListOut) const
 {
 	const FGameplayTagContainer* Result = nullptr;
@@ -34,30 +73,9 @@ void FGameplayTagColumn::Filter(FChooserEvaluationContext& Context, const FChoos
 
 		for (uint32 Index : IndexListIn)
 		{
-		 
-			if (RowValues.Num() > (int)Index)
+			if (TestRow(Index, *Result))
 			{
-				if (RowValues[Index].IsEmpty())
-				{
-					IndexListOut.Push(Index);
-				}
-				else
-				{
-					if (TagMatchType == EGameplayContainerMatchType::All)
-					{
-						if (Result->HasAll(RowValues[Index]))
-						{
-							IndexListOut.Push(Index);
-						}
-					}
-					else
-					{
-						if (Result->HasAny(RowValues[Index]))
-						{
-							IndexListOut.Push(Index);
-						}
-					}
-				}
+				IndexListOut.Push(Index);
 			}
 		}
 	}
