Soft Asset references in Chooser Table

diff --git a/Engine/Plugins/Experimental/Chooser/Source/Chooser/Internal/ObjectChooser_Asset.h b/Engine/Plugins/Experimental/Chooser/Source/Chooser/Internal/ObjectChooser_Asset.h
--- a/Engine/Plugins/Experimental/Chooser/Source/Chooser/Internal/ObjectChooser_Asset.h
+++ b/Engine/Plugins/Experimental/Chooser/Source/Chooser/Internal/ObjectChooser_Asset.h
@@ -18,6 +18,18 @@ public:
 	TObjectPtr<UObject> Asset;
 };
 
+USTRUCT(DisplayName = "Soft Asset")
+struct CHOOSER_API FSoftAssetChooser : public FObjectChooserBase
+{
+	GENERATED_BODY()
+
+	// FObjectChooserBase interface
+	virtual UObject* ChooseObject(FChooserEvaluationContext& Context) const final override;
+public: 
+	UPROPERTY(EditAnywhere, Category = "Parameters")
+	TSoftObjectPtr<UObject> Asset;
+};
+
 // deprecated class for upgrading old data
 UCLASS(ClassGroup = "LiveLink", deprecated)
 class CHOOSER_API UDEPRECATED_ObjectChooser_Asset : public UObject, public IObjectChooser
diff --git a/Engine/Plugins/Experimental/Chooser/Source/Chooser/Private/ObjectChooser_Asset.cpp b/Engine/Plugins/Experimental/Chooser/Source/Chooser/Private/ObjectChooser_Asset.cpp
--- a/Engine/Plugins/Experimental/Chooser/Source/Chooser/Private/ObjectChooser_Asset.cpp
+++ b/Engine/Plugins/Experimental/Chooser/Source/Chooser/Private/ObjectChooser_Asset.cpp
@@ -5,4 +5,9 @@
 UObject* FAssetChooser::ChooseObject(FChooserEvaluationContext& Context) const
 {
 	return Asset;
-}
\ No newline at end of file
+}
+
+UObject* FSoftAssetChooser::ChooseObject(FChooserEvaluationContext& Context) const
+{
+	return Asset.LoadSynchronous();
+}
diff --git a/Engine/Plugins/Experimental/Chooser/Source/ChooserEditor/Private/ChooserEditorWidgets.cpp b/Engine/Plugins/Experimental/Chooser/Source/ChooserEditor/Private/ChooserEditorWidgets.cpp
--- a/Engine/Plugins/Experimental/Chooser/Source/ChooserEditor/Private/ChooserEditorWidgets.cpp
+++ b/Engine/Plugins/Experimental/Chooser/Source/ChooserEditor/Private/ChooserEditorWidgets.cpp
@@ -44,6 +44,22 @@ TSharedRef<SWidget> CreateAssetWidget(bool bReadOnly, UObject* TransactionObject
 		});
 }
 	
+TSharedRef<SWidget> CreateSoftAssetWidget(bool bReadOnly, UObject* TransactionObject, void* Value, UClass* ResultBaseClass, FChooserWidgetValueChanged ValueChanged)
+{
+	FSoftAssetChooser* DIAsset = static_cast<FSoftAssetChooser*>(Value);
+
+	return SNew(SObjectPropertyEntryBox)
+		.IsEnabled(!bReadOnly)
+		.AllowedClass(ResultBaseClass ? ResultBaseClass : UObject::StaticClass())
+		.ObjectPath_Lambda([DIAsset](){ return DIAsset->Asset.ToSoftObjectPath().ToString(); })
+		.OnObjectChanged_Lambda([TransactionObject, DIAsset](const FAssetData& AssetData)
+		{
+			const FScopedTransaction Transaction(LOCTEXT("Edit Asset", "Edit Asset"));
+			TransactionObject->Modify(true);
+			DIAsset->Asset = AssetData.GetAsset();
+		});
+}
+	
 TSharedRef<SWidget> CreateClassWidget(bool bReadOnly, UObject* TransactionObject, void* Value, UClass* ResultBaseClass, FChooserWidgetValueChanged ValueChanged)
 {
 	FClassChooser* ClassChooser = static_cast<FClassChooser*>(Value);
diff --git a/Engine/Plugins/Experimental/Chooser/Source/ChooserEditor/Private/ChooserEditorWidgets.h b/Engine/Plugins/Experimental/Chooser/Source/ChooserEditor/Private/ChooserEditorWidgets.h
--- a/Engine/Plugins/Experimental/Chooser/Source/ChooserEditor/Private/ChooserEditorWidgets.h
+++ b/Engine/Plugins/Experimental/Chooser/Source/ChooserEditor/Private/ChooserEditorWidgets.h
@@ -7,6 +7,7 @@ namespace UE::ChooserEditor
 {
 
 TSharedRef<SWidget> CreateAssetWidget(bool bReadOnly, UObject* TransactionObject, void* Value, UClass* ResultBaseClass, FChooserWidgetValueChanged ValueChanged);
+TSharedRef<SWidget> CreateSoftAssetWidget(bool bReadOnly, UObject* TransactionObject, void* Value, UClass* ResultBaseClass, FChooserWidgetValueChanged ValueChanged);
 TSharedRef<SWidget> CreateClassWidget(bool bReadOnly, UObject* TransactionObject, void* Value, UClass* ResultBaseClass, FChooserWidgetValueChanged ValueChanged);
 TSharedRef<SWidget> CreateEvaluateChooserWidget(bool bReadOnly, UObject* TransactionObject, void* Value, UClass* ResultBaseClass, FChooserWidgetValueChanged ValueChanged);
 TSharedRef<SWidget> CreateNestedChooserWidget(bool bReadOnly, UObject* TransactionObject, void* Value, UClass* ResultBaseClass, FChooserWidgetValueChanged ValueChanged);
diff --git a/Engine/Plugins/Experimental/Chooser/Source/ChooserEditor/Private/ChooserTableEditor.cpp b/Engine/Plugins/Experimental/Chooser/Source/ChooserEditor/Private/ChooserTableEditor.cpp
--- a/Engine/Plugins/Experimental/Chooser/Source/ChooserEditor/Private/ChooserTableEditor.cpp
+++ b/Engine/Plugins/Experimental/Chooser/Source/ChooserEditor/Private/ChooserTableEditor.cpp
@@ -1479,6 +1479,7 @@ void FChooserTableEditor::DeleteSelection()
 void FChooserTableEditor::RegisterWidgets()
 {
 	FObjectChooserWidgetFactories::RegisterWidgetCreator(FAssetChooser::StaticStruct(), CreateAssetWidget);
+	FObjectChooserWidgetFactories::RegisterWidgetCreator(FSoftAssetChooser::StaticStruct(), CreateSoftAssetWidget);
 	FObjectChooserWidgetFactories::RegisterWidgetCreator(FClassChooser::StaticStruct(), CreateClassWidget);
 	FObjectChooserWidgetFactories::RegisterWidgetCreator(FEvaluateChooser::StaticStruct(), CreateEvaluateChooserWidget);
 	FObjectChooserWidgetFactories::RegisterWidgetCreator(FNestedChooser::StaticStruct(), CreateNestedChooserWidget);
