Fix crashes on delete row in proxy table editor

diff --git a/Engine/Plugins/Experimental/Chooser/Source/ProxyTableEditor/Private/ProxyTableEditor.cpp b/Engine/Plugins/Experimental/Chooser/Source/ProxyTableEditor/Private/ProxyTableEditor.cpp
--- a/Engine/Plugins/Experimental/Chooser/Source/ProxyTableEditor/Private/ProxyTableEditor.cpp
+++ b/Engine/Plugins/Experimental/Chooser/Source/ProxyTableEditor/Private/ProxyTableEditor.cpp
@@ -719,7 +719,8 @@ void FProxyTableEditor::DeleteSelectedRows()
 	{
 		ProxyTable->Entries.RemoveAt(RowIndex);
 	}
-	
+
+	ClearSelectedRows();
 	UpdateTableRows();
 }
 	
