OutputObjectColumn for Choosers

diff --git a/Engine/Plugins/Experimental/Chooser/Source/Chooser/Internal/ObjectColumn.h b/Engine/Plugins/Experimental/Chooser/Source/Chooser/Internal/ObjectColumn.h
--- a/Engine/Plugins/Experimental/Chooser/Source/Chooser/Internal/ObjectColumn.h
+++ b/Engine/Plugins/Experimental/Chooser/Source/Chooser/Internal/ObjectColumn.h
@@ -20,6 +20,7 @@ struct CHOOSER_API FObjectContextProperty : public FChooserParameterObjectBase
 	FChooserObjectPropertyBinding Binding;
 
 	virtual bool GetValue(FChooserEvaluationContext& Context, FSoftObjectPath& OutResult) const override;
+	virtual bool SetValue(FChooserEvaluationContext& Context, UObject* OutResult) const override;
 
 	CHOOSER_PARAMETER_BOILERPLATE();
 
diff --git a/Engine/Plugins/Experimental/Chooser/Source/Chooser/Internal/OutputObjectColumn.h b/Engine/Plugins/Experimental/Chooser/Source/Chooser/Internal/OutputObjectColumn.h
new file mode 100644
--- /dev/null
+++ b/Engine/Plugins/Experimental/Chooser/Source/Chooser/Internal/OutputObjectColumn.h
@@ -0,0 +1,71 @@
+// Copyright Epic Games, Inc. All Rights Reserved.
+#pragma once
+
+#include "ChooserPropertyAccess.h"
+#include "CoreMinimal.h"
+#include "IChooserColumn.h"
+#include "IChooserParameterObject.h"
+#include "InstancedStruct.h"
+#include "Serialization/MemoryReader.h"
+#include "OutputObjectColumn.generated.h"
+
+struct FBindingChainElement;
+
+USTRUCT()
+struct FChooserOutputObjectRowData
+{
+	GENERATED_BODY()
+
+	UPROPERTY(EditAnywhere, Category="Value", Meta = (ExcludeBaseStruct, BaseStruct = "/Script/Chooser.ObjectChooserBase"))
+	FInstancedStruct Value;
+};
+
+USTRUCT()
+struct CHOOSER_API FOutputObjectColumn : public FChooserColumnBase
+{
+	GENERATED_BODY()
+
+	FOutputObjectColumn();
+
+	UPROPERTY(EditAnywhere, NoClear, Meta = (ExcludeBaseStruct, BaseStruct = "/Script/Chooser.ChooserParameterObjectBase"), Category = "Data")
+	FInstancedStruct InputValue;
+
+#if WITH_EDITORONLY_DATA
+	UPROPERTY(EditAnywhere, Category = "Data")
+	FChooserOutputObjectRowData DefaultRowValue;
+#endif
+
+	UPROPERTY(EditAnywhere, Category = "Data")
+	// array of results (cells for this column for each row in the table)
+	// should match the length of the Results array
+	TArray<FChooserOutputObjectRowData> RowValues;
+
+	// FallbackValue will be used as the output value if the all rows in the chooser fail, and the FallbackResult from the chooser is used.
+    UPROPERTY(EditAnywhere, Category=Data);
+	FChooserOutputObjectRowData FallbackValue;
+
+	virtual bool HasFilters() const override { return false; }
+	virtual bool HasOutputs() const override { return true; }
+	virtual void SetOutputs(FChooserEvaluationContext& Context, int RowIndex) const override;
+
+#if WITH_EDITOR
+	virtual void AddToDetails(FInstancedPropertyBag& PropertyBag, int32 ColumnIndex, int32 RowIndex) override;
+	virtual void SetFromDetails(FInstancedPropertyBag& PropertyBag, int32 ColumnIndex, int32 RowIndex) override;
+#endif
+
+	CHOOSER_COLUMN_BOILERPLATE(FChooserParameterObjectBase);
+
+	virtual void Compile(IHasContextClass* Owner, bool bForce) override;
+
+#if WITH_EDITORONLY_DATA
+	virtual void PostLoad() override
+	{
+		Super::PostLoad();
+
+		if (InputValue.IsValid())
+		{
+			InputValue.GetMutable<FChooserParameterBase>().PostLoad();
+		}
+	}
+#endif
+};
diff --git a/Engine/Plugins/Experimental/Chooser/Source/Chooser/Private/ObjectColumn.cpp b/Engine/Plugins/Experimental/Chooser/Source/Chooser/Private/ObjectColumn.cpp
--- a/Engine/Plugins/Experimental/Chooser/Source/Chooser/Private/ObjectColumn.cpp
+++ b/Engine/Plugins/Experimental/Chooser/Source/Chooser/Private/ObjectColumn.cpp
@@ -25,8 +25,31 @@ bool FObjectContextProperty::GetValue(FChooserEvaluationContext& Context, FSoftO
 			}
 		
 			// otherwise get the value from the object property and convert to a soft object path
-			const UObject* LoadedObject = reinterpret_cast<const UObject*>(Result.Container + Result.PropertyOffset);
-			OutResult = LoadedObject;
+			const UObject** LoadedObject = reinterpret_cast<const UObject**>(Result.Container + Result.PropertyOffset);
+			OutResult = *LoadedObject;
+			return true;
+		}
+	}
+
+	return false;
+}
+
+bool FObjectContextProperty::SetValue(FChooserEvaluationContext& Context, UObject* OutputObject) const
+{
+	UE::Chooser::FResolvedPropertyChainResult Result;
+	if (UE::Chooser::ResolvePropertyChain(Context, Binding, Result))
+	{
+		if (Result.Function == nullptr)
+		{
+			if (Result.PropertyType == UE::Chooser::EChooserPropertyAccessType::SoftObjectRef)
+			{
+				FSoftObjectPtr& SoftObjectPtr = *reinterpret_cast<FSoftObjectPtr*>(Result.Container + Result.PropertyOffset);
+				SoftObjectPtr = OutputObject;
+				return true;
+			}
+		
+			UObject** ObjectPtr = reinterpret_cast<UObject**>(Result.Container + Result.PropertyOffset);
+			*ObjectPtr = OutputObject;
 			return true;
 		}
 	}
diff --git a/Engine/Plugins/Experimental/Chooser/Source/Chooser/Private/OutputObjectColumn.cpp b/Engine/Plugins/Experimental/Chooser/Source/Chooser/Private/OutputObjectColumn.cpp
new file mode 100644
--- /dev/null
+++ b/Engine/Plugins/Experimental/Chooser/Source/Chooser/Private/OutputObjectColumn.cpp
@@ -0,0 +1,76 @@
+// Copyright Epic Games, Inc. All Rights Reserved.
+#include "OutputObjectColumn.h"
+#include "ObjectColumn.h"
+#include "ChooserPropertyAccess.h"
+#include "ChooserTrace.h"
+
+#if WITH_EDITOR
+#include "IPropertyAccessEditor.h"
+#include "PropertyBag.h"
+#endif
+
+FOutputObjectColumn::FOutputObjectColumn()
+{
+#if WITH_EDITOR
+	InputValue.InitializeAs(FObjectContextProperty::StaticStruct());
+#endif
+}
+
+void FOutputObjectColumn::Compile(IHasContextClass* Owner, bool bForce)
+{
+	FChooserColumnBase::Compile(Owner, bForce);
+
+	for(FChooserOutputObjectRowData& Value : RowValues)
+	{
+		if (Value.Value.IsValid())
+		{
+			FObjectChooserBase& Result = Value.Value.GetMutable<FObjectChooserBase>();
+			Result.Compile(Owner, bForce);
+		}
+	}
+}
+
+void FOutputObjectColumn::SetOutputs(FChooserEvaluationContext& Context, int RowIndex) const
+{
+	if (RowValues.IsValidIndex(RowIndex))
+	{
+		if (const FObjectChooserBase* Value = RowValues[RowIndex].Value.GetPtr<FObjectChooserBase>())
+		{
+			UObject* Result = Value->ChooseObject(Context);
+			InputValue.Get<FChooserParameterObjectBase>().SetValue(Context, Result);
+		}
+	}
+	else
+	{
+		if (const FObjectChooserBase* Value = FallbackValue.Value.GetPtr<FObjectChooserBase>())
+		{
+			UObject* Result = Value->ChooseObject(Context);
+			InputValue.Get<FChooserParameterObjectBase>().SetValue(Context, Result);
+		}
+	}
+}
+
+#if WITH_EDITOR
+	void FOutputObjectColumn::AddToDetails(FInstancedPropertyBag& PropertyBag, int32 ColumnIndex, int32 RowIndex)
+	{
+		FText DisplayName;
+		InputValue.Get<FChooserParameterObjectBase>().GetDisplayName(DisplayName);
+		FName PropertyName("RowData",ColumnIndex);
+		FPropertyBagPropertyDesc PropertyDesc(PropertyName,  EPropertyBagPropertyType::Struct, FChooserOutputObjectRowData::StaticStruct());
+		PropertyDesc.MetaData.Add(FPropertyBagPropertyDescMetaData("DisplayName", DisplayName.ToString()));
+		PropertyBag.AddProperties({PropertyDesc});
+		PropertyBag.SetValueStruct(PropertyName, RowValues[RowIndex]);
+	}
+
+	void FOutputObjectColumn::SetFromDetails(FInstancedPropertyBag& PropertyBag, int32 ColumnIndex, int32 RowIndex)
+	{
+		FName PropertyName("RowData", ColumnIndex);
+		
+		TValueOrError<FStructView, EPropertyBagResult> Result = PropertyBag.GetValueStruct(PropertyName, FChooserOutputObjectRowData::StaticStruct());
+		if (FStructView* StructView = Result.TryGetValue())
+		{
+			RowValues[RowIndex] = StructView->Get<FChooserOutputObjectRowData>();
+		}
+	}
+
+#endif
diff --git a/Engine/Plugins/Experimental/Chooser/Source/Chooser/Public/IChooserParameterObject.h b/Engine/Plugins/Experimental/Chooser/Source/Chooser/Public/IChooserParameterObject.h
--- a/Engine/Plugins/Experimental/Chooser/Source/Chooser/Public/IChooserParameterObject.h
+++ b/Engine/Plugins/Experimental/Chooser/Source/Chooser/Public/IChooserParameterObject.h
@@ -14,6 +14,7 @@ struct FChooserParameterObjectBase : public FChooserParameterBase
 	GENERATED_BODY()
 
 	virtual bool GetValue(FChooserEvaluationContext& Context, FSoftObjectPath& OutResult) const { return false; }
+	virtual bool SetValue(FChooserEvaluationContext& Context, UObject* OutResult) const { return false; }
 
 #if WITH_EDITOR
 	virtual UClass* GetAllowedClass() const { return UObject::StaticClass(); }
diff --git a/Engine/Plugins/Experimental/Chooser/Source/ChooserEditor/Private/ChooserEditorModule.cpp b/Engine/Plugins/Experimental/Chooser/Source/ChooserEditor/Private/ChooserEditorModule.cpp
--- a/Engine/Plugins/Experimental/Chooser/Source/ChooserEditor/Private/ChooserEditorModule.cpp
+++ b/Engine/Plugins/Experimental/Chooser/Source/ChooserEditor/Private/ChooserEditorModule.cpp
@@ -18,6 +18,7 @@
 #include "MultiEnumColumnEditor.h"
 #include "ObjectColumnEditor.h"
 #include "OutputFloatColumnEditor.h"
+#include "OutputObjectColumnEditor.h"
 #include "OutputStructColumnEditor.h"
 #include "PropertyAccessChainCustomization.h"
 #include "PropertyEditorModule.h"
@@ -43,6 +44,7 @@ void FModule::StartupModule()
 	RegisterEnumWidgets();
 	RegisterMultiEnumWidgets();
 	RegisterObjectWidgets();
+	RegisterOutputObjectWidgets();
 	RegisterStructWidgets();
 	RegisterRandomizeWidgets();
 	
diff --git a/Engine/Plugins/Experimental/Chooser/Source/ChooserEditor/Private/ObjectChooserWidgetFactories.cpp b/Engine/Plugins/Experimental/Chooser/Source/ChooserEditor/Private/ObjectChooserWidgetFactories.cpp
--- a/Engine/Plugins/Experimental/Chooser/Source/ChooserEditor/Private/ObjectChooserWidgetFactories.cpp
+++ b/Engine/Plugins/Experimental/Chooser/Source/ChooserEditor/Private/ObjectChooserWidgetFactories.cpp
@@ -2,16 +2,17 @@
 
 #include "ObjectChooserWidgetFactories.h"
 
-#include "Modules/ModuleManager.h"
-#include "Widgets/Input/SHyperlink.h"
-#include "Widgets/Input/SNumericEntryBox.h"
-#include "Widgets/Input/SComboButton.h"
-#include "SClassViewer.h"
 #include "DetailCategoryBuilder.h"
 #include "IObjectChooser.h"
-#include "Widgets/Layout/SWidgetSwitcher.h"
 #include "ObjectChooserClassFilter.h"
+#include "SClassViewer.h"
+#include "ScopedTransaction.h"
 #include "StructViewerModule.h"
+#include "Modules/ModuleManager.h"
+#include "Widgets/Input/SComboButton.h"
+#include "Widgets/Input/SHyperlink.h"
+#include "Widgets/Input/SNumericEntryBox.h"
+#include "Widgets/Layout/SWidgetSwitcher.h"
 
 #define LOCTEXT_NAMESPACE "DataInterfaceEditor"
 
@@ -44,8 +45,8 @@ TSharedPtr<SWidget> FObjectChooserWidgetFactories::CreateColumnWidget(FChooserCo
 
 	return nullptr;
 }
-	
-	
+
+
 TSharedPtr<SWidget> FObjectChooserWidgetFactories::CreateWidget(bool bReadOnly, UObject* TransactionObject, void* Value, const UStruct* ValueType, UClass* ResultBaseClass, FChooserWidgetValueChanged ValueChanged)
 {
 	if (Value)
@@ -63,7 +64,137 @@ TSharedPtr<SWidget> FObjectChooserWidgetFactories::CreateWidget(bool bReadOnly,
 	return nullptr;
 }
 
-TSharedPtr<SWidget> FObjectChooserWidgetFactories::CreateWidget(bool bReadOnly, UObject* TransactionObject, const UScriptStruct* BaseType, void* Value, const UStruct* ValueType, UClass* ResultBaseClass, const FOnStructPicked& CreateClassCallback, TSharedPtr<SBorder>* InnerWidget, FChooserWidgetValueChanged ValueChanged, FText NullValueDisplayText)
+class SObjectChooserWidget : public SCompoundWidget
+{
+public:
+
+	SLATE_BEGIN_ARGS(SObjectChooserWidget)
+	{
+		_TransactionObject = nullptr;
+		_Data = nullptr;
+		_NullValueDisplayText =  LOCTEXT("SelectDataType", "Select Data Type...");
+		_bReadOnly = false;
+		
+	}
+
+	SLATE_ARGUMENT(UObject*, TransactionObject)
+	SLATE_ARGUMENT(FInstancedStruct*, Data)
+	SLATE_ARGUMENT(UScriptStruct*, DataBaseType)
+	SLATE_ARGUMENT(UClass*, ResultBaseClass)
+	SLATE_ARGUMENT(FText, NullValueDisplayText)
+	SLATE_ARGUMENT(bool, bReadOnly)
+	SLATE_EVENT(FChooserWidgetValueChanged, ValueChanged);
+			
+	SLATE_END_ARGS()
+
+	void Construct( const FArguments& InArgs)
+	{
+		TransactionObject = InArgs._TransactionObject;
+		Data = InArgs._Data;
+		DataBaseType = InArgs._DataBaseType;
+		ResultBaseClass = InArgs._ResultBaseClass;
+		bReadOnly = InArgs._bReadOnly;
+		ValueChanged = InArgs._ValueChanged;
+
+		if (!InArgs._NullValueDisplayText.IsEmpty())
+		{
+			NullValueDisplayText = InArgs._NullValueDisplayText;
+		}
+
+		Border = SNew(SBorder);
+		UpdateValueWidget();
+		
+		TSharedPtr<SWidget> Widget = Border;
+		
+		// don't need the type selector dropdown when read only
+		if (!bReadOnly)
+		{
+			// button for replacing data with a different Data Interface class
+			TSharedPtr<SComboButton> Button = SNew(SComboButton)
+					.ComboButtonStyle(FAppStyle::Get(), "SimpleComboButton");
+		
+			Button->SetOnGetMenuContent(FOnGetContent::CreateLambda([this, Button]()
+			{
+				FStructViewerInitializationOptions Options;
+				Options.StructFilter = MakeShared<FStructFilter>(DataBaseType);
+				Options.NameTypeToDisplay = EStructViewerNameTypeToDisplay::DisplayName;
+				Options.bShowNoneOption = true;
+			
+				TSharedRef<SWidget> StructViewerWidget = FModuleManager::LoadModuleChecked<FStructViewerModule>("StructViewer").CreateStructViewer(Options, FOnStructPicked::CreateLambda(
+					[this, Button](const UScriptStruct* ChosenStruct)
+				{
+					Button->SetIsOpen(false);
+					const FScopedTransaction Transaction(LOCTEXT("Change Object Type", "Change Object Type"));
+					TransactionObject->Modify(true);
+					Data->InitializeAs(ChosenStruct);
+					UpdateValueWidget();
+				}));
+				return StructViewerWidget;
+			}));
+
+			Widget = SNew(SHorizontalBox)
+			+SHorizontalBox::Slot().FillWidth(100)
+			[
+				Border.ToSharedRef()
+			]
+			+SHorizontalBox::Slot().AutoWidth()
+			[
+				Button.ToSharedRef()
+			];
+		}
+
+		ChildSlot
+		[
+			Widget.ToSharedRef()
+		];
+		
+	}
+
+	~SObjectChooserWidget()
+	{
+	}
+
+private:
+
+	void UpdateValueWidget()
+	{
+		TSharedPtr<SWidget> NewWidget = FObjectChooserWidgetFactories::CreateWidget(bReadOnly, TransactionObject, Data->GetMutableMemory(), Data->GetScriptStruct(), ResultBaseClass, ValueChanged);
+
+		if (!NewWidget.IsValid())
+		{
+			NewWidget = SNew(STextBlock)
+				.Font(FAppStyle::GetFontStyle("PropertyWindow.NormalFont"))
+				.Margin(2)
+				.Text(NullValueDisplayText);
+		}
+		
+		Border->SetContent(NewWidget.ToSharedRef());
+	}
+
+	
+	UObject* TransactionObject = nullptr;
+	FInstancedStruct* Data = nullptr;
+	UScriptStruct* DataBaseType = nullptr;
+	UClass* ResultBaseClass = nullptr;
+	TSharedPtr<SBorder> Border;
+	bool bReadOnly = false;
+	FChooserWidgetValueChanged ValueChanged;
+	FText NullValueDisplayText;
+};
+
+TSharedPtr<SWidget> FObjectChooserWidgetFactories::CreateWidget(bool bReadOnly, UObject* TransactionObject, UScriptStruct* DataBaseType, FInstancedStruct* Data, UClass* ResultBaseClass, FChooserWidgetValueChanged ValueChanged, FText NullValueDisplayText)
+{
+	return SNew(SObjectChooserWidget)
+		.bReadOnly(bReadOnly)
+		.TransactionObject(TransactionObject)
+		.Data(Data).DataBaseType(DataBaseType)
+		.ResultBaseClass(ResultBaseClass)
+		.ValueChanged(ValueChanged)
+		.NullValueDisplayText(NullValueDisplayText);
+}
+
+TSharedPtr<SWidget> FObjectChooserWidgetFactories::CreateWidget(bool bReadOnly, UObject* TransactionObject, const UScriptStruct* BaseType, void* Value, const UStruct* ValueType, UClass* ResultBaseClass,
+	const FOnStructPicked& CreateClassCallback, TSharedPtr<SBorder>* InnerWidget, FChooserWidgetValueChanged ValueChanged, FText NullValueDisplayText)
 {
 	TSharedPtr<SWidget> LeftWidget = CreateWidget(bReadOnly, TransactionObject, Value, ValueType, ResultBaseClass, ValueChanged);
 	if (bReadOnly)
@@ -128,7 +259,6 @@ TSharedPtr<SWidget> FObjectChooserWidgetFactories::CreateWidget(bool bReadOnly,
 		]
 	;
 
-
 	return Widget;
 }
 
diff --git a/Engine/Plugins/Experimental/Chooser/Source/ChooserEditor/Private/OutputObjectColumnEditor.cpp b/Engine/Plugins/Experimental/Chooser/Source/ChooserEditor/Private/OutputObjectColumnEditor.cpp
new file mode 100644
--- /dev/null
+++ b/Engine/Plugins/Experimental/Chooser/Source/ChooserEditor/Private/OutputObjectColumnEditor.cpp
@@ -0,0 +1,186 @@
+// Copyright Epic Games, Inc. All Rights Reserved.
+
+#include "OutputObjectColumnEditor.h"
+#include "SPropertyAccessChainWidget.h"
+#include "GraphEditorSettings.h"
+#include "ObjectChooserWidgetFactories.h"
+#include "OutputObjectColumn.h"
+#include "PropertyCustomizationHelpers.h"
+#include "TransactionCommon.h"
+#include "Widgets/Images/SImage.h"
+#include "Widgets/Input/SButton.h"
+#include "Widgets/Layout/SBox.h"
+#include "ScopedTransaction.h"
+
+#define LOCTEXT_NAMESPACE "OutputObjectColumnEditor"
+
+namespace UE::ChooserEditor
+{
+	static UClass* GetAllowedClass(const FOutputObjectColumn* OutputObjectColumn)
+	{
+		UClass* AllowedClass = nullptr;
+		if (const FChooserParameterObjectBase* InputValue = OutputObjectColumn->InputValue.GetPtr<FChooserParameterObjectBase>())
+		{
+			AllowedClass = InputValue->GetAllowedClass();
+		}
+
+		if (AllowedClass == nullptr)
+		{
+			AllowedClass = UObject::StaticClass();
+		}
+
+		return AllowedClass;
+	}
+
+
+	
+	// Wrapper widget for EnumComboBox which will reconstruct the combo box when the Enum has changed
+    class SOutputObjectCell : public SCompoundWidget
+    {
+    public:
+    
+    	DECLARE_DELEGATE_OneParam(FOnValueSet, int);
+    	
+    	SLATE_BEGIN_ARGS(SOutputObjectCell)
+    	{}
+    
+    	SLATE_ARGUMENT(UChooserTable*, Chooser)
+    	SLATE_ARGUMENT(FOutputObjectColumn*, Column)
+    	SLATE_ATTRIBUTE(int32, Row)
+                
+    	SLATE_END_ARGS()
+    
+    	TSharedRef<SWidget> CreateWidget()
+    	{
+    		if (Column)
+    		{
+				int Row = RowIndex.Get();
+				UChooserTable* ContextOwner = Chooser->GetContextOwner();
+
+				if (Row == ColumnWidget_SpecialIndex_Fallback)
+				{
+					UClass* AllowedClass = UObject::StaticClass();
+					if (Column->InputValue.IsValid())
+					{
+						AllowedClass = Column->InputValue.Get<FChooserParameterObjectBase>().GetAllowedClass();
+					}
+					
+					TSharedPtr<SWidget> ResultWidget = FObjectChooserWidgetFactories::CreateWidget(false, ContextOwner, FObjectChooserBase::StaticStruct(),
+						&Column->FallbackValue.Value, AllowedClass, FChooserWidgetValueChanged(), LOCTEXT("None", "(None)"));
+					return ResultWidget. ToSharedRef();
+				}
+				else if (Row == ColumnWidget_SpecialIndex_Header)
+				{
+					// create column header widget
+					TSharedPtr<SWidget> InputValueWidget = nullptr;
+					if (FChooserParameterBase* InputValue = Column->GetInputValue())
+					{
+						InputValueWidget = FObjectChooserWidgetFactories::CreateWidget(false, Chooser, InputValue, Column->GetInputType(), ContextOwner->OutputObjectType);
+					}
+					
+					const FSlateBrush* ColumnIcon = FCoreStyle::Get().GetBrush("Icons.ArrowRight");
+					
+					TSharedRef<SWidget> ColumnHeaderWidget = SNew(SHorizontalBox)
+						+ SHorizontalBox::Slot().AutoWidth()
+						[
+							SNew(SBorder)
+							.BorderBackgroundColor(FLinearColor(0,0,0,0))
+							.Content()
+							[
+								SNew(SImage).Image(ColumnIcon)
+							]
+						]
+						+ SHorizontalBox::Slot()
+						[
+							InputValueWidget ? InputValueWidget.ToSharedRef() : SNullWidget::NullWidget
+						];
+				
+					return ColumnHeaderWidget;
+				}
+ 
+				if (Column->RowValues.IsValidIndex(Row))
+				{
+					UClass* AllowedClass = UObject::StaticClass();
+					if (Column->InputValue.IsValid())
+					{
+						AllowedClass = Column->InputValue.Get<FChooserParameterObjectBase>().GetAllowedClass();
+					}
+					
+					TSharedPtr<SWidget> ResultWidget = FObjectChooserWidgetFactories::CreateWidget(false, ContextOwner, FObjectChooserBase::StaticStruct(),
+						&Column->RowValues[Row].Value, AllowedClass, FChooserWidgetValueChanged(), LOCTEXT("None", "(None)"));
+					return ResultWidget. ToSharedRef();
+				}
+ 
+				return SNullWidget::NullWidget;
+    				
+    		}
+    		
+    		return SNullWidget::NullWidget;
+    	}
+    
+    	void UpdateWidget()
+    	{
+    		ChildSlot[ CreateWidget() ];
+    	}
+    
+    	virtual void Tick(const FGeometry& AllottedGeometry, const double InCurrentTime, const float InDeltaTime) override
+    	{
+    		UClass* CurrentBaseType = nullptr;
+    		if (Column->InputValue.IsValid())
+    		{
+    			CurrentBaseType = Column->InputValue.Get<FChooserParameterObjectBase>().GetAllowedClass();
+    		}
+    		if (ObjectBaseType != CurrentBaseType)
+    		{
+    			UpdateWidget();
+    			ObjectBaseType = CurrentBaseType;
+    		}
+    	}
+        					
+    
+    	void Construct( const FArguments& InArgs)
+    	{
+    		SetEnabled(InArgs._IsEnabled);
+    		
+    		SetCanTick(true);
+    		Column = InArgs._Column;
+    		Chooser = InArgs._Chooser;
+    		RowIndex = InArgs._Row;
+    
+			if (Column)
+    		{
+    			if (Column->InputValue.IsValid())
+    			{
+    				ObjectBaseType = Column->InputValue.Get<FChooserParameterObjectBase>().GetAllowedClass();
+    			}
+    		}
+    
+    		UpdateWidget();
+    	}
+    
+    	virtual ~SOutputObjectCell()
+    	{
+    	}
+    
+    private:
+    	UChooserTable* Chooser = nullptr;
+    	FOutputObjectColumn* Column = nullptr;
+		UClass* ObjectBaseType = nullptr;
+    	TAttribute<int> RowIndex;
+    	FDelegateHandle EnumChangedHandle;
+    };
+
+
+	static TSharedRef<SWidget> CreateOutputObjectColumnWidget(UChooserTable* Chooser, FChooserColumnBase* Column, int32 Row)
+	{
+		return SNew(SOutputObjectCell).Chooser(Chooser).Column(static_cast<FOutputObjectColumn*>(Column)).Row(Row);
+	}
+
+	void RegisterOutputObjectWidgets()
+	{
+		FObjectChooserWidgetFactories::RegisterColumnWidgetCreator(FOutputObjectColumn::StaticStruct(), CreateOutputObjectColumnWidget);
+	}
+
+} // namespace UE::ChooserEditor
+
+#undef LOCTEXT_NAMESPACE
diff --git a/Engine/Plugins/Experimental/Chooser/Source/ChooserEditor/Private/OutputObjectColumnEditor.h b/Engine/Plugins/Experimental/Chooser/Source/ChooserEditor/Private/OutputObjectColumnEditor.h
new file mode 100644
--- /dev/null
+++ b/Engine/Plugins/Experimental/Chooser/Source/ChooserEditor/Private/OutputObjectColumnEditor.h
@@ -0,0 +1,10 @@
+// Copyright Epic Games, Inc. All Rights Reserved.
+
+#pragma once
+
+#include "CoreMinimal.h"
+
+namespace UE::ChooserEditor
+{
+	void RegisterOutputObjectWidgets();
+}
diff --git a/Engine/Plugins/Experimental/Chooser/Source/ChooserEditor/Public/ObjectChooserWidgetFactories.h b/Engine/Plugins/Experimental/Chooser/Source/ChooserEditor/Public/ObjectChooserWidgetFactories.h
--- a/Engine/Plugins/Experimental/Chooser/Source/ChooserEditor/Public/ObjectChooserWidgetFactories.h
+++ b/Engine/Plugins/Experimental/Chooser/Source/ChooserEditor/Public/ObjectChooserWidgetFactories.h
@@ -24,6 +24,9 @@ namespace UE::ChooserEditor
 	{
 	public:
 
+		static TSharedPtr<SWidget> CreateWidget(bool bReadOnly, UObject* TransactionObject, UScriptStruct* DataBaseType, FInstancedStruct* Data, UClass* ResultBaseClass,
+												FChooserWidgetValueChanged ValueChanged = FChooserWidgetValueChanged(), FText NullValueDisplayText = FText());
+		
 		static TSharedPtr<SWidget> CreateWidget(bool ReadOnly, UObject* TransactionObject, void* Value, const UStruct* ValueType, UClass* ResultBaseClass, FChooserWidgetValueChanged ValueChanged = FChooserWidgetValueChanged());
 		static TSharedPtr<SWidget> CreateWidget(bool ReadOnly, UObject* TransactionObject, const UScriptStruct* BaseType, void* Value, const UStruct* ValueType, UClass* ResultBaseClass,
 										const FOnStructPicked& CreateClassCallback, TSharedPtr<SBorder>* InnerWidget = nullptr, FChooserWidgetValueChanged ValueChanged = FChooserWidgetValueChanged(),
