Add type checking to chooser struct output columns, and proxy table output structs

diff --git a/Engine/Plugins/Experimental/Chooser/Source/Chooser/Private/ChooserPropertyAccess.cpp b/Engine/Plugins/Experimental/Chooser/Source/Chooser/Private/ChooserPropertyAccess.cpp
--- a/Engine/Plugins/Experimental/Chooser/Source/Chooser/Private/ChooserPropertyAccess.cpp
+++ b/Engine/Plugins/Experimental/Chooser/Source/Chooser/Private/ChooserPropertyAccess.cpp
@@ -207,6 +207,11 @@ void FChooserPropertyBinding::Compile(IHasContextClass* Owner, bool bForce)
 	{
 		// handle binding directly to a context struct
 		OutCompiledBinding.CompiledChain.Add(UE::Chooser::FCompiledBindingElement(0));
+		
+		if (const FContextObjectTypeStruct* StructContext = ContextData[ContextIndex].GetPtr<FContextObjectTypeStruct>())
+		{
+			OutCompiledBinding.StructType = StructContext->Struct;
+		}
 		bFound = true;
 	}
 	else
@@ -245,6 +250,10 @@ void FChooserPropertyBinding::Compile(IHasContextClass* Owner, bool bForce)
 			{
 				OutCompiledBinding.PropertyType = UE::Chooser::EChooserPropertyAccessType::SoftObjectRef;
 			}
+			else if (BaseProperty->IsA<FStructProperty>())
+            {
+            	OutCompiledBinding.StructType = CastField<FStructProperty>(BaseProperty)->Struct;
+            }
 		}
 		else
 		{
@@ -377,6 +386,7 @@ namespace UE::Chooser
 			Result.Mask = Last.Mask;
 		}
 		Result.PropertyType = CompiledBinding.PropertyType;
+		Result.StructType = CompiledBinding.StructType; 
 		return true;
 	}
 	
@@ -387,6 +397,7 @@ namespace UE::Chooser
 			if (PropertyBinding.IsBoundToRoot)
 			{
 				Result.Container = Container;
+				Result.StructType = StructType;
 				return true;
 			}
 			else
@@ -482,6 +493,10 @@ namespace UE::Chooser
 			{
 				Result.PropertyType = UE::Chooser::EChooserPropertyAccessType::SoftObjectRef;
 			}
+			else if (BaseProperty->IsA<FStructProperty>())
+			{
+				Result.StructType = CastField<FStructProperty>(BaseProperty)->Struct;
+			}
 		}
 		else
 		{
diff --git a/Engine/Plugins/Experimental/Chooser/Source/Chooser/Private/OutputStructColumn.cpp b/Engine/Plugins/Experimental/Chooser/Source/Chooser/Private/OutputStructColumn.cpp
--- a/Engine/Plugins/Experimental/Chooser/Source/Chooser/Private/OutputStructColumn.cpp
+++ b/Engine/Plugins/Experimental/Chooser/Source/Chooser/Private/OutputStructColumn.cpp
@@ -11,10 +11,18 @@
 bool FStructContextProperty::SetValue(FChooserEvaluationContext& Context, const FInstancedStruct& InValue) const
 {
 	void* TargetData;
-	if (Binding.GetValuePtr(Context, TargetData))
+	const UStruct* StructType;
+	if (Binding.GetStructPtr(Context, TargetData, StructType))
 	{
-		InValue.GetScriptStruct()->CopyScriptStruct(TargetData, InValue.GetMemory());
-		return true;
+		if (StructType == InValue.GetScriptStruct())
+		{
+			InValue.GetScriptStruct()->CopyScriptStruct(TargetData, InValue.GetMemory());
+			return true;
+		}
+		else
+		{
+			return false;
+		}
 	}
 	return false;
 }
diff --git a/Engine/Plugins/Experimental/Chooser/Source/Chooser/Public/ChooserPropertyAccess.h b/Engine/Plugins/Experimental/Chooser/Source/Chooser/Public/ChooserPropertyAccess.h
--- a/Engine/Plugins/Experimental/Chooser/Source/Chooser/Public/ChooserPropertyAccess.h
+++ b/Engine/Plugins/Experimental/Chooser/Source/Chooser/Public/ChooserPropertyAccess.h
@@ -54,7 +54,7 @@ namespace UE::Chooser
 		Int32,
 		Float,
 		Double,
-		SoftObjectRef
+		SoftObjectRef,
 	};
 
 	struct FCompiledBinding
@@ -64,6 +64,9 @@ namespace UE::Chooser
 		// type info for number and enum conversions
 		TArray<FCompiledBindingElement> CompiledChain;
 		const UStruct* TargetType = nullptr;
+
+		// Type of struct for when the property itself is a struct
+		const UStruct* StructType = nullptr;
 #if WITH_EDITORONLY_DATA
 		int SerialNumber = 0;
 		TArray<const UStruct*> Dependencies;
@@ -103,6 +106,9 @@ struct CHOOSER_API FChooserPropertyBinding
 	template <typename T>
 	bool GetValuePtr(FChooserEvaluationContext& Context, T*& Value) const;
 	
+	template <typename T>
+	bool GetStructPtr(FChooserEvaluationContext& Context, T*& Value, UStruct const* &StructType) const;
+	
 	template <typename T>
 	bool GetValue(FChooserEvaluationContext& Context, T& Value) const;
 	
@@ -185,6 +191,7 @@ namespace UE::Chooser
 		uint8* Container = nullptr;
 		uint32 PropertyOffset = 0;
 		UFunction* Function = nullptr;
+		const UStruct* StructType = nullptr;
 		EChooserPropertyAccessType PropertyType = EChooserPropertyAccessType::None;
 		uint8 Mask = 0;
 	};
@@ -213,6 +220,25 @@ bool FChooserPropertyBinding::GetValuePtr(FChooserEvaluationContext& Context, T*
 	return false;
 }
 
+template <typename T>
+bool FChooserPropertyBinding::GetStructPtr(FChooserEvaluationContext& Context, T*& OutResult, UStruct const* &OutStructType) const
+{
+	using namespace  UE::Chooser;
+
+	FResolvedPropertyChainResult Result;
+	if (ResolvePropertyChain(Context, *this, Result))
+	{
+		if (Result.Function == nullptr)
+		{
+			OutResult = reinterpret_cast<T*>(Result.Container + Result.PropertyOffset);
+			OutStructType = Result.StructType;
+			return true;
+		}
+	}
+	return false;
+}
+
+
 template <typename T>
 bool FChooserPropertyBinding::GetValue(FChooserEvaluationContext& Context, T& OutResult) const
 {
diff --git a/Engine/Plugins/Experimental/Chooser/Source/ProxyTable/Private/ProxyTable.cpp b/Engine/Plugins/Experimental/Chooser/Source/ProxyTable/Private/ProxyTable.cpp
--- a/Engine/Plugins/Experimental/Chooser/Source/ProxyTable/Private/ProxyTable.cpp
+++ b/Engine/Plugins/Experimental/Chooser/Source/ProxyTable/Private/ProxyTable.cpp
@@ -197,10 +197,14 @@ static void OutputStructData(const FRuntimeProxyValue& EntryValueData, FChooserE
 	for (const FProxyStructOutput& StructOutput : EntryValueData.OutputStructData)
 	{
 		// copy each struct output value
-		void* TargetData;
-		if (StructOutput.Binding.GetValuePtr(Context, TargetData))
+		void* TargetData = nullptr;
+		const UStruct* StructType = nullptr;
+		if (StructOutput.Binding.GetStructPtr(Context, TargetData, StructType))
 		{
-			StructOutput.Value.GetScriptStruct()->CopyScriptStruct(TargetData, StructOutput.Value.GetMemory());
+			if (StructType == StructOutput.Value.GetScriptStruct())
+			{
+				StructOutput.Value.GetScriptStruct()->CopyScriptStruct(TargetData, StructOutput.Value.GetMemory());
+			}
 		}
 	}
 }
