UE-211195 - Editor crashes on start up when attempting to run  -asyncloadingthread

This is a similar fix as CL 31886022.

With async editor loading, PostLoad can only be done by the loader during its postload pass (unless we're calling PostLoad on another object during a PostLoad). Additionally, Blueprint compilation only guarantees that any dependencies (eg: graphs, nodes, etc.) are preloaded, but not postloaded.

In this case, UK2Node_EvaluateChooser2::AllocateDefaultPins uses some fields of the UChooserTable object, but they only need to have their deprecated data fixed first. Data fixup is normally done in Serialize, so we can just move that logic from PostLoad into there, and then remove the direct PostLoad call.

#rnx

diff --git a/Engine/Plugins/Experimental/Chooser/Source/Chooser/Private/Chooser.cpp b/Engine/Plugins/Experimental/Chooser/Source/Chooser/Private/Chooser.cpp
--- a/Engine/Plugins/Experimental/Chooser/Source/Chooser/Private/Chooser.cpp
+++ b/Engine/Plugins/Experimental/Chooser/Source/Chooser/Private/Chooser.cpp
@@ -39,50 +39,6 @@ void UChooserTable::PostLoad()
 	CachedPreviousResultType = ResultType;
 #endif
 
-#if WITH_EDITORONLY_DATA
-	// convert old data if it exists
-
-	if (ContextObjectType_DEPRECATED)
-	{
-		ContextData.SetNum(1);
-		ContextData[0].InitializeAs<FContextObjectTypeClass>();
-		FContextObjectTypeClass& Context = ContextData[0].GetMutable<FContextObjectTypeClass>();
-		Context.Class = ContextObjectType_DEPRECATED;
-		Context.Direction = EContextObjectDirection::ReadWrite;
-		ContextObjectType_DEPRECATED = nullptr;
-		OnContextClassChanged.Broadcast();
-	}
-
-	if (Results_DEPRECATED.Num() > 0 || Columns_DEPRECATED.Num() > 0)
-	{
-		ResultsStructs.Reserve(Results_DEPRECATED.Num());
-		ColumnsStructs.Reserve(Columns_DEPRECATED.Num());
-		
-		for(TScriptInterface<IObjectChooser>& Result : Results_DEPRECATED)
-		{
-			ResultsStructs.SetNum(ResultsStructs.Num() + 1);
-			IObjectChooser* ResultInterface = Result.GetInterface();
-			if (ResultInterface)
-			{
-				ResultInterface->ConvertToInstancedStruct(ResultsStructs.Last());
-			}
-		}
-		
-		for(TScriptInterface<IChooserColumn>& Column : Columns_DEPRECATED)
-		{
-			ColumnsStructs.SetNum(ColumnsStructs.Num() + 1);
-			IChooserColumn* ColumnInterface = Column.GetInterface();
-			if (ColumnInterface)
-			{
-				ColumnInterface->ConvertToInstancedStruct(ColumnsStructs.Last());
-			}
-		}
-
-		Results_DEPRECATED.SetNum(0);
-		Columns_DEPRECATED.SetNum(0);
-	}
-#endif
-
 	Compile();
 }
 
@@ -196,7 +152,53 @@ void UChooserTable::Serialize(FArchive& Ar)
 	}
 #endif
 	
-	UObject::Serialize(Ar);
+	Super::Serialize(Ar);
+
+#if WITH_EDITORONLY_DATA
+	if (Ar.IsLoading())
+	{
+		// convert old data if it exists
+
+		if (ContextObjectType_DEPRECATED)
+		{
+			ContextData.SetNum(1);
+			ContextData[0].InitializeAs<FContextObjectTypeClass>();
+			FContextObjectTypeClass& Context = ContextData[0].GetMutable<FContextObjectTypeClass>();
+			Context.Class = ContextObjectType_DEPRECATED;
+			Context.Direction = EContextObjectDirection::ReadWrite;
+			ContextObjectType_DEPRECATED = nullptr;
+		}
+
+		if (Results_DEPRECATED.Num() > 0 || Columns_DEPRECATED.Num() > 0)
+		{
+			ResultsStructs.Reserve(Results_DEPRECATED.Num());
+			ColumnsStructs.Reserve(Columns_DEPRECATED.Num());
+
+			for (TScriptInterface<IObjectChooser>& Result : Results_DEPRECATED)
+			{
+				ResultsStructs.SetNum(ResultsStructs.Num() + 1);
+				IObjectChooser* ResultInterface = Result.GetInterface();
+				if (ResultInterface)
+				{
+					ResultInterface->ConvertToInstancedStruct(ResultsStructs.Last());
+				}
+			}
+
+			for (TScriptInterface<IChooserColumn>& Column : Columns_DEPRECATED)
+			{
+				ColumnsStructs.SetNum(ColumnsStructs.Num() + 1);
+				IChooserColumn* ColumnInterface = Column.GetInterface();
+				if (ColumnInterface)
+				{
+					ColumnInterface->ConvertToInstancedStruct(ColumnsStructs.Last());
+				}
+			}
+
+			Results_DEPRECATED.SetNum(0);
+			Columns_DEPRECATED.SetNum(0);
+		}
+	}
+#endif
 }
 
 #if WITH_EDITORONLY_DATA
diff --git a/Engine/Plugins/Experimental/Chooser/Source/ChooserUncooked/Private/EvaluateChooserNode.cpp b/Engine/Plugins/Experimental/Chooser/Source/ChooserUncooked/Private/EvaluateChooserNode.cpp
--- a/Engine/Plugins/Experimental/Chooser/Source/ChooserUncooked/Private/EvaluateChooserNode.cpp
+++ b/Engine/Plugins/Experimental/Chooser/Source/ChooserUncooked/Private/EvaluateChooserNode.cpp
@@ -405,10 +405,7 @@ void UK2Node_EvaluateChooser2::AllocateDefaultPins()
 
 	if (Chooser)
 	{
-		// ensure any data upgrades have been applied to Chooser before generating pins
-		Chooser->ConditionalPostLoad();
-		
-		for(FInstancedStruct& ContextDataEntry : Chooser->ContextData)
+		for (FInstancedStruct& ContextDataEntry : Chooser->ContextData)
 		{
 			if (ContextDataEntry.IsValid())
 			{
@@ -547,7 +544,6 @@ void UK2Node_EvaluateChooser2::PostLoad()
 
 void UK2Node_EvaluateChooser2::PinConnectionListChanged(UEdGraphPin* Pin)
 {
-
 	Modify();
 
 	FBlueprintEditorUtils::MarkBlueprintAsStructurallyModified(GetBlueprint());
