Scoring/Cost columns for chooser tables

diff --git a/Engine/Plugins/Experimental/Chooser/Source/Chooser/Internal/FloatDistanceColumn.h b/Engine/Plugins/Experimental/Chooser/Source/Chooser/Internal/FloatDistanceColumn.h
new file mode 100644
--- /dev/null
+++ b/Engine/Plugins/Experimental/Chooser/Source/Chooser/Internal/FloatDistanceColumn.h
@@ -0,0 +1,97 @@
+// Copyright Epic Games, Inc. All Rights Reserved.
+#pragma once
+
+#include <IHasContext.h>
+
+#include "CoreMinimal.h"
+#include "IChooserColumn.h"
+#include "IChooserParameterFloat.h"
+#include "ChooserPropertyAccess.h"
+#include "InstancedStruct.h"
+#include "Serialization/MemoryReader.h"
+#include "FloatDistanceColumn.generated.h"
+
+USTRUCT()
+struct FChooserFloatDistanceRowData
+{
+	GENERATED_BODY()
+
+	UPROPERTY(EditAnywhere, Category=Runtime)
+	float Value=0;
+};
+
+
+USTRUCT()
+struct CHOOSER_API FFloatDistanceColumn : public FChooserColumnBase
+{
+	GENERATED_BODY()
+	public:
+	FFloatDistanceColumn();
+		
+	UPROPERTY(EditAnywhere, NoClear, Meta = (ExcludeBaseStruct, BaseStruct = "/Script/Chooser.ChooserParameterFloatBase"), Category = "Data")
+	FInstancedStruct InputValue;
+
+	// Maximum Distance used for normalizing scoring (greater distances will be considered equal to the max)
+	UPROPERTY(EditAnywhere, Category = "Data");
+	double MaxDistance= 100;
+
+	// Multiplier for controlling which scoring column has the most influence.  Higher values will make the match from this column more important.
+	UPROPERTY(EditAnywhere, Category = "Data");
+	float CostMultiplier = 1;
+
+	// For rows with distance greater than MaxDistance, filter out the row
+	UPROPERTY(EditAnywhere, Category = "Data");
+	bool bFilterOverMaxDistance = false;
+	// Wrap input, and distance calculations for numbers such as angles
+	UPROPERTY(EditAnywhere, Category = "Data");
+	bool bWrapInput = false;
+	
+	// Minimum value (for WrapInput)
+	UPROPERTY(EditAnywhere, Category = "Data", meta=(DisplayAfter="bWrapInput", EditCondition="bWrapInput"));
+	double MinValue = -180;
+	
+	// Maximum value (for WrapInput)
+	UPROPERTY(EditAnywhere, Category = "Data", meta=(DisplayAfter="bWrapInput", EditCondition="bWrapInput"));
+	double MaxValue = 180;
+
+#if WITH_EDITORONLY_DATA
+	UPROPERTY(EditAnywhere, Category="Data")
+	FChooserFloatDistanceRowData DefaultRowValue;
+#endif
+	
+	UPROPERTY(EditAnywhere, Category="Data")
+	// array of results (cells for this column for each row in the table)
+	// should match the length of the Results array 
+	TArray<FChooserFloatDistanceRowData> RowValues;
+
+	virtual bool HasFilters() const { return true; }
+	virtual void Filter(FChooserEvaluationContext& Context, const FChooserIndexArray& IndexListIn, FChooserIndexArray& IndexListOut) const override;
+
+	virtual bool HasCosts() const override { return true; }
+
+#if WITH_EDITOR
+	mutable double TestValue = 0.0;
+	virtual bool EditorTestFilter(int32 RowIndex) const override;
+	virtual float EditorTestCost(int32 RowIndex) const override;
+	
+	virtual void SetTestValue(TArrayView<const uint8> Value) override
+	{
+		FMemoryReaderView Reader(Value);
+		Reader << TestValue;
+	}
+
+	virtual void AddToDetails(FInstancedPropertyBag& PropertyBag, int32 ColumnIndex, int32 RowIndex) override;
+	virtual void SetFromDetails(FInstancedPropertyBag& PropertyBag, int32 ColumnIndex, int32 RowIndex) override;
+#endif
+	
+	virtual void PostLoad() override
+	{
+		if (InputValue.IsValid())
+		{
+			InputValue.GetMutable<FChooserParameterBase>().PostLoad();
+		}
+	}
+
+	CHOOSER_COLUMN_BOILERPLATE(FChooserParameterFloatBase);
+
+};
\ No newline at end of file
diff --git a/Engine/Plugins/Experimental/Chooser/Source/Chooser/Internal/RandomizeColumn.h b/Engine/Plugins/Experimental/Chooser/Source/Chooser/Internal/RandomizeColumn.h
--- a/Engine/Plugins/Experimental/Chooser/Source/Chooser/Internal/RandomizeColumn.h
+++ b/Engine/Plugins/Experimental/Chooser/Source/Chooser/Internal/RandomizeColumn.h
@@ -36,6 +36,9 @@ struct CHOOSER_API FRandomizeColumn : public FChooserColumnBase
 	UPROPERTY(EditAnywhere, Category= "Data", meta=(ClampMin="0.0",Tooltip="Multiplies the weight of the previous chosen result (set to 0 to never pick the same result twice in a row)"));
 	float RepeatProbabilityMultiplier = 1.0f;
 
+	UPROPERTY(EditAnywhere, Category= "Data", meta=(ClampMin="0.0",Tooltip="When columns with scoring are used, randomize will pick from among all rows that have a cost nearly equal to the minumum cost, using this threshold"));
+	float EqualCostThreshold = 0.001f;
+
 #if WITH_EDITORONLY_DATA
 	UPROPERTY(EditAnywhere, Category= "Data", DisplayName="DefaultRowValue");
 	float DefaultRowValue = 1.0f;
diff --git a/Engine/Plugins/Experimental/Chooser/Source/Chooser/Private/BoolColumn.cpp b/Engine/Plugins/Experimental/Chooser/Source/Chooser/Private/BoolColumn.cpp
--- a/Engine/Plugins/Experimental/Chooser/Source/Chooser/Private/BoolColumn.cpp
+++ b/Engine/Plugins/Experimental/Chooser/Source/Chooser/Private/BoolColumn.cpp
@@ -40,14 +40,14 @@ void FBoolColumn::Filter(FChooserEvaluationContext& Context, const FChooserIndex
 		}
 	#endif
 		
-		for (uint32 Index : IndexListIn)
+		for (const FChooserIndexArray::FIndexData& IndexData : IndexListIn)
 		{
-			if (RowValuesWithAny.Num() > (int)Index)
+			if (RowValuesWithAny.IsValidIndex(IndexData.Index))
 			{
 				
-				if (RowValuesWithAny[Index] == EBoolColumnCellValue::MatchAny || Result == static_cast<bool>(RowValuesWithAny[Index]))
+				if (RowValuesWithAny[IndexData.Index] == EBoolColumnCellValue::MatchAny || Result == static_cast<bool>(RowValuesWithAny[IndexData.Index]))
 				{
-					IndexListOut.Push(Index);
+					IndexListOut.Push(IndexData);
 				}
 			}
 		}
diff --git a/Engine/Plugins/Experimental/Chooser/Source/Chooser/Private/Chooser.cpp b/Engine/Plugins/Experimental/Chooser/Source/Chooser/Private/Chooser.cpp
--- a/Engine/Plugins/Experimental/Chooser/Source/Chooser/Private/Chooser.cpp
+++ b/Engine/Plugins/Experimental/Chooser/Source/Chooser/Private/Chooser.cpp
@@ -362,16 +362,16 @@ FObjectChooserBase::EIteratorStatus UChooserTable::EvaluateChooser(FChooserEvalu
 	
 
 	uint32 Count = ResultsArray->Num();
-	uint32 BufferSize = Count * sizeof(uint32);
+	uint32 BufferSize = Count * sizeof(FChooserIndexArray::FIndexData);
 
-	FChooserIndexArray Indices1(static_cast<uint32*>(FMemory_Alloca(BufferSize)), Count);
-	FChooserIndexArray Indices2(static_cast<uint32*>(FMemory_Alloca(BufferSize)), Count);
+	FChooserIndexArray Indices1(static_cast<FChooserIndexArray::FIndexData*>(FMemory_Alloca(BufferSize)), Count);
+	FChooserIndexArray Indices2(static_cast<FChooserIndexArray::FIndexData*>(FMemory_Alloca(BufferSize)), Count);
 
 	for(uint32 i=0;i<Count;i++)
 	{
 		if (!Chooser->IsRowDisabled(i))
 		{
-			Indices1.Push(i);
+			Indices1.Push({i, 0});
 		}
 	}
 	FChooserIndexArray* IndicesOut = &Indices1;
@@ -393,8 +393,19 @@ FObjectChooserBase::EIteratorStatus UChooserTable::EvaluateChooser(FChooserEvalu
 			Swap(IndicesIn, IndicesOut);
 			IndicesOut->SetNum(0);
 			Column.Filter(Context, *IndicesIn, *IndicesOut);
+			
+			if (IndicesIn->HasCosts() || Column.HasCosts())
+			{
+				IndicesOut->SetHasCosts();
+			}
 		}
 	}
+
+	//	No need to score with only one valid option
+	if (IndicesOut->Num() > 1 && IndicesOut->HasCosts())
+	{
+		Algo::Sort(*IndicesOut);
+	}
 	
 	bool bSetOutputs = false;
 	if (IndicesOut->IsEmpty())
@@ -432,11 +443,11 @@ FObjectChooserBase::EIteratorStatus UChooserTable::EvaluateChooser(FChooserEvalu
 	else
 	{
 		// of the rows that passed all column filters, iterate through them calling the callback until it returns Stop
-		for (uint32 SelectedIndex : *IndicesOut)
+		for (FChooserIndexArray::FIndexData& SelectedIndexData : *IndicesOut)
 		{
-			if (ResultsArray->Num() > (int32)SelectedIndex)
+			if (ResultsArray->IsValidIndex(SelectedIndexData.Index))
 			{
-				const FObjectChooserBase& SelectedResult = (*ResultsArray)[SelectedIndex].Get<FObjectChooserBase>();
+				const FObjectChooserBase& SelectedResult = (*ResultsArray)[SelectedIndexData.Index].Get<FObjectChooserBase>();
 				FObjectChooserBase::EIteratorStatus Status = SelectedResult.ChooseMulti(Context, Callback);
 				if (Status != FObjectChooserBase::EIteratorStatus::Continue)
 				{
@@ -445,15 +456,15 @@ FObjectChooserBase::EIteratorStatus UChooserTable::EvaluateChooser(FChooserEvalu
 					for (const FInstancedStruct& ColumnData : Chooser->ColumnsStructs)
 					{
 						const FChooserColumnBase& Column = ColumnData.Get<FChooserColumnBase>();
-						Column.SetOutputs(Context, SelectedIndex);
+						Column.SetOutputs(Context, SelectedIndexData.Index);
 					}
 					#if WITH_EDITOR
 					if (Context.DebuggingInfo.bCurrentDebugTarget)
 					{
-						Chooser->SetDebugSelectedRow(SelectedIndex);
+						Chooser->SetDebugSelectedRow(SelectedIndexData.Index);
 					}
 					#endif
-					TRACE_CHOOSER_EVALUATION(Chooser, Context, SelectedIndex);
+					TRACE_CHOOSER_EVALUATION(Chooser, Context, SelectedIndexData.Index);
 				}
 				if (Status == FObjectChooserBase::EIteratorStatus::Stop)
 				{
diff --git a/Engine/Plugins/Experimental/Chooser/Source/Chooser/Private/EnumColumn.cpp b/Engine/Plugins/Experimental/Chooser/Source/Chooser/Private/EnumColumn.cpp
--- a/Engine/Plugins/Experimental/Chooser/Source/Chooser/Private/EnumColumn.cpp
+++ b/Engine/Plugins/Experimental/Chooser/Source/Chooser/Private/EnumColumn.cpp
@@ -79,14 +79,14 @@ void FEnumColumn::Filter(FChooserEvaluationContext& Context, const FChooserIndex
 
 		TRACE_CHOOSER_VALUE(Context, ToCStr(InputValue.Get<FChooserParameterBase>().GetDebugName()), Result);
 		
-		for (const uint32 Index : IndexListIn)
+		for (const FChooserIndexArray::FIndexData& IndexData : IndexListIn)
 		{
-			if (RowValues.IsValidIndex(Index))
+			if (RowValues.IsValidIndex(IndexData.Index))
 			{
-				const FChooserEnumRowData& RowValue = RowValues[Index];
+				const FChooserEnumRowData& RowValue = RowValues[IndexData.Index];
 				if (RowValue.Evaluate(Result))
 				{
-					IndexListOut.Push(Index);
+					IndexListOut.Push(IndexData);
 				}
 			}
 		}
diff --git a/Engine/Plugins/Experimental/Chooser/Source/Chooser/Private/FloatDistanceColumn.cpp b/Engine/Plugins/Experimental/Chooser/Source/Chooser/Private/FloatDistanceColumn.cpp
new file mode 100644
--- /dev/null
+++ b/Engine/Plugins/Experimental/Chooser/Source/Chooser/Private/FloatDistanceColumn.cpp
@@ -0,0 +1,162 @@
+// Copyright Epic Games, Inc. All Rights Reserved.
+#include "FloatDistanceColumn.h"
+#include "FloatRangeColumn.h"
+#include "ChooserIndexArray.h"
+#include "ChooserPropertyAccess.h"
+#include "ChooserTrace.h"
+#if WITH_EDITOR
+#include "PropertyBag.h"
+#endif
+
+FFloatDistanceColumn::FFloatDistanceColumn()
+{
+	InputValue.InitializeAs(FFloatContextProperty::StaticStruct());
+}
+
+#if WITH_EDITOR
+bool FFloatDistanceColumn::EditorTestFilter(int32 RowIndex) const
+{
+	if (!bFilterOverMaxDistance)
+	{
+		// filtering not enabled, only scoring
+		return true;
+	}
+	
+	if (RowValues.IsValidIndex(RowIndex))
+	{
+		const FChooserFloatDistanceRowData& RowValue = RowValues[RowIndex];
+		
+		if (bWrapInput)
+		{
+			double WrappedValue = FMath::Wrap(TestValue, MinValue, MaxValue);
+			double Distance =  fabs(WrappedValue - RowValue.Value);
+			if (Distance > 0.5 * MaxValue - MinValue)
+			{
+				Distance = MaxValue - MinValue - Distance;
+			}
+			return Distance <= MaxDistance;
+		}
+		else
+		{
+			return fabs(TestValue - RowValue.Value) <= MaxDistance;
+		}
+	}
+	
+	return false;
+}
+
+
+float FFloatDistanceColumn::EditorTestCost(int32 RowIndex) const
+{
+	if (RowValues.IsValidIndex(RowIndex))
+	{
+		const FChooserFloatDistanceRowData& RowValue = RowValues[RowIndex];
+		
+		if (bWrapInput)
+		{
+			double WrappedValue = FMath::Wrap(TestValue, MinValue, MaxValue);
+			double Distance =  fabs(WrappedValue - RowValue.Value);
+			if (Distance > 0.5 * MaxValue - MinValue)
+			{
+				Distance = MaxValue - MinValue - Distance;
+			}
+			return FMath::Min(Distance / MaxDistance, 1.f);
+		}
+		else
+		{
+			return FMath::Min(fabs(TestValue - RowValue.Value) / MaxDistance, 1.f);
+		}
+	}
+	
+	return 0.0f;
+}
+
+#endif
+
+void FFloatDistanceColumn::Filter(FChooserEvaluationContext& Context, const FChooserIndexArray& IndexListIn, FChooserIndexArray& IndexListOut) const
+{
+	if (InputValue.IsValid())
+	{
+		double Result = 0.0f;
+		InputValue.Get<FChooserParameterFloatBase>().GetValue(Context, Result);
+        	
+		TRACE_CHOOSER_VALUE(Context, ToCStr(InputValue.Get<FChooserParameterBase>().GetDebugName()), Result);
+
+#if WITH_EDITOR
+		if (Context.DebuggingInfo.bCurrentDebugTarget)
+		{
+			TestValue = Result;
+		}
+#endif
+	
+		if (bWrapInput)
+		{
+			double WrappedValue = FMath::Wrap(Result, MinValue, MaxValue);
+			for (const FChooserIndexArray::FIndexData& IndexData : IndexListIn)
+			{
+				if (RowValues.Num() > static_cast<int>(IndexData.Index))
+				{
+					const FChooserFloatDistanceRowData& RowValue = RowValues[IndexData.Index];
+					
+					double Distance =  fabs(WrappedValue - RowValue.Value);
+					if (Distance > 0.5 * MaxValue - MinValue)
+					{
+						Distance = MaxValue - MinValue - Distance;
+					}
+					
+					if (!bFilterOverMaxDistance || Distance < MaxDistance)
+					{
+						float Cost = FMath::Min(Distance / MaxDistance, 1.0);
+						IndexListOut.Push({IndexData.Index, IndexData.Cost + CostMultiplier * Cost});
+					}
+				}
+			}
+		}
+		else
+		{
+			for (const FChooserIndexArray::FIndexData& IndexData : IndexListIn)
+			{
+				if (RowValues.IsValidIndex(IndexData.Index))
+				{
+					const FChooserFloatDistanceRowData& RowValue = RowValues[IndexData.Index];
+					double Distance = fabs(Result - RowValue.Value);
+					if (!bFilterOverMaxDistance || Distance < MaxDistance)
+					{
+						float Cost = FMath::Min(Distance / MaxDistance, 1.0);
+						IndexListOut.Push({IndexData.Index, IndexData.Cost + CostMultiplier * Cost});
+					}
+				}
+			}
+		}
+	}
+	else
+	{
+		// passthrough fallback (behaves better during live editing)
+		IndexListOut = IndexListIn;
+	}
+}
+
+#if WITH_EDITOR
+	void FFloatDistanceColumn::AddToDetails(FInstancedPropertyBag& PropertyBag, int32 ColumnIndex, int32 RowIndex)
+	{
+		FText DisplayName;
+		InputValue.Get<FChooserParameterFloatBase>().GetDisplayName(DisplayName);
+		FName PropertyName("RowData", ColumnIndex);
+		FPropertyBagPropertyDesc PropertyDesc(PropertyName,  EPropertyBagPropertyType::Struct, FChooserFloatDistanceRowData::StaticStruct());
+		PropertyDesc.MetaData.Add(FPropertyBagPropertyDescMetaData("DisplayName", DisplayName.ToString()));
+		PropertyBag.AddProperties({PropertyDesc});
+		
+		PropertyBag.SetValueStruct(PropertyName, RowValues[RowIndex]);
+	}
+
+	void FFloatDistanceColumn::SetFromDetails(FInstancedPropertyBag& PropertyBag, int32 ColumnIndex, int32 RowIndex)
+	{
+		FName PropertyName("RowData", ColumnIndex);
+		
+   		TValueOrError<FStructView, EPropertyBagResult> Result = PropertyBag.GetValueStruct(PropertyName, FChooserFloatDistanceRowData::StaticStruct());
+		if (FStructView* StructView = Result.TryGetValue())
+		{
+			RowValues[RowIndex] = StructView->Get<FChooserFloatDistanceRowData>();
+		}
+	}
+#endif
\ No newline at end of file
diff --git a/Engine/Plugins/Experimental/Chooser/Source/Chooser/Private/FloatRangeColumn.cpp b/Engine/Plugins/Experimental/Chooser/Source/Chooser/Private/FloatRangeColumn.cpp
--- a/Engine/Plugins/Experimental/Chooser/Source/Chooser/Private/FloatRangeColumn.cpp
+++ b/Engine/Plugins/Experimental/Chooser/Source/Chooser/Private/FloatRangeColumn.cpp
@@ -69,35 +69,36 @@ void FFloatRangeColumn::Filter(FChooserEvaluationContext& Context, const FChoose
 		if (bWrapInput)
 		{
 			Result = FMath::Wrap(Result, MinValue, MaxValue);
-			for(uint32 Index : IndexListIn)
+			
+			for (const FChooserIndexArray::FIndexData& IndexData : IndexListIn)
 			{
-				if (RowValues.Num() > static_cast<int>(Index))
+				if (RowValues.Num() > static_cast<int>(IndexData.Index))
 				{
-					const FChooserFloatRangeRowData& RowValue = RowValues[Index];
+					const FChooserFloatRangeRowData& RowValue = RowValues[IndexData.Index];
 					if (RowValue.Max < RowValue.Min) // eg for an angle range from  135 to -135  (135 to 180 or -180 to -135)
 					{
 						if (Result >= RowValue.Min || Result <= RowValue.Max)
 						{
-							IndexListOut.Push(Index);
+							IndexListOut.Push(IndexData);
 						}
 					}
 					else if (Result >= RowValue.Min && Result <= RowValue.Max)
 					{
-						IndexListOut.Push(Index);
+						IndexListOut.Push(IndexData);
 					}
 				}
 			}
 		}
 		else
 		{
-			for(uint32 Index : IndexListIn)
+			for (const FChooserIndexArray::FIndexData& IndexData : IndexListIn)
 			{
-				if (RowValues.Num() > static_cast<int>(Index))
+				if (RowValues.Num() > static_cast<int>(IndexData.Index))
 				{
-					const FChooserFloatRangeRowData& RowValue = RowValues[Index];
+					const FChooserFloatRangeRowData& RowValue = RowValues[IndexData.Index];
 					if (Result >= RowValue.Min && Result <= RowValue.Max)
 					{
-						IndexListOut.Push(Index);
+						IndexListOut.Push(IndexData);
 					}
 				}
 			}
diff --git a/Engine/Plugins/Experimental/Chooser/Source/Chooser/Private/GameplayTagColumn.cpp b/Engine/Plugins/Experimental/Chooser/Source/Chooser/Private/GameplayTagColumn.cpp
--- a/Engine/Plugins/Experimental/Chooser/Source/Chooser/Private/GameplayTagColumn.cpp
+++ b/Engine/Plugins/Experimental/Chooser/Source/Chooser/Private/GameplayTagColumn.cpp
@@ -71,11 +71,11 @@ void FGameplayTagColumn::Filter(FChooserEvaluationContext& Context, const FChoos
 		}
 #endif
 
-		for (uint32 Index : IndexListIn)
+		for (const FChooserIndexArray::FIndexData& IndexData : IndexListIn)
 		{
-			if (TestRow(Index, *Result))
+			if (TestRow(IndexData.Index, *Result))
 			{
-				IndexListOut.Push(Index);
+				IndexListOut.Push(IndexData);
 			}
 		}
 	}
diff --git a/Engine/Plugins/Experimental/Chooser/Source/Chooser/Private/MultiEnumColumn.cpp b/Engine/Plugins/Experimental/Chooser/Source/Chooser/Private/MultiEnumColumn.cpp
--- a/Engine/Plugins/Experimental/Chooser/Source/Chooser/Private/MultiEnumColumn.cpp
+++ b/Engine/Plugins/Experimental/Chooser/Source/Chooser/Private/MultiEnumColumn.cpp
@@ -40,14 +40,14 @@ void FMultiEnumColumn::Filter(FChooserEvaluationContext& Context, const FChooser
 
 		// log if Result > 31
 		uint32 ResultBit = 1 << Result;
-		for (const uint32 Index : IndexListIn)
+		for (const FChooserIndexArray::FIndexData& IndexData : IndexListIn)
 		{
-			if (RowValues.IsValidIndex(Index))
+			if (RowValues.IsValidIndex(IndexData.Index))
 			{
-				const FChooserMultiEnumRowData& RowValue = RowValues[Index];
+				const FChooserMultiEnumRowData& RowValue = RowValues[IndexData.Index];
 				if (RowValue.Evaluate(ResultBit))
 				{
-					IndexListOut.Push(Index);
+					IndexListOut.Push(IndexData);
 				}
 			}
 		}
diff --git a/Engine/Plugins/Experimental/Chooser/Source/Chooser/Private/ObjectColumn.cpp b/Engine/Plugins/Experimental/Chooser/Source/Chooser/Private/ObjectColumn.cpp
--- a/Engine/Plugins/Experimental/Chooser/Source/Chooser/Private/ObjectColumn.cpp
+++ b/Engine/Plugins/Experimental/Chooser/Source/Chooser/Private/ObjectColumn.cpp
@@ -97,14 +97,14 @@ void FObjectColumn::Filter(FChooserEvaluationContext& Context, const FChooserInd
 		}
 #endif
 		
-		for (const uint32 Index : IndexListIn)
+		for (const FChooserIndexArray::FIndexData& IndexData : IndexListIn)
 		{
-			if (RowValues.IsValidIndex(Index))
+			if (RowValues.IsValidIndex(IndexData.Index))
 			{
-				const FChooserObjectRowData& RowValue = RowValues[Index];
+				const FChooserObjectRowData& RowValue = RowValues[IndexData.Index];
 				if (RowValue.Evaluate(Result))
 				{
-					IndexListOut.Push(Index);
+					IndexListOut.Push(IndexData);
 				}
 			}
 		}
diff --git a/Engine/Plugins/Experimental/Chooser/Source/Chooser/Private/RandomizeColumn.cpp b/Engine/Plugins/Experimental/Chooser/Source/Chooser/Private/RandomizeColumn.cpp
--- a/Engine/Plugins/Experimental/Chooser/Source/Chooser/Private/RandomizeColumn.cpp
+++ b/Engine/Plugins/Experimental/Chooser/Source/Chooser/Private/RandomizeColumn.cpp
@@ -37,49 +37,127 @@ void FRandomizeColumn::Filter(FChooserEvaluationContext& Context, const FChooser
 				LastSelectedIndex = State->LastSelectedRow;
 			}
 		}
-		
-		// compute the sum of all weights/probabilities
-		float TotalWeight = 0;
-		for (uint32 Index : IndexListIn)
+
+		if (IndexListIn.HasCosts())
 		{
-			float RowWeight = 1.0f;
-			if (RowValues.Num() > static_cast<int32>(Index))
+			// find the lowest cost row
+			float LowestCost = UE_MAX_FLT;
+			uint32 LowestCostIndex = 0;
+			for (const FChooserIndexArray::FIndexData& IndexData : IndexListIn)
 			{
-				RowWeight = RowValues[Index];
+				if (IndexData.Cost < LowestCost)
+				{
+					LowestCost = IndexData.Cost;
+					LowestCostIndex = IndexData.Index;
+				}
 			}
+
+			// compute the sum of all weights/probabilities - only considering rows with cost nearly equal to the lowest cost
+			float TotalWeight = 0;
+			int MinCount = 0;
+			for (const FChooserIndexArray::FIndexData& IndexData : IndexListIn)
+			{
+				if (FMath::IsNearlyEqual(LowestCost, IndexData.Cost, EqualCostThreshold))
+				{
+					float RowWeight = 1.0f;
+					if (RowValues.Num() > static_cast<int32>(IndexData.Index))
+					{
+						RowWeight = RowValues[IndexData.Index];
+					}
+			
+					if (IndexData.Index == LastSelectedIndex)
+					{
+						RowWeight *= RepeatProbabilityMultiplier;
+					}
 			
-			if (Index == LastSelectedIndex)
+					TotalWeight += RowWeight;
+					MinCount++;
+				}
+			}
+
+			if (MinCount == 1)
 			{
-				RowWeight *= RepeatProbabilityMultiplier;
+				// only one entry with lowest cost, we can't randomize
+				IndexListOut.Push({LowestCostIndex, LowestCost});
+				return;
 			}
-			
-			TotalWeight += RowWeight;
-		}
+				
 
-		// pick a random float from 0-total weight
-		float RandomNumber = FMath::FRandRange(0.0f, TotalWeight);
-		float Weight = 0;
+			// pick a random float from 0-total weight
+			float RandomNumber = FMath::FRandRange(0.0f, TotalWeight);
+			float Weight = 0;
 
-		// add up the weights again, and select the index where our sum clears the random float
-		for (; Selection < Count - 1; Selection++)
+			// add up the weights again, and select the index where our sum clears the random float
+			for (; Selection < Count - 1; Selection++)
+			{
+				const FChooserIndexArray::FIndexData& IndexData = IndexListIn[Selection];
+				if (FMath::IsNearlyEqual(LowestCost, IndexData.Cost, EqualCostThreshold))
+				{
+					float RowWeight = 1.0f;
+				
+					if (RowValues.IsValidIndex(IndexData.Index))
+					{
+						RowWeight = RowValues[IndexData.Index];
+					}
+				
+					if (IndexData.Index == LastSelectedIndex)
+					{
+						RowWeight *= RepeatProbabilityMultiplier;
+					}
+				
+					Weight += RowWeight;
+					if (Weight > RandomNumber)
+					{
+						break;
+					}
+				}
+			}	
+		}
+		else
 		{
-			int Index = static_cast<int32>(IndexListIn[Selection]);
-			float RowWeight = 1.0f;
-			
-			if (RowValues.Num() > Index)
+			// compute the sum of all weights/probabilities
+			float TotalWeight = 0;
+			for (const FChooserIndexArray::FIndexData& IndexData : IndexListIn)
 			{
-				RowWeight = RowValues[Index];
-			}
+				float RowWeight = 1.0f;
+				if (RowValues.Num() > static_cast<int32>(IndexData.Index))
+				{
+					RowWeight = RowValues[IndexData.Index];
+				}
 			
-			if (Index == LastSelectedIndex)
-			{
-				RowWeight *= RepeatProbabilityMultiplier;
-			}
+				if (IndexData.Index == LastSelectedIndex)
+				{
+					RowWeight *= RepeatProbabilityMultiplier;
+				}
 			
-			Weight += RowWeight;
-			if (Weight > RandomNumber)
+				TotalWeight += RowWeight;
+			}
+
+			// pick a random float from 0-total weight
+			float RandomNumber = FMath::FRandRange(0.0f, TotalWeight);
+			float Weight = 0;
+
+			// add up the weights again, and select the index where our sum clears the random float
+			for (; Selection < Count - 1; Selection++)
 			{
-				break;
+				const FChooserIndexArray::FIndexData& IndexData = IndexListIn[Selection];
+				float RowWeight = 1.0f;
+			
+				if (RowValues.IsValidIndex(IndexData.Index))
+				{
+					RowWeight = RowValues[IndexData.Index];
+				}
+			
+				if (IndexData.Index == LastSelectedIndex)
+				{
+					RowWeight *= RepeatProbabilityMultiplier;
+				}
+			
+				Weight += RowWeight;
+				if (Weight > RandomNumber)
+				{
+					break;
+				}
 			}
 		}
 	}
diff --git a/Engine/Plugins/Experimental/Chooser/Source/Chooser/Public/ChooserIndexArray.h b/Engine/Plugins/Experimental/Chooser/Source/Chooser/Public/ChooserIndexArray.h
--- a/Engine/Plugins/Experimental/Chooser/Source/Chooser/Public/ChooserIndexArray.h
+++ b/Engine/Plugins/Experimental/Chooser/Source/Chooser/Public/ChooserIndexArray.h
@@ -7,10 +7,25 @@
 class FChooserIndexArray
 {
 public:
-	FChooserIndexArray(uint32* InData, uint32 InMaxSize): Data(InData), MaxSize(InMaxSize), Size(0) {}
+	struct FIndexData
+	{
+		UE_DEPRECATED(5.6,"This is a fallback conversion operator for Chooser Column implementations which expect index data to be a uint32.  Please convert Chooser Column implementations to use FIndexData before the next release")
+		FIndexData(uint32 InIndex) : Index(InIndex), Cost(0) { }
+		
+		FIndexData(uint32 InIndex, float InCost) : Index(InIndex), Cost(InCost) {}
+	
+		UE_DEPRECATED(5.5, "This is a fallback conversion operator for Chooser Column implementations which expect index data to be a uint32.  Please convert Chooser Column implementations to use FIndexData before the next release")
+		operator uint32() const { return Index; }
+		uint32 Index;
+		float Cost;
+
+		bool operator < (const FIndexData& Other) const { return Cost < Other.Cost; }
+	};
 	
-	typedef uint32* iterator;
-	typedef const uint32* const_iterator;
+	FChooserIndexArray(FIndexData* InData, uint32 InMaxSize): Data(InData), MaxSize(InMaxSize), Size(0) {}
+	
+	typedef FIndexData* iterator;
+	typedef const FIndexData* const_iterator;
 	
 	iterator begin() { return Data; }
 	const_iterator begin() const { return Data; }
@@ -18,7 +33,7 @@ public:
 	const_iterator end() const { return Data + Size; }
 	
 	
-	void Push(uint32 Value)
+	void Push(const FIndexData& Value)
 	{
 		check(Size < MaxSize);
 		Data[Size++] = Value;
@@ -40,13 +55,13 @@ public:
 		Size = Num;		
 	}
 
-	uint32& operator [] (uint32 Index)
+	FIndexData& operator [] (uint32 Index)
 	{
 		check(Index < MaxSize);
 		return Data[Index];
 	}
 	
-	uint32 operator [] (uint32 Index) const
+	const FIndexData& operator [] (uint32 Index) const
 	{
 		check(Index < MaxSize);
 		return Data[Index];
@@ -59,8 +74,24 @@ public:
 		FMemory::Memcpy(Data,Other.Data, Size*sizeof(uint32));
 	}
 
+	FIndexData* GetData() { return Data; }
+
+	bool HasCosts() const { return bHasCosts; };
+	void SetHasCosts() { bHasCosts = true; };
 private:
-	uint32* Data;
+	FIndexData* Data;
 	uint32 MaxSize;
 	uint32 Size;
+	bool bHasCosts = false;
 };
+
+inline FChooserIndexArray::FIndexData* GetData(FChooserIndexArray& Array)
+{
+	return Array.GetData();
+}
+
+inline uint32 GetNum(FChooserIndexArray& Array)
+{
+	return Array.Num();
+}
+
diff --git a/Engine/Plugins/Experimental/Chooser/Source/Chooser/Public/IChooserColumn.h b/Engine/Plugins/Experimental/Chooser/Source/Chooser/Public/IChooserColumn.h
--- a/Engine/Plugins/Experimental/Chooser/Source/Chooser/Public/IChooserColumn.h
+++ b/Engine/Plugins/Experimental/Chooser/Source/Chooser/Public/IChooserColumn.h
@@ -40,6 +40,7 @@ public:
 	virtual void PostLoad() {};
 	virtual void Filter(FChooserEvaluationContext& Context, const FChooserIndexArray& IndexListIn, FChooserIndexArray& IndexListOut) const {}
 
+	virtual bool HasCosts() const { return false; }
 	virtual bool HasFilters() const { return true; }
 	virtual bool HasOutputs() const { return false; }
 	
@@ -70,6 +71,7 @@ public:
 	// using a virtual fucntion to identify them (rather than hard coding a specific type) to potentially support multiple varieties of randomization column.
 	virtual bool IsRandomizeColumn() const { return false; }
 
+	virtual float EditorTestCost(int32 RowIndex) const { return 0.0f; }
 	virtual bool EditorTestFilter(int32 RowIndex) const { return false; }
 
 	virtual void AddToDetails(FInstancedPropertyBag& PropertyBag, int32 ColumnIndex, int32 RowIndex) {}
diff --git a/Engine/Plugins/Experimental/Chooser/Source/ChooserEditor/Private/ChooserEditorModule.cpp b/Engine/Plugins/Experimental/Chooser/Source/ChooserEditor/Private/ChooserEditorModule.cpp
--- a/Engine/Plugins/Experimental/Chooser/Source/ChooserEditor/Private/ChooserEditorModule.cpp
+++ b/Engine/Plugins/Experimental/Chooser/Source/ChooserEditor/Private/ChooserEditorModule.cpp
@@ -11,6 +11,7 @@
 #include "CurveOverrideCustomization.h"
 #include "EnumColumnEditor.h"
 #include "Features/IModularFeatures.h"
+#include "FloatDistanceColumnEditor.h"
 #include "FloatRangeColumnEditor.h"
 #include "FrameTimeCustomization.h"
 #include "GameplayTagColumnEditor.h"
@@ -38,6 +39,7 @@ void FModule::StartupModule()
 	
 	FChooserTableEditor::RegisterWidgets();
 	RegisterGameplayTagWidgets();
+	RegisterFloatDistanceWidgets();
 	RegisterFloatRangeWidgets();
 	RegisterOutputFloatWidgets();
 	RegisterBoolWidgets();
diff --git a/Engine/Plugins/Experimental/Chooser/Source/ChooserEditor/Private/FloatDistanceColumnEditor.cpp b/Engine/Plugins/Experimental/Chooser/Source/ChooserEditor/Private/FloatDistanceColumnEditor.cpp
new file mode 100644
--- /dev/null
+++ b/Engine/Plugins/Experimental/Chooser/Source/ChooserEditor/Private/FloatDistanceColumnEditor.cpp
@@ -0,0 +1,95 @@
+// Copyright Epic Games, Inc. All Rights Reserved.
+
+#include "FloatDistanceColumnEditor.h"
+#include "FloatDistanceColumn.h"
+#include "ObjectChooserWidgetFactories.h"
+#include "ChooserTableEditor.h"
+#include "Widgets/Input/SEditableTextBox.h"
+#include "Widgets/Input/SNumericEntryBox.h"
+#include "Widgets/Images/SImage.h"
+#include "GraphEditorSettings.h"
+#include "ScopedTransaction.h"
+
+#define LOCTEXT_NAMESPACE "FloatDistanceColumnEditor"
+
+namespace UE::ChooserEditor
+{
+
+TSharedRef<SWidget> CreateFloatDistanceColumnWidget(UChooserTable* Chooser, FChooserColumnBase* Column, int Row)
+{
+	FFloatDistanceColumn* FloatDistanceColumn = static_cast<FFloatDistanceColumn*>(Column);
+
+	if (Row == ColumnWidget_SpecialIndex_Fallback)
+	{
+		return SNullWidget::NullWidget;
+	}
+	else if (Row == ColumnWidget_SpecialIndex_Header)
+	{
+		// create column header widget
+		TSharedPtr<SWidget> InputValueWidget = nullptr;
+		if (FChooserParameterBase* InputValue = Column->GetInputValue())
+		{
+			InputValueWidget = FObjectChooserWidgetFactories::CreateWidget(false, Chooser, InputValue, Column->GetInputType(), Chooser->OutputObjectType);
+		}
+		
+		const FSlateBrush* ColumnIcon = FCoreStyle::Get().GetBrush("Icons.Filter");
+		
+		TSharedRef<SWidget> ColumnHeaderWidget = SNew(SHorizontalBox)
+			+ SHorizontalBox::Slot().AutoWidth()
+			[
+				SNew(SBorder)
+				.BorderBackgroundColor(FLinearColor(0,0,0,0))
+				.Content()
+				[
+					SNew(SImage).Image(ColumnIcon)
+				]
+			]
+			+ SHorizontalBox::Slot()
+			[
+				InputValueWidget ? InputValueWidget.ToSharedRef() : SNullWidget::NullWidget
+			];
+	
+		if (Chooser->bEnableDebugTesting)
+		{
+			ColumnHeaderWidget = SNew(SVerticalBox)
+			+ SVerticalBox::Slot()
+			[
+				ColumnHeaderWidget
+			]
+			+ SVerticalBox::Slot()
+			[
+				SNew(SNumericEntryBox<float>).IsEnabled_Lambda([Chooser](){ return !Chooser->HasDebugTarget(); })
+				.Value_Lambda([FloatDistanceColumn]() { return FloatDistanceColumn->TestValue; })
+				.OnValueCommitted_Lambda([Chooser, FloatDistanceColumn](float NewValue, ETextCommit::Type CommitType) { FloatDistanceColumn->TestValue = NewValue; })
+			];
+		}
+
+		return ColumnHeaderWidget;
+	}
+
+	// create cell widget
+	return SNew(SNumericEntryBox<float>)
+	.Value_Lambda([FloatDistanceColumn, Row]()
+	{
+		return (Row < FloatDistanceColumn->RowValues.Num()) ? FloatDistanceColumn->RowValues[Row].Value : 0;
+	})
+	.OnValueCommitted_Lambda([Chooser, FloatDistanceColumn, Row](float NewValue, ETextCommit::Type CommitType)
+	{
+		if (Row < FloatDistanceColumn->RowValues.Num())
+		{
+			const FScopedTransaction Transaction(LOCTEXT("Edit Float Distance Value", "Edit Float Distance Value"));
+			Chooser->Modify(true);
+			FloatDistanceColumn->RowValues[Row].Value = NewValue;
+		}
+	});
+}
+
+	
+void RegisterFloatDistanceWidgets()
+{
+	FObjectChooserWidgetFactories::RegisterColumnWidgetCreator(FFloatDistanceColumn::StaticStruct(), CreateFloatDistanceColumnWidget);
+}
+	
+}
+
+#undef LOCTEXT_NAMESPACE
diff --git a/Engine/Plugins/Experimental/Chooser/Source/ChooserEditor/Private/FloatDistanceColumnEditor.h b/Engine/Plugins/Experimental/Chooser/Source/ChooserEditor/Private/FloatDistanceColumnEditor.h
new file mode 100644
--- /dev/null
+++ b/Engine/Plugins/Experimental/Chooser/Source/ChooserEditor/Private/FloatDistanceColumnEditor.h
@@ -0,0 +1,10 @@
+// Copyright Epic Games, Inc. All Rights Reserved.
+
+#pragma once
+
+#include "CoreMinimal.h"
+
+namespace UE::ChooserEditor
+{
+	void RegisterFloatDistanceWidgets();
+}
\ No newline at end of file
diff --git a/Engine/Plugins/Experimental/Chooser/Source/ChooserEditor/Private/SChooserTableRow.cpp b/Engine/Plugins/Experimental/Chooser/Source/ChooserEditor/Private/SChooserTableRow.cpp
--- a/Engine/Plugins/Experimental/Chooser/Source/ChooserEditor/Private/SChooserTableRow.cpp
+++ b/Engine/Plugins/Experimental/Chooser/Source/ChooserEditor/Private/SChooserTableRow.cpp
@@ -167,15 +167,38 @@ namespace UE::ChooserEditor
 							SNew(SColorBlock).Visibility(EVisibility::HitTestInvisible).Color_Lambda(
 									[this,Column]()
 									{
-										if (Chooser->bDebugTestValuesValid && Column->HasFilters())
+										if (Chooser->bDebugTestValuesValid)
 										{
-											if (Column->EditorTestFilter(RowIndex->RowIndex))
+											if (Column->HasCosts())
 											{
-												return TestPassedColor;
+												if (!Column->HasFilters() || Column->EditorTestFilter(RowIndex->RowIndex))
+												{
+													const double Cost = Column->EditorTestCost(RowIndex->RowIndex);
+
+													const FLinearColor Low(0.0, 1.0, 0.0, 0.30);
+													const FLinearColor Mid(0.8, 0.8, 0.0, 0.30);
+													const FLinearColor High(1.0, 0.0, 0.0, 0.20);
+
+													if (Cost < 0.5)
+													{
+														return FLinearColor::LerpUsingHSV(Low, Mid, FMath::Clamp(Cost * 2.0, 0.0f, 1.0f));
+													}
+													else
+													{
+														return FLinearColor::LerpUsingHSV(Mid, High, FMath::Clamp((Cost - 0.5) * 2.0, 0.0f, 1.0f));
+													}
+												}
 											}
-											else
+											else if (Column->HasFilters())
 											{
-												return TestFailedColor;
+												if (Column->EditorTestFilter(RowIndex->RowIndex))
+												{
+													return TestPassedColor;
+												}
+												else
+												{
+													return TestFailedColor;
+												}
 											}
 										}
 										return FLinearColor::Transparent;
@@ -218,7 +241,7 @@ namespace UE::ChooserEditor
 				TSharedPtr<SWidget> ResultWidget = FObjectChooserWidgetFactories::CreateWidget(false,ContextOwner, FObjectChooserBase::StaticStruct(), Chooser->FallbackResult.GetMutableMemory(), Chooser->FallbackResult.GetScriptStruct(), ContextOwner->OutputObjectType,
 				FOnStructPicked::CreateLambda([this](const UScriptStruct* ChosenStruct)
 				{
-				UChooserTable* ContextOwner = Chooser->GetContextOwner();
+					UChooserTable* ContextOwner = Chooser->GetContextOwner();
 					const FScopedTransaction Transaction(LOCTEXT("Change Row Result Type", "Change Row Result Type"));
 					Chooser->Modify(true);
 					Chooser->FallbackResult.InitializeAs(ChosenStruct);
