Chooser Table AutoPopulate feature

diff --git a/Engine/Plugins/Experimental/Chooser/Source/Chooser/Internal/FloatDistanceColumn.h b/Engine/Plugins/Experimental/Chooser/Source/Chooser/Internal/FloatDistanceColumn.h
--- a/Engine/Plugins/Experimental/Chooser/Source/Chooser/Internal/FloatDistanceColumn.h
+++ b/Engine/Plugins/Experimental/Chooser/Source/Chooser/Internal/FloatDistanceColumn.h
@@ -9,6 +9,7 @@
 #include "ChooserPropertyAccess.h"
 #include "InstancedStruct.h"
 #include "Serialization/MemoryReader.h"
+#include "Templates/SubclassOf.h"
 #include "FloatDistanceColumn.generated.h"
 
 USTRUCT()
@@ -18,8 +19,24 @@ struct FChooserFloatDistanceRowData
 
 	UPROPERTY(EditAnywhere, Category=Runtime)
 	float Value=0;
+
+#if WITH_EDITORONLY_DATA
+	// Set this property to lock the current cell value, preventing it from being updated by AutoPopulate
+	UPROPERTY(EditAnywhere, Category=Runtime)
+	bool DisableAutoPopulate = false;
+#endif
 };
 
+UCLASS(Blueprintable)
+class UFloatAutoPopulator : public UObject
+{
+	GENERATED_BODY()
+public:
+	UFUNCTION(BlueprintImplementableEvent)
+	void AutoPopulate(UObject* Object, bool& Success, float& Value);
+	
+	virtual void NativeAutoPopulate(UObject* InObject, bool& OutSuccess, float& OutValue) { AutoPopulate(InObject, OutSuccess, OutValue); }
+};
 
 USTRUCT()
 struct CHOOSER_API FFloatDistanceColumn : public FChooserColumnBase
@@ -57,6 +74,10 @@ struct CHOOSER_API FFloatDistanceColumn : public FChooserColumnBase
 #if WITH_EDITORONLY_DATA
 	UPROPERTY(EditAnywhere, Category="Data")
 	FChooserFloatDistanceRowData DefaultRowValue;
+
+	// Optional class to auto populate column data based on the result asset
+	UPROPERTY(EditAnywhere, Category = "Data")
+	TSubclassOf<UFloatAutoPopulator> AutoPopulator;
 #endif
 	
 	UPROPERTY(EditAnywhere, Category="Data")
@@ -82,6 +103,9 @@ struct CHOOSER_API FFloatDistanceColumn : public FChooserColumnBase
 
 	virtual void AddToDetails(FInstancedPropertyBag& PropertyBag, int32 ColumnIndex, int32 RowIndex) override;
 	virtual void SetFromDetails(FInstancedPropertyBag& PropertyBag, int32 ColumnIndex, int32 RowIndex) override;
+
+	virtual bool AutoPopulates() const override { return AutoPopulator != nullptr; }
+	virtual void AutoPopulate(int32 RowIndex, UObject* OutputObject) override;
 #endif
 	
 	virtual void PostLoad() override
diff --git a/Engine/Plugins/Experimental/Chooser/Source/Chooser/Internal/ObjectChooser_Asset.h b/Engine/Plugins/Experimental/Chooser/Source/Chooser/Internal/ObjectChooser_Asset.h
--- a/Engine/Plugins/Experimental/Chooser/Source/Chooser/Internal/ObjectChooser_Asset.h
+++ b/Engine/Plugins/Experimental/Chooser/Source/Chooser/Internal/ObjectChooser_Asset.h
@@ -13,7 +13,11 @@ struct CHOOSER_API FAssetChooser : public FObjectChooserBase
 	
 	// FObjectChooserBase interface
 	virtual UObject* ChooseObject(FChooserEvaluationContext& Context) const final override;
-public: 
+
+#if WITH_EDITOR
+	virtual UObject* GetReferencedObject() const override { return Asset; }
+#endif
+	
 	UPROPERTY(EditAnywhere, Category = "Parameters")
 	TObjectPtr<UObject> Asset;
 };
@@ -25,7 +29,11 @@ struct CHOOSER_API FSoftAssetChooser : public FObjectChooserBase
 
 	// FObjectChooserBase interface
 	virtual UObject* ChooseObject(FChooserEvaluationContext& Context) const final override;
-public: 
+	
+#if WITH_EDITOR
+	virtual UObject* GetReferencedObject() const override { return Asset.LoadSynchronous(); }
+#endif
+	
 	UPROPERTY(EditAnywhere, Category = "Parameters")
 	TSoftObjectPtr<UObject> Asset;
 };
diff --git a/Engine/Plugins/Experimental/Chooser/Source/Chooser/Private/FloatDistanceColumn.cpp b/Engine/Plugins/Experimental/Chooser/Source/Chooser/Private/FloatDistanceColumn.cpp
--- a/Engine/Plugins/Experimental/Chooser/Source/Chooser/Private/FloatDistanceColumn.cpp
+++ b/Engine/Plugins/Experimental/Chooser/Source/Chooser/Private/FloatDistanceColumn.cpp
@@ -14,6 +14,23 @@ FFloatDistanceColumn::FFloatDistanceColumn()
 }
 
 #if WITH_EDITOR
+void FFloatDistanceColumn::AutoPopulate(int32 RowIndex, UObject* OutputObject)
+{
+	if (AutoPopulator)
+	{
+		if (RowValues.IsValidIndex(RowIndex) && !RowValues[RowIndex].DisableAutoPopulate) 
+		{
+			float Result = 0.f;
+			bool Success = false;
+			static_cast<UFloatAutoPopulator*>(AutoPopulator->GetDefaultObject())->NativeAutoPopulate(OutputObject, Success, Result);
+			if (Success)
+			{
+				RowValues[RowIndex].Value = Result;
+			}
+		}
+	}
+}
+
 bool FFloatDistanceColumn::EditorTestFilter(int32 RowIndex) const
 {
 	if (!bFilterOverMaxDistance)
diff --git a/Engine/Plugins/Experimental/Chooser/Source/Chooser/Public/IChooserColumn.h b/Engine/Plugins/Experimental/Chooser/Source/Chooser/Public/IChooserColumn.h
--- a/Engine/Plugins/Experimental/Chooser/Source/Chooser/Public/IChooserColumn.h
+++ b/Engine/Plugins/Experimental/Chooser/Source/Chooser/Public/IChooserColumn.h
@@ -71,6 +71,9 @@ public:
 	// using a virtual fucntion to identify them (rather than hard coding a specific type) to potentially support multiple varieties of randomization column.
 	virtual bool IsRandomizeColumn() const { return false; }
 
+	virtual bool AutoPopulates() const { return false; }
+	virtual void AutoPopulate(int32 RowIndex, UObject* OutputObject) { };
+
 	virtual float EditorTestCost(int32 RowIndex) const { return 0.0f; }
 	virtual bool EditorTestFilter(int32 RowIndex) const { return false; }
 
diff --git a/Engine/Plugins/Experimental/Chooser/Source/Chooser/Public/IObjectChooser.h b/Engine/Plugins/Experimental/Chooser/Source/Chooser/Public/IObjectChooser.h
--- a/Engine/Plugins/Experimental/Chooser/Source/Chooser/Public/IObjectChooser.h
+++ b/Engine/Plugins/Experimental/Chooser/Source/Chooser/Public/IObjectChooser.h
@@ -109,4 +109,9 @@ public:
 	}
 
 	virtual void GetDebugName(FString& OutDebugName) const {};
+
+#if WITH_EDITOR
+	// for implementations with a reference to an asset or other object, get the object for use in editor
+	virtual UObject* GetReferencedObject() const { return nullptr; };
+#endif
 };
\ No newline at end of file
diff --git a/Engine/Plugins/Experimental/Chooser/Source/ChooserEditor/Private/ChooserTableEditor.cpp b/Engine/Plugins/Experimental/Chooser/Source/ChooserEditor/Private/ChooserTableEditor.cpp
--- a/Engine/Plugins/Experimental/Chooser/Source/ChooserEditor/Private/ChooserTableEditor.cpp
+++ b/Engine/Plugins/Experimental/Chooser/Source/ChooserEditor/Private/ChooserTableEditor.cpp
@@ -239,6 +239,8 @@ void FChooserTableEditor::RegisterToolbar()
 			TAttribute<FText>(),
 			FSlateIcon("EditorStyle", "FullBlueprintEditor.EditGlobalOptions")));
 
+		Section.AddEntry(FToolMenuEntry::InitToolBarButton(Commands.AutoPopulateAll));
+
 
 		Section.AddDynamicEntry("DebuggingCommands", FNewToolMenuSectionDelegate::CreateLambda([](FToolMenuSection& InSection)
 		{
@@ -290,17 +292,9 @@ void FChooserTableEditor::RegisterMenus()
 	{
 		FToolMenuSection& Section = ToolMenu->AddSection("ChooserTableRow", TAttribute<FText>());
 	
-		Section.AddEntry(FToolMenuEntry::InitMenuEntry(
-			Commands.Delete,
-			TAttribute<FText>(),
-			TAttribute<FText>(),
-			FSlateIcon()));
-	
-		Section.AddEntry(FToolMenuEntry::InitMenuEntry(
-			Commands.Disable,
-			TAttribute<FText>(),
-			TAttribute<FText>(),
-			FSlateIcon()));
+		Section.AddEntry(FToolMenuEntry::InitMenuEntry(Commands.Delete));
+		Section.AddEntry(FToolMenuEntry::InitMenuEntry(Commands.Disable));
+		Section.AddEntry(FToolMenuEntry::InitMenuEntry(Commands.AutoPopulateSelection));
 	}
 
 	
@@ -348,6 +342,10 @@ void FChooserTableEditor::BindCommands()
 		Commands.EditChooserSettings,
 		FExecuteAction::CreateSP(this, &FChooserTableEditor::SelectRootProperties));
 	
+	ToolkitCommands->MapAction(
+		Commands.AutoPopulateAll,
+		FExecuteAction::CreateSP(this, &FChooserTableEditor::AutoPopulateAll));
+	
 	ToolkitCommands->MapAction(
 		Commands.RemoveDisabledData,
 		FExecuteAction::CreateSP(this, &FChooserTableEditor::RemoveDisabledData));
@@ -357,6 +355,12 @@ void FChooserTableEditor::BindCommands()
 		FExecuteAction::CreateSP(this, &FChooserTableEditor::DeleteSelection),
 		FCanExecuteAction::CreateSP(this, &FChooserTableEditor::HasSelection)
 		);
+
+	ToolkitCommands->MapAction(
+		Commands.AutoPopulateSelection,
+		FExecuteAction::CreateSP(this, &FChooserTableEditor::AutoPopulateSelection),
+		FCanExecuteAction::CreateSP(this, &FChooserTableEditor::HasSelection)
+		);
 	
 	ToolkitCommands->MapAction(
 		Commands.Disable,
@@ -894,7 +898,6 @@ void FChooserTableEditor::UpdateTableColumns()
 							));
 				}
 
-
 				MenuBuilder.AddMenuEntry(LOCTEXT("Delete Column", "Delete"), LOCTEXT("Delete Column ToolTip", "Remove this column and all its data from the table"), FSlateIcon(),
 					FUIAction(
 						FExecuteAction::CreateLambda([this, Chooser, ColumnIndex, &Column]()
@@ -902,6 +905,18 @@ void FChooserTableEditor::UpdateTableColumns()
 							DeleteColumn(ColumnIndex);
 						})
 						));
+
+				if(Column.AutoPopulates())
+				{
+					MenuBuilder.AddMenuEntry(LOCTEXT("Auto Populate", "Auto Populate"), LOCTEXT("Auto Populate ToolTip", "Auto populate cell values for this colun"), FSlateIcon(),
+						FUIAction(
+							FExecuteAction::CreateLambda([this, &Column]()
+							{
+								const FScopedTransaction Transaction(LOCTEXT("Auto Populate Column", "Auto Populate Column"));
+								AutoPopulateColumn(Column);
+							})
+							));
+				}
 			
 				MenuBuilder.AddSubMenu(LOCTEXT("Input Type", "Input Type"),
 					LOCTEXT("InputTypeToolTip", "Change input parameter type"),
@@ -1393,6 +1408,73 @@ void FChooserTableEditor::DeleteSelectedRows()
 	UpdateTableRows();
 }
 	
+void FChooserTableEditor::AutoPopulateColumn(FChooserColumnBase& Column)
+{
+	if (UChooserTable* Chooser = GetChooser())
+	{
+		const int RowCount = Chooser->ResultsStructs.Num();
+		if (Column.AutoPopulates())
+		{
+			for (int i = 0; i < RowCount; ++i)
+			{
+				if (Chooser->ResultsStructs[i].IsValid())
+				{
+					if (UObject* ReferencedObject = Chooser->ResultsStructs[i].Get<FObjectChooserBase>().GetReferencedObject()) 
+					{
+						Column.AutoPopulate(i, ReferencedObject);
+					}
+				}
+			}
+		}
+	}
+}
+
+void FChooserTableEditor::AutoPopulateRow(int Index)
+{
+	if (UChooserTable* Chooser = GetChooser())
+	{
+		if (Chooser->ResultsStructs.IsValidIndex(Index) && Chooser->ResultsStructs[Index].IsValid())
+		{
+			if (UObject* ReferencedObject = Chooser->ResultsStructs[Index].Get<FObjectChooserBase>().GetReferencedObject()) 
+			{
+				for (FInstancedStruct& ColumnData : Chooser->ColumnsStructs)
+				{
+					FChooserColumnBase& Column = ColumnData.GetMutable<FChooserColumnBase>();
+					Column.AutoPopulate(Index, ReferencedObject);
+				}
+			}
+		}
+	}
+}
+	
+void FChooserTableEditor::AutoPopulateSelection()
+{
+	if (UChooserTable* Chooser = GetChooser())
+	{
+		const FScopedTransaction Transaction(LOCTEXT("Auto Populate Chooser", "Auto Populate All"));
+		Chooser->Modify();
+		for(UChooserRowDetails* RowDetails : SelectedRows)
+		{
+			AutoPopulateRow(RowDetails->Row);
+		}
+	}
+}
+
+void FChooserTableEditor::AutoPopulateAll()
+{
+	if (UChooserTable* Chooser = GetChooser())
+	{
+		const FScopedTransaction Transaction(LOCTEXT("Auto Populate Chooser", "Auto Populate All"));
+		Chooser->Modify();
+		for (FInstancedStruct& ColumnData : Chooser->ColumnsStructs)
+		{
+			FChooserColumnBase& Column = ColumnData.GetMutable<FChooserColumnBase>();
+			AutoPopulateColumn(Column);
+		}
+	}
+}
+
+	
 inline bool FChooserTableEditor::HasSelection()
 {
 	if (CurrentSelectionType == ESelectionType::Column)
@@ -1493,3 +1575,4 @@ void FChooserTableEditor::RegisterWidgets()
 }
 
 #undef LOCTEXT_NAMESPACE
+
diff --git a/Engine/Plugins/Experimental/Chooser/Source/ChooserEditor/Private/ChooserTableEditor.h b/Engine/Plugins/Experimental/Chooser/Source/ChooserEditor/Private/ChooserTableEditor.h
--- a/Engine/Plugins/Experimental/Chooser/Source/ChooserEditor/Private/ChooserTableEditor.h
+++ b/Engine/Plugins/Experimental/Chooser/Source/ChooserEditor/Private/ChooserTableEditor.h
@@ -96,6 +96,11 @@ namespace UE::ChooserEditor
 		bool IsSelectionDisabled();
 		void ToggleDisableSelection();
 		void DeleteSelection();
+
+		void AutoPopulateColumn(FChooserColumnBase& Column);
+		void AutoPopulateRow(int Index);
+		void AutoPopulateSelection();
+		void AutoPopulateAll();
 		
 		void UpdateTableRows();
 		void SelectColumn(UChooserTable* Chooser, int Index);
diff --git a/Engine/Plugins/Experimental/Chooser/Source/ChooserEditor/Private/ChooserTableEditorCommands.h b/Engine/Plugins/Experimental/Chooser/Source/ChooserEditor/Private/ChooserTableEditorCommands.h
--- a/Engine/Plugins/Experimental/Chooser/Source/ChooserEditor/Private/ChooserTableEditorCommands.h
+++ b/Engine/Plugins/Experimental/Chooser/Source/ChooserEditor/Private/ChooserTableEditorCommands.h
@@ -18,6 +18,8 @@ public:
 	}
 
 	TSharedPtr<FUICommandInfo> EditChooserSettings;
+	TSharedPtr<FUICommandInfo> AutoPopulateAll;
+	TSharedPtr<FUICommandInfo> AutoPopulateSelection;
 	TSharedPtr<FUICommandInfo> RemoveDisabledData;
 	TSharedPtr<FUICommandInfo> Delete;
 	TSharedPtr<FUICommandInfo> Disable;
@@ -26,6 +28,8 @@ public:
 	virtual void RegisterCommands() override
 	{
 		UI_COMMAND(EditChooserSettings, "Table Settings", "Edit the root properties of the ChooserTable asset.", EUserInterfaceActionType::Button, FInputChord())
+		UI_COMMAND(AutoPopulateAll, "AutoPopulate All", "Auto Populate cell data for all supported columns", EUserInterfaceActionType::Button, FInputChord())
+		UI_COMMAND(AutoPopulateSelection, "AutoPopulate", "Auto Populate cell data for selected rows for all supported columns", EUserInterfaceActionType::Button, FInputChord())
 		UI_COMMAND(RemoveDisabledData, "Remove Disabled Data", "Delete all data that's marked as disabled.", EUserInterfaceActionType::Button, FInputChord())
 		UI_COMMAND(Delete, "Delete", "Delete the selected Rows or Column.", EUserInterfaceActionType::Button, FInputChord())
 		UI_COMMAND(Disable, "Disable", "Disable the selected Rows or Column.", EUserInterfaceActionType::Check, FInputChord())
diff --git a/Engine/Plugins/Experimental/Chooser/Source/ChooserEditor/Private/SChooserTableRow.cpp b/Engine/Plugins/Experimental/Chooser/Source/ChooserEditor/Private/SChooserTableRow.cpp
--- a/Engine/Plugins/Experimental/Chooser/Source/ChooserEditor/Private/SChooserTableRow.cpp
+++ b/Engine/Plugins/Experimental/Chooser/Source/ChooserEditor/Private/SChooserTableRow.cpp
@@ -442,6 +442,7 @@ namespace UE::ChooserEditor
 					Editor->ClearSelectedRows();
 					for(int Index = InsertRowIndex; Index < InsertRowIndex + NewResults.Num(); Index++)
 					{
+						Editor->AutoPopulateRow(Index);
 						Editor->SelectRow(Index, false);
 					}
 				}
