chooser context data validation

diff --git a/Engine/Plugins/Experimental/Chooser/Source/Chooser/Private/Chooser.cpp b/Engine/Plugins/Experimental/Chooser/Source/Chooser/Private/Chooser.cpp
--- a/Engine/Plugins/Experimental/Chooser/Source/Chooser/Private/Chooser.cpp
+++ b/Engine/Plugins/Experimental/Chooser/Source/Chooser/Private/Chooser.cpp
@@ -259,7 +259,7 @@ void UChooserTable::PostEditUndo()
 void UChooserTable::PostEditChangeProperty(FPropertyChangedEvent& PropertyChangedEvent)
 {
 	UObject::PostEditChangeProperty(PropertyChangedEvent);
-	
+
 	if (PropertyChangedEvent.Property)
 	{
 		static FName OutputObjectTypeName = "OutputObjectType";
@@ -295,7 +295,7 @@ void UChooserTable::PostEditChangeProperty(FPropertyChangedEvent& PropertyChange
 void UChooserTable::IterateRecentContextObjects(TFunction<void(const FString&)> Callback) const
 {
 	FScopeLock Lock(&DebugLock);
-	for(const FString& ObjectName : RecentContextObjects)
+	for (const FString& ObjectName : RecentContextObjects)
 	{
 		Callback(ObjectName);
 	}
@@ -306,7 +306,7 @@ void UChooserTable::UpdateDebugging(FChooserEvaluationContext& Context) const
 	FScopeLock Lock(&DebugLock);
 
 	const UChooserTable* ContextOwner = GetContextOwner();
-	
+
 	for (const FStructView& Param : Context.Params)
 	{
 		if (const FChooserEvaluationInputObject* ObjectParam = Param.GetPtr<const FChooserEvaluationInputObject>())
@@ -314,8 +314,8 @@ void UChooserTable::UpdateDebugging(FChooserEvaluationContext& Context) const
 			if (UObject* ContextObject = ObjectParam->Object)
 			{
 				RecentContextObjects.Add(ContextObject->GetName());
-				
-				if (ContextObject->GetName() == ContextOwner->GetDebugTargetName()) 
+
+				if (ContextObject->GetName() == ContextOwner->GetDebugTargetName())
 				{
 					bDebugTestValuesValid = true;
 					Context.DebuggingInfo.bCurrentDebugTarget = true;
@@ -332,13 +332,13 @@ void UChooserTable::UpdateDebugging(FChooserEvaluationContext& Context) const
 FObjectChooserBase::EIteratorStatus UChooserTable::EvaluateChooser(FChooserEvaluationContext& Context, const UChooserTable* Chooser, FObjectChooserBase::FObjectChooserIteratorCallback Callback)
 {
 	TRACE_CPUPROFILER_EVENT_SCOPE(EvaluateChooser);
-	
+
 	if (Chooser == nullptr)
 	{
 		return FObjectChooserBase::EIteratorStatus::Continue;
 	}
 
-	// todo validate that parameter types in context data match
+	VALIDATE_CHOOSER_CONTEXT(Chooser, Chooser->ContextData, Context);
 
 #if WITH_EDITOR
 	{
diff --git a/Engine/Plugins/Experimental/Chooser/Source/Chooser/Private/ChooserPropertyAccess.cpp b/Engine/Plugins/Experimental/Chooser/Source/Chooser/Private/ChooserPropertyAccess.cpp
--- a/Engine/Plugins/Experimental/Chooser/Source/Chooser/Private/ChooserPropertyAccess.cpp
+++ b/Engine/Plugins/Experimental/Chooser/Source/Chooser/Private/ChooserPropertyAccess.cpp
@@ -15,6 +15,12 @@
 
 #define LOCTEXT_NAMESPACE "ChooserPropertyAccess"
 
+TAutoConsoleVariable<bool> CVarEnableDetailedWarnings(
+	TEXT("Choosers.EnableDetailedWarnings"),
+	false,
+	TEXT("Enable detailed context validation with warnings when choosers are evaluated on an incorrect context. \n0: Disable (default), 1: Enable"),
+	ECVF_Default);
+
 TAutoConsoleVariable<bool> CVarUseCompiledPropertyChainsInEditor(
 	TEXT("Choosers.UseCompiledPropertyChainsInEditor"),
 	false,
@@ -304,6 +310,84 @@ void FChooserPropertyBinding::Compile(IHasContextClass* Owner, bool bForce)
 
 namespace UE::Chooser
 {
+	
+	void RuntimeValidateContext(const UObject* Chooser, const TArray<FInstancedStruct>& ContextData, FChooserEvaluationContext& Context)
+	{
+		if (!CVarEnableDetailedWarnings.GetValueOnAnyThread())
+		{
+			return;
+		}
+
+		int ContextNum = ContextData.Num();
+		for (int i = 0; i < ContextNum; i++)
+		{
+			const FContextObjectTypeClass* ExpectedClassType = ContextData[i].GetPtr<FContextObjectTypeClass>();
+			const FContextObjectTypeStruct* ExpectedStructType = ContextData[i].GetPtr<FContextObjectTypeStruct>();
+
+			if (ExpectedClassType)
+			{
+				if (Context.Params.IsValidIndex(i) && Context.Params[i].IsValid())
+				{
+					if (FChooserEvaluationInputObject* InputObjectParam = Context.Params[i].GetPtr<FChooserEvaluationInputObject>())
+					{
+						if (InputObjectParam->Object)
+						{
+							if (!InputObjectParam->Object->GetClass()->IsChildOf(ExpectedClassType->Class))
+							{
+								UE_LOG(LogChooser, Warning, TEXT("Chooser Table: %s ContextData entry %d expects an object of type %s, but an object of type %s was passed in."),
+									ToCStr(Chooser->GetName()), i, ToCStr(ExpectedClassType->Class->GetName()), ToCStr(InputObjectParam->Object->GetClass()->GetName()));
+							}
+						}
+						else
+						{
+							UE_LOG(LogChooser, Warning, TEXT("Chooser Table: %s ContextData entry %d expects an object of type %s, but null was passed in."),
+								ToCStr(Chooser->GetName()), i, ToCStr(ExpectedClassType->Class->GetName()));
+						}
+
+					}
+					else
+					{
+						UE_LOG(LogChooser, Warning, TEXT("Chooser Table: %s ContextData entry %d expects an object of type %s, but was passed a struct of type %s."),
+							ToCStr(Chooser->GetName()), i, ToCStr(ExpectedClassType->Class->GetName()), ToCStr(Context.Params[i].GetScriptStruct()->GetName()));
+					}
+				}
+				else
+				{
+					UE_LOG(LogChooser, Warning, TEXT("Chooser Table: %s ContextData entry %d expects an object of type %s, but nothing was passed in."),
+						ToCStr(Chooser->GetName()), i, ToCStr(ExpectedClassType->Class->GetName()));
+				}
+			}
+			else if (ExpectedStructType)
+			{
+				if (Context.Params.IsValidIndex(i) && Context.Params[i].IsValid())
+				{
+					if (FChooserEvaluationInputObject* InputObjectParam = Context.Params[i].GetPtr<FChooserEvaluationInputObject>())
+					{
+						UE_LOG(LogChooser, Warning, TEXT("Chooser Table: %s ContextData entry %d expects an struct of type %s, but was passed a object of type %s."),
+							ToCStr(Chooser->GetName()), i, ToCStr(ExpectedStructType->Struct->GetName()), ToCStr(InputObjectParam->Object->GetClass()->GetName()));
+					}
+					else
+					{
+						if (Context.Params[i].GetScriptStruct() != ExpectedStructType->Struct)
+						{
+							UE_LOG(LogChooser, Warning, TEXT("Chooser Table: %s ContextData entry %d expects an struct of type %s, but was passed a struct of type %s."),
+								ToCStr(Chooser->GetName()), i, ToCStr(ExpectedStructType->Struct->GetName()), ToCStr(Context.Params[i].GetScriptStruct()->GetName()));
+						}
+					}
+				}
+				else
+				{
+					UE_LOG(LogChooser, Warning, TEXT("Chooser Table: %s ContextData entry %d expects an struct of type %s, but nothing was passed in."),
+						ToCStr(Chooser->GetName()), i, ToCStr(ExpectedStructType->Struct->GetName()));
+				}
+			}
+			else
+			{
+				UE_LOG(LogChooser, Warning, TEXT("Chooser Table: %s  ContextData entry %i is of unknown type or none."), ToCStr(Chooser->GetName()), i);
+			}
+		}
+	}
+
 
 	bool ResolveCompiledPropertyChain(FChooserEvaluationContext& Context, const FChooserPropertyBinding& Binding, FResolvedPropertyChainResult& Result)
 	{
diff --git a/Engine/Plugins/Experimental/Chooser/Source/Chooser/Public/ChooserPropertyAccess.h b/Engine/Plugins/Experimental/Chooser/Source/Chooser/Public/ChooserPropertyAccess.h
--- a/Engine/Plugins/Experimental/Chooser/Source/Chooser/Public/ChooserPropertyAccess.h
+++ b/Engine/Plugins/Experimental/Chooser/Source/Chooser/Public/ChooserPropertyAccess.h
@@ -186,6 +186,13 @@ struct FContextObjectTypeStruct : public FContextObjectTypeBase
 
 namespace UE::Chooser
 {
+#if WITH_EDITORONLY_DATA
+	CHOOSER_API void RuntimeValidateContext(const UObject* Asset, const TArray<FInstancedStruct>& ContextData, FChooserEvaluationContext& Context);
+#define VALIDATE_CHOOSER_CONTEXT(Asset,ContextData,Context) UE::Chooser::RuntimeValidateContext(Asset, ContextData, Context)
+#else
+#define VALIDATE_CHOOSER_CONTEXT(Asset,ContextData,Context)
+#endif
+
 	struct FResolvedPropertyChainResult
 	{
 		uint8* Container = nullptr;
diff --git a/Engine/Plugins/Experimental/Chooser/Source/ProxyTable/Private/LookupProxy.cpp b/Engine/Plugins/Experimental/Chooser/Source/ProxyTable/Private/LookupProxy.cpp
--- a/Engine/Plugins/Experimental/Chooser/Source/ProxyTable/Private/LookupProxy.cpp
+++ b/Engine/Plugins/Experimental/Chooser/Source/ProxyTable/Private/LookupProxy.cpp
@@ -13,6 +13,8 @@ FObjectChooserBase::EIteratorStatus FLookupProxy::ChooseMulti(FChooserEvaluation
 {
 	if (Proxy)
 	{
+		VALIDATE_CHOOSER_CONTEXT(Proxy, Proxy->ContextData, Context);
+
 		if (const FChooserParameterProxyTableBase* ProxyTableParameter = ProxyTable.GetPtr<FChooserParameterProxyTableBase>())
 		{
 			const UProxyTable* Table = nullptr;
@@ -34,6 +36,7 @@ UObject* FLookupProxy::ChooseObject(FChooserEvaluationContext& Context) const
 {
 	if (Proxy)
 	{
+		VALIDATE_CHOOSER_CONTEXT(Proxy, Proxy->ContextData, Context);
 		if (const FChooserParameterProxyTableBase* ProxyTableParameter = ProxyTable.GetPtr<FChooserParameterProxyTableBase>())
 		{
 			const UProxyTable* Table = nullptr;
diff --git a/Engine/Plugins/Experimental/Chooser/Source/ProxyTable/Private/ProxyTable.cpp b/Engine/Plugins/Experimental/Chooser/Source/ProxyTable/Private/ProxyTable.cpp
--- a/Engine/Plugins/Experimental/Chooser/Source/ProxyTable/Private/ProxyTable.cpp
+++ b/Engine/Plugins/Experimental/Chooser/Source/ProxyTable/Private/ProxyTable.cpp
@@ -192,10 +192,12 @@ void UProxyTable::BeginDestroy()
 	Super::BeginDestroy();
 }
 
-static void OutputStructData(const FRuntimeProxyValue& EntryValueData, FChooserEvaluationContext& Context)
+static void OutputStructData(const UProxyTable* ProxyTable, const FRuntimeProxyValue& EntryValueData, FChooserEvaluationContext& Context)
 {
 	for (const FProxyStructOutput& StructOutput : EntryValueData.OutputStructData)
 	{
+		VALIDATE_CHOOSER_CONTEXT(EntryValueData.ProxyAsset, EntryValueData.ProxyAsset->ContextData, Context);
+
 		// copy each struct output value
 		void* TargetData = nullptr;
 		const UStruct* StructType = nullptr;
@@ -205,6 +207,11 @@ static void OutputStructData(const FRuntimeProxyValue& EntryValueData, FChooserE
 			{
 				StructOutput.Value.GetScriptStruct()->CopyScriptStruct(TargetData, StructOutput.Value.GetMemory());
 			}
+			else
+			{
+				UE_LOG(LogChooser, Warning, TEXT("Proxy Table %s Struct Output Type mismatch for Proxy Asset: %s"),
+					ToCStr(ProxyTable->GetName()), ToCStr(EntryValueData.ProxyAsset->GetName()));
+			}
 		}
 	}
 }
@@ -219,7 +226,7 @@ UObject* UProxyTable::FindProxyObject(const FGuid& Key, FChooserEvaluationContex
 
 		UObject* Result = EntryValue.ChooseObject(Context);
 
-		OutputStructData(EntryValueData, Context);
+		OutputStructData(this, EntryValueData, Context);
 
 		return Result;
 	}
@@ -229,20 +236,20 @@ UObject* UProxyTable::FindProxyObject(const FGuid& Key, FChooserEvaluationContex
 
 FObjectChooserBase::EIteratorStatus UProxyTable::FindProxyObjectMulti(const FGuid& Key, FChooserEvaluationContext &Context, FObjectChooserBase::FObjectChooserIteratorCallback Callback) const
 {
-		const int FoundIndex = Algo::BinarySearch(Keys, Key);
-    	if (FoundIndex != INDEX_NONE)
-    	{
-    		const FRuntimeProxyValue& EntryValueData = RuntimeValues[FoundIndex];
-    		const FObjectChooserBase &EntryValue = EntryValueData.Value.Get<const FObjectChooserBase>();
+	const int FoundIndex = Algo::BinarySearch(Keys, Key);
+    if (FoundIndex != INDEX_NONE)
+    {
+    	const FRuntimeProxyValue& EntryValueData = RuntimeValues[FoundIndex];
+    	const FObjectChooserBase &EntryValue = EntryValueData.Value.Get<const FObjectChooserBase>();
     
-    		FObjectChooserBase::EIteratorStatus Result = EntryValue.ChooseMulti(Context, Callback);
+    	FObjectChooserBase::EIteratorStatus Result = EntryValue.ChooseMulti(Context, Callback);
     
-			OutputStructData(EntryValueData, Context);
+		OutputStructData(this, EntryValueData, Context);
     
-    		return Result;
-    	}
+    	return Result;
+    }
     	
-    	return FObjectChooserBase::EIteratorStatus::Continue;
+    return FObjectChooserBase::EIteratorStatus::Continue;
 }
 
 UProxyTable::UProxyTable(const FObjectInitializer& Initializer)
