Fixes for autopopulate

diff --git a/Engine/Plugins/Experimental/Chooser/Source/Chooser/Internal/FloatDistanceColumn.h b/Engine/Plugins/Experimental/Chooser/Source/Chooser/Internal/FloatDistanceColumn.h
--- a/Engine/Plugins/Experimental/Chooser/Source/Chooser/Internal/FloatDistanceColumn.h
+++ b/Engine/Plugins/Experimental/Chooser/Source/Chooser/Internal/FloatDistanceColumn.h
@@ -27,15 +27,12 @@ struct FChooserFloatDistanceRowData
 #endif
 };
 
-UCLASS(Blueprintable)
-class UFloatAutoPopulator : public UObject
+UCLASS()
+class CHOOSER_API UFloatAutoPopulator : public UObject
 {
 	GENERATED_BODY()
 public:
-	UFUNCTION(BlueprintImplementableEvent)
-	void AutoPopulate(UObject* Object, bool& Success, float& Value);
-	
-	virtual void NativeAutoPopulate(UObject* InObject, bool& OutSuccess, float& OutValue) { AutoPopulate(InObject, OutSuccess, OutValue); }
+	virtual void NativeAutoPopulate(UObject* InObject, bool& OutSuccess, float& OutValue) { }
 };
 
 USTRUCT()
diff --git a/Engine/Plugins/Experimental/Chooser/Source/ChooserEditor/Private/ChooserTableEditor.cpp b/Engine/Plugins/Experimental/Chooser/Source/ChooserEditor/Private/ChooserTableEditor.cpp
--- a/Engine/Plugins/Experimental/Chooser/Source/ChooserEditor/Private/ChooserTableEditor.cpp
+++ b/Engine/Plugins/Experimental/Chooser/Source/ChooserEditor/Private/ChooserTableEditor.cpp
@@ -1451,11 +1451,21 @@ void FChooserTableEditor::AutoPopulateSelection()
 {
 	if (UChooserTable* Chooser = GetChooser())
 	{
-		const FScopedTransaction Transaction(LOCTEXT("Auto Populate Chooser", "Auto Populate All"));
+		const FScopedTransaction Transaction(LOCTEXT("Auto Populate Selection", "Auto Populate Selection"));
 		Chooser->Modify();
-		for(UChooserRowDetails* RowDetails : SelectedRows)
+		if (SelectedColumn)
 		{
-			AutoPopulateRow(RowDetails->Row);
+			if (Chooser->ColumnsStructs.IsValidIndex(SelectedColumn->Column))
+			{
+				AutoPopulateColumn(Chooser->ColumnsStructs[SelectedColumn->Column].GetMutable<FChooserColumnBase>());
+			}
+		}
+		else
+		{
+			for(UChooserRowDetails* RowDetails : SelectedRows)
+			{
+				AutoPopulateRow(RowDetails->Row);
+			}
 		}
 	}
 }
diff --git a/Engine/Plugins/Experimental/Chooser/Source/ChooserEditor/Private/FloatDistanceColumnEditor.h b/Engine/Plugins/Experimental/Chooser/Source/ChooserEditor/Private/FloatDistanceColumnEditor.h
--- a/Engine/Plugins/Experimental/Chooser/Source/ChooserEditor/Private/FloatDistanceColumnEditor.h
+++ b/Engine/Plugins/Experimental/Chooser/Source/ChooserEditor/Private/FloatDistanceColumnEditor.h
@@ -3,8 +3,21 @@
 #pragma once
 
 #include "CoreMinimal.h"
+#include "FloatDistanceColumn.h"
+#include "FloatDistanceColumnEditor.generated.h"
 
 namespace UE::ChooserEditor
 {
 	void RegisterFloatDistanceWidgets();
-}
\ No newline at end of file
+}
+
+UCLASS(Blueprintable)
+class UFloatAutoPopulatorBlueprint : public UFloatAutoPopulator
+{
+	GENERATED_BODY()
+public:
+	UFUNCTION(BlueprintImplementableEvent)
+	void AutoPopulate(UObject* Object, bool& Success, float& Value);
+	
+	virtual void NativeAutoPopulate(UObject* InObject, bool& OutSuccess, float& OutValue) { AutoPopulate(InObject, OutSuccess, OutValue); }
+};
