fix deprecated version number on ChooserIndexArray

diff --git a/Engine/Plugins/Experimental/Chooser/Source/Chooser/Public/ChooserIndexArray.h b/Engine/Plugins/Experimental/Chooser/Source/Chooser/Public/ChooserIndexArray.h
--- a/Engine/Plugins/Experimental/Chooser/Source/Chooser/Public/ChooserIndexArray.h
+++ b/Engine/Plugins/Experimental/Chooser/Source/Chooser/Public/ChooserIndexArray.h
@@ -9,7 +9,7 @@ class FChooserIndexArray
 public:
 	struct FIndexData
 	{
-		UE_DEPRECATED(5.6,"This is a fallback conversion operator for Chooser Column implementations which expect index data to be a uint32.  Please convert Chooser Column implementations to use FIndexData before the next release")
+		UE_DEPRECATED(5.5,"This is a fallback conversion operator for Chooser Column implementations which expect index data to be a uint32.  Please convert Chooser Column implementations to use FIndexData before the next release")
 		FIndexData(uint32 InIndex) : Index(InIndex), Cost(0) { }
 		
 		FIndexData(uint32 InIndex, float InCost) : Index(InIndex), Cost(InCost) {}
