Fix non-deterministic shutdown crash in proxy table editor module

diff --git a/Engine/Plugins/Experimental/Chooser/Source/ProxyTableEditor/Private/ProxyTableEditorModule.cpp b/Engine/Plugins/Experimental/Chooser/Source/ProxyTableEditor/Private/ProxyTableEditorModule.cpp
--- a/Engine/Plugins/Experimental/Chooser/Source/ProxyTableEditor/Private/ProxyTableEditorModule.cpp
+++ b/Engine/Plugins/Experimental/Chooser/Source/ProxyTableEditor/Private/ProxyTableEditorModule.cpp
@@ -26,8 +26,11 @@ void FModule::StartupModule()
 void FModule::ShutdownModule()
 {
 	FProxyTableEditorCommands::Unregister();
-	FPropertyEditorModule& PropertyModule = FModuleManager::LoadModuleChecked<FPropertyEditorModule>("PropertyEditor");
-	PropertyModule.UnregisterCustomPropertyTypeLayout(FProxyStructOutput::StaticStruct()->GetFName());
+	
+	if (FPropertyEditorModule* PropertyModule = FModuleManager::GetModulePtr<FPropertyEditorModule>("PropertyEditor"))
+	{
+		PropertyModule->UnregisterCustomPropertyTypeLayout(FProxyStructOutput::StaticStruct()->GetFName());
+	}
 	
 	FProxyTableEditorStyle::Shutdown();
 }
