Object class chooser column

diff --git a/Engine/Plugins/Experimental/Chooser/Source/Chooser/Internal/ObjectClassColumn.h b/Engine/Plugins/Experimental/Chooser/Source/Chooser/Internal/ObjectClassColumn.h
new file mode 100644
--- /dev/null
+++ b/Engine/Plugins/Experimental/Chooser/Source/Chooser/Internal/ObjectClassColumn.h
@@ -0,0 +1,98 @@
+// Copyright Epic Games, Inc. All Rights Reserved.
+#pragma once
+
+#include "ChooserPropertyAccess.h"
+#include "CoreMinimal.h"
+#include "IChooserColumn.h"
+#include "IChooserParameterObject.h"
+#include "InstancedStruct.h"
+#include "Serialization/MemoryReader.h"
+#include "ObjectColumn.h"
+#include "ObjectClassColumn.generated.h"
+
+UENUM()
+enum class EObjectClassColumnCellValueComparison
+{
+	Equal,
+	NotEqual,
+	SubClassOf,
+	NotSubClassOf,
+	Any,
+};
+
+USTRUCT()
+struct FChooserObjectClassRowData
+{
+	GENERATED_BODY()
+
+	UPROPERTY(EditAnywhere, Category = Runtime)
+	EObjectClassColumnCellValueComparison Comparison = EObjectClassColumnCellValueComparison::SubClassOf;
+
+	UPROPERTY(EditAnywhere, Category = "Runtime")
+	TObjectPtr<UClass> Value;
+
+	bool Evaluate(const UObject* LeftHandSide) const;
+};
+
+USTRUCT()
+struct CHOOSER_API FObjectClassColumn : public FChooserColumnBase
+{
+	GENERATED_BODY()
+
+	FObjectClassColumn();
+
+	UPROPERTY(EditAnywhere, NoClear, Meta = (ExcludeBaseStruct, BaseStruct = "/Script/Chooser.ChooserParameterObjectBase"), Category = "Data")
+	FInstancedStruct InputValue;
+
+#if WITH_EDITORONLY_DATA
+	UPROPERTY(EditAnywhere, Category = "Data")
+	FChooserObjectClassRowData DefaultRowValue;
+#endif
+
+	UPROPERTY(EditAnywhere, Category = "Data")
+	// array of results (cells for this column for each row in the table)
+	// should match the length of the Results array
+	TArray<FChooserObjectClassRowData> RowValues;
+
+	virtual void Filter(FChooserEvaluationContext& Context, const FChooserIndexArray& IndexListIn, FChooserIndexArray& IndexListOut) const override;
+
+#if WITH_EDITOR
+	mutable FSoftObjectPath TestValue;
+	virtual bool EditorTestFilter(int32 RowIndex) const override
+	{
+		if (RowValues.IsValidIndex(RowIndex))
+		{
+			if (UObject* Object = TestValue.ResolveObject())
+			{
+				return RowValues[RowIndex].Evaluate(Object);
+			}
+		}
+		return false;
+	}
+	
+	virtual void SetTestValue(TArrayView<const uint8> Value) override
+	{
+		FMemoryReaderView Reader(Value);
+		FString Path;
+		Reader << Path;
+		TestValue.SetPath(Path);
+	}
+
+	virtual void AddToDetails(FInstancedPropertyBag& PropertyBag, int32 ColumnIndex, int32 RowIndex) override;
+	virtual void SetFromDetails(FInstancedPropertyBag& PropertyBag, int32 ColumnIndex, int32 RowIndex) override;
+#endif
+
+	CHOOSER_COLUMN_BOILERPLATE(FChooserParameterObjectBase);
+
+#if WITH_EDITOR
+	virtual void PostLoad() override
+	{
+		Super::PostLoad();
+
+		if (InputValue.IsValid())
+		{
+			InputValue.GetMutable<FChooserParameterBase>().PostLoad();
+		}
+	}
+#endif
+};
diff --git a/Engine/Plugins/Experimental/Chooser/Source/Chooser/Internal/ObjectColumn.h b/Engine/Plugins/Experimental/Chooser/Source/Chooser/Internal/ObjectColumn.h
--- a/Engine/Plugins/Experimental/Chooser/Source/Chooser/Internal/ObjectColumn.h
+++ b/Engine/Plugins/Experimental/Chooser/Source/Chooser/Internal/ObjectColumn.h
@@ -20,8 +20,8 @@ struct CHOOSER_API FObjectContextProperty : public FChooserParameterObjectBase
 	FChooserObjectPropertyBinding Binding;
 
 	virtual bool GetValue(FChooserEvaluationContext& Context, FSoftObjectPath& OutResult) const override;
-	virtual bool GetValue(FChooserEvaluationContext& Context, UObject* OutResult) const override;
-	virtual bool SetValue(FChooserEvaluationContext& Context, UObject* OutResult) const override;
+	virtual bool GetValue(FChooserEvaluationContext& Context, UObject*& OutResult) const override;
+	virtual bool SetValue(FChooserEvaluationContext& Context, UObject* InValue) const override;
 
 	CHOOSER_PARAMETER_BOILERPLATE();
 
diff --git a/Engine/Plugins/Experimental/Chooser/Source/Chooser/Private/ObjectClassColumn.cpp b/Engine/Plugins/Experimental/Chooser/Source/Chooser/Private/ObjectClassColumn.cpp
new file mode 100644
--- /dev/null
+++ b/Engine/Plugins/Experimental/Chooser/Source/Chooser/Private/ObjectClassColumn.cpp
@@ -0,0 +1,99 @@
+// Copyright Epic Games, Inc. All Rights Reserved.
+#include "ObjectClassColumn.h"
+#include "ChooserIndexArray.h"
+#include "ChooserPropertyAccess.h"
+#include "ChooserTrace.h"
+
+#if WITH_EDITOR
+#include "IPropertyAccessEditor.h"
+#include "PropertyBag.h"
+#endif
+
+FObjectClassColumn::FObjectClassColumn()
+{
+#if WITH_EDITOR
+	InputValue.InitializeAs(FObjectContextProperty::StaticStruct());
+#endif
+}
+
+bool FChooserObjectClassRowData::Evaluate(const UObject* LeftHandSide) const
+{
+	if (Value && LeftHandSide)
+	{
+		switch (Comparison)
+		{
+			case EObjectClassColumnCellValueComparison::Equal:
+				return LeftHandSide->GetClass() == Value.Get();
+				
+			case EObjectClassColumnCellValueComparison::NotEqual:
+				return LeftHandSide->GetClass() != Value.Get();
+                				
+			case EObjectClassColumnCellValueComparison::SubClassOf:
+				return LeftHandSide->GetClass()->IsChildOf(Value.Get());
+				
+			case EObjectClassColumnCellValueComparison::NotSubClassOf:
+				return !LeftHandSide->GetClass()->IsChildOf(Value.Get());
+
+			default:
+				return false;
+		}
+	}
+	// always pass if the class was not set
+	return true;
+}
+
+void FObjectClassColumn::Filter(FChooserEvaluationContext& Context, const FChooserIndexArray& IndexListIn, FChooserIndexArray& IndexListOut) const
+{
+	UObject* Result = nullptr;
+	if (InputValue.IsValid() &&
+		InputValue.Get<FChooserParameterObjectBase>().GetValue(Context, Result))
+	{
+		TRACE_CHOOSER_VALUE(Context, ToCStr(InputValue.Get<FChooserParameterBase>().GetDebugName()), Result->GetPathName());	
+#if WITH_EDITOR
+		if (Context.DebuggingInfo.bCurrentDebugTarget)
+		{
+			TestValue = Result;
+		}
+#endif
+		
+		for (const FChooserIndexArray::FIndexData& IndexData : IndexListIn)
+		{
+			if (RowValues.IsValidIndex(IndexData.Index))
+			{
+				const FChooserObjectClassRowData& RowValue = RowValues[IndexData.Index];
+				if (RowValue.Evaluate(Result))
+				{
+					IndexListOut.Push(IndexData);
+				}
+			}
+		}
+	}
+	else
+	{
+		// passthrough fallback (behaves better during live editing)
+		IndexListOut = IndexListIn;
+	}
+}
+#if WITH_EDITOR
+	void FObjectClassColumn::AddToDetails(FInstancedPropertyBag& PropertyBag, int32 ColumnIndex, int32 RowIndex)
+	{
+		FText DisplayName;
+		InputValue.Get<FChooserParameterObjectBase>().GetDisplayName(DisplayName);
+		FName PropertyName("RowData",ColumnIndex);
+		FPropertyBagPropertyDesc PropertyDesc(PropertyName,  EPropertyBagPropertyType::Struct, FChooserObjectClassRowData::StaticStruct());
+		PropertyDesc.MetaData.Add(FPropertyBagPropertyDescMetaData("DisplayName", DisplayName.ToString()));
+		PropertyBag.AddProperties({PropertyDesc});
+		PropertyBag.SetValueStruct(PropertyName, RowValues[RowIndex]);
+	}
+
+	void FObjectClassColumn::SetFromDetails(FInstancedPropertyBag& PropertyBag, int32 ColumnIndex, int32 RowIndex)
+	{
+		FName PropertyName("RowData", ColumnIndex);
+		
+		TValueOrError<FStructView, EPropertyBagResult> Result = PropertyBag.GetValueStruct(PropertyName, FChooserObjectClassRowData::StaticStruct());
+		if (FStructView* StructView = Result.TryGetValue())
+		{
+			RowValues[RowIndex] = StructView->Get<FChooserObjectClassRowData>();
+		}
+	}
+#endif
\ No newline at end of file
diff --git a/Engine/Plugins/Experimental/Chooser/Source/Chooser/Private/ObjectColumn.cpp b/Engine/Plugins/Experimental/Chooser/Source/Chooser/Private/ObjectColumn.cpp
--- a/Engine/Plugins/Experimental/Chooser/Source/Chooser/Private/ObjectColumn.cpp
+++ b/Engine/Plugins/Experimental/Chooser/Source/Chooser/Private/ObjectColumn.cpp
@@ -34,7 +34,7 @@ bool FObjectContextProperty::GetValue(FChooserEvaluationContext& Context, FSoftO
 	return false;
 }
 
-bool FObjectContextProperty::GetValue(FChooserEvaluationContext& Context, UObject* OutResult) const
+bool FObjectContextProperty::GetValue(FChooserEvaluationContext& Context, UObject*& OutResult) const
 {
 	UE::Chooser::FResolvedPropertyChainResult Result;
 	if (UE::Chooser::ResolvePropertyChain(Context, Binding, Result))
diff --git a/Engine/Plugins/Experimental/Chooser/Source/Chooser/Public/IChooserParameterObject.h b/Engine/Plugins/Experimental/Chooser/Source/Chooser/Public/IChooserParameterObject.h
--- a/Engine/Plugins/Experimental/Chooser/Source/Chooser/Public/IChooserParameterObject.h
+++ b/Engine/Plugins/Experimental/Chooser/Source/Chooser/Public/IChooserParameterObject.h
@@ -14,7 +14,7 @@ struct FChooserParameterObjectBase : public FChooserParameterBase
 	GENERATED_BODY()
 
 	virtual bool GetValue(FChooserEvaluationContext& Context, FSoftObjectPath& OutResult) const { return false; }
-	virtual bool GetValue(FChooserEvaluationContext& Context, UObject* OutResult) const { return false; }
+	virtual bool GetValue(FChooserEvaluationContext& Context, UObject*& OutResult) const { return false; }
 	virtual bool SetValue(FChooserEvaluationContext& Context, UObject* OutResult) const { return false; }
 
 #if WITH_EDITOR
diff --git a/Engine/Plugins/Experimental/Chooser/Source/ChooserEditor/Private/ChooserEditorModule.cpp b/Engine/Plugins/Experimental/Chooser/Source/ChooserEditor/Private/ChooserEditorModule.cpp
--- a/Engine/Plugins/Experimental/Chooser/Source/ChooserEditor/Private/ChooserEditorModule.cpp
+++ b/Engine/Plugins/Experimental/Chooser/Source/ChooserEditor/Private/ChooserEditorModule.cpp
@@ -18,6 +18,7 @@
 #include "IAssetTools.h"
 #include "MultiEnumColumnEditor.h"
 #include "ObjectColumnEditor.h"
+#include "ObjectClassColumnEditor.h"
 #include "OutputFloatColumnEditor.h"
 #include "OutputObjectColumnEditor.h"
 #include "OutputStructColumnEditor.h"
@@ -46,6 +47,7 @@ void FModule::StartupModule()
 	RegisterEnumWidgets();
 	RegisterMultiEnumWidgets();
 	RegisterObjectWidgets();
+	RegisterObjectClassWidgets();
 	RegisterOutputObjectWidgets();
 	RegisterStructWidgets();
 	RegisterRandomizeWidgets();
diff --git a/Engine/Plugins/Experimental/Chooser/Source/ChooserEditor/Private/ObjectClassColumnEditor.cpp b/Engine/Plugins/Experimental/Chooser/Source/ChooserEditor/Private/ObjectClassColumnEditor.cpp
new file mode 100644
--- /dev/null
+++ b/Engine/Plugins/Experimental/Chooser/Source/ChooserEditor/Private/ObjectClassColumnEditor.cpp
@@ -0,0 +1,192 @@
+// Copyright Epic Games, Inc. All Rights Reserved.
+
+#include "ObjectClassColumnEditor.h"
+
+#include "ClassViewerFilter.h"
+#include "SPropertyAccessChainWidget.h"
+#include "GraphEditorSettings.h"
+#include "ObjectChooserWidgetFactories.h"
+#include "ObjectClassColumn.h"
+#include "PropertyCustomizationHelpers.h"
+#include "TransactionCommon.h"
+#include "Widgets/Images/SImage.h"
+#include "Widgets/Input/SButton.h"
+#include "Widgets/Layout/SBox.h"
+#include "SClassViewer.h"
+#include "ScopedTransaction.h"
+#include "Framework/Application/SlateApplication.h"
+#include "SEnumCombo.h"
+
+#define LOCTEXT_NAMESPACE "ObjectClassColumnEditor"
+
+namespace UE::ChooserEditor
+{
+	class FBaseClassFilter : public IClassViewerFilter
+	{
+	public:
+		const UClass* BaseClass = nullptr;
+
+		virtual bool IsClassAllowed(const FClassViewerInitializationOptions& InInitOptions, const UClass* InClass, TSharedRef< FClassViewerFilterFuncs > InFilterFuncs) override
+		{
+			return BaseClass == nullptr || InClass->IsChildOf(BaseClass);
+		};
+
+		virtual bool IsUnloadedClassAllowed(const FClassViewerInitializationOptions& InInitOptions, const TSharedRef< const IUnloadedBlueprintData > InUnloadedClassData, TSharedRef< FClassViewerFilterFuncs > InFilterFuncs) override
+		{
+			return BaseClass == nullptr || InUnloadedClassData->IsChildOf(BaseClass);
+		};
+	};
+
+	static TSharedRef<SWidget> CreateObjectClassColumnWidget(UChooserTable* Chooser, FChooserColumnBase* Column, int32 Row)
+	{
+		FObjectClassColumn* ObjectClassColumn = static_cast<FObjectClassColumn*>(Column);
+
+		if (Row == ColumnWidget_SpecialIndex_Fallback)
+		{
+			return SNullWidget::NullWidget;
+		}
+		else if (Row == ColumnWidget_SpecialIndex_Header)
+		{
+			// create column header widget
+			TSharedPtr<SWidget> InputValueWidget = nullptr;
+			if (FChooserParameterBase* InputValue = Column->GetInputValue())
+			{
+				InputValueWidget = FObjectChooserWidgetFactories::CreateWidget(false, Chooser, InputValue, Column->GetInputType(), Chooser->OutputObjectType);
+			}
+			
+			const FSlateBrush* ColumnIcon = FCoreStyle::Get().GetBrush("Icons.Filter");
+			
+			TSharedRef<SWidget> ColumnHeaderWidget = SNew(SHorizontalBox)
+				+ SHorizontalBox::Slot().AutoWidth()
+				[
+					SNew(SBorder)
+					.BorderBackgroundColor(FLinearColor(0,0,0,0))
+					.Content()
+					[
+						SNew(SImage).Image(ColumnIcon)
+					]
+				]
+				+ SHorizontalBox::Slot()
+				[
+					InputValueWidget ? InputValueWidget.ToSharedRef() : SNullWidget::NullWidget
+				];
+
+			if (Chooser->bEnableDebugTesting)
+			{
+				UClass* AllowedClass = UObject::StaticClass();
+				if (ObjectClassColumn->InputValue.IsValid())
+				{
+					const FChooserParameterObjectBase& Parameter = ObjectClassColumn->InputValue.Get<FChooserParameterObjectBase>();
+					AllowedClass = Parameter.GetAllowedClass();
+				}
+				
+				// create widget for test value object picker
+				TSharedRef<SWidget> ObjectPicker = SNew(SObjectPropertyEntryBox)
+                                                   			.ObjectPath_Lambda([ObjectClassColumn]() {
+                                                   				return ObjectClassColumn->TestValue.ToString();
+                                                   			})
+                                                   			.OnObjectChanged_Lambda([ObjectClassColumn](const FAssetData& AssetData) {
+                                                   				ObjectClassColumn->TestValue = AssetData.ToSoftObjectPath();
+                                                   			})
+															.AllowedClass(AllowedClass)
+                                                   			.DisplayUseSelected(false)
+                                                   			.DisplayBrowse(false)
+                                                   			.DisplayThumbnail(false);
+				
+				ObjectPicker->SetEnabled(TAttribute<bool>::CreateLambda([Chooser]() { return !Chooser->HasDebugTarget(); }));
+				
+				ColumnHeaderWidget = SNew(SVerticalBox)
+				+ SVerticalBox::Slot()
+				[
+					ColumnHeaderWidget
+				]
+				+ SVerticalBox::Slot()
+				[
+					ObjectPicker
+				];
+			}
+			
+			return ColumnHeaderWidget;
+		}
+
+		// create widget for cell
+		return SNew(SHorizontalBox)
+		+ SHorizontalBox::Slot()
+		[
+			SNew(SEnumComboBox, StaticEnum<EObjectClassColumnCellValueComparison>())
+				.CurrentValue_Lambda([ObjectClassColumn, Row]()
+				{
+					if (ObjectClassColumn->RowValues.IsValidIndex(Row))
+					{
+						return static_cast<int32>(ObjectClassColumn->RowValues[Row].Comparison);
+					}
+					else
+					{
+						return 0;
+					}
+				})
+				.OnEnumSelectionChanged_Lambda([ObjectClassColumn, Chooser, Row](int32 NewValue, ESelectInfo::Type SelectionType)
+				{
+					if (ObjectClassColumn->RowValues.IsValidIndex(Row))
+					{
+						FScopedTransaction ScopedTransaction(LOCTEXT("Change Comparison Type", "Change Comparison Type"));
+						Chooser->Modify();
+						ObjectClassColumn->RowValues[Row].Comparison = static_cast<EObjectClassColumnCellValueComparison>(NewValue);
+					}
+				})
+		]
+		+ SHorizontalBox::Slot()
+		[
+			SNew(SComboButton)
+				.OnGetMenuContent_Lambda([ObjectClassColumn, Chooser, Row]()
+				{
+					FClassViewerModule& ClassViewerModule = FModuleManager::LoadModuleChecked<FClassViewerModule>("ClassViewer");
+					FClassViewerInitializationOptions Options;
+					Options.Mode = EClassViewerMode::ClassPicker;
+
+					TSharedRef<FBaseClassFilter> BaseClassFilter = MakeShared<FBaseClassFilter>();
+					if (ObjectClassColumn->InputValue.IsValid())
+					{
+						const FChooserParameterObjectBase& Parameter = ObjectClassColumn->InputValue.Get<FChooserParameterObjectBase>();
+						BaseClassFilter->BaseClass = Parameter.GetAllowedClass();
+					}
+					Options.ClassFilters.Add(BaseClassFilter);
+					
+					return ClassViewerModule.CreateClassViewer(
+						Options,
+						FOnClassPicked::CreateLambda([ObjectClassColumn, Chooser, Row](UClass* InNewClass)
+						{
+							if (ObjectClassColumn->RowValues.IsValidIndex(Row))
+							{
+								FScopedTransaction ScopedTransaction(LOCTEXT("Change Class", "Change Class"));
+								Chooser->Modify();
+								ObjectClassColumn->RowValues[Row].Value = InNewClass;
+							}
+							FSlateApplication::Get().DismissAllMenus();
+						}));
+				})
+				.ButtonContent()
+				[
+					SNew(STextBlock).Text_Lambda([ObjectClassColumn, Row]()
+					{
+						if (ObjectClassColumn->RowValues.IsValidIndex(Row))
+						{
+							if (ObjectClassColumn->RowValues[Row].Value)
+							{
+								return FText::FromString(ObjectClassColumn->RowValues[Row].Value->GetName());
+							}
+						}
+						return LOCTEXT("None","None");
+					})
+				]
+		];
+	}
+
+	void RegisterObjectClassWidgets()
+	{
+		FObjectChooserWidgetFactories::RegisterColumnWidgetCreator(FObjectClassColumn::StaticStruct(), CreateObjectClassColumnWidget);
+	}
+
+} // namespace UE::ChooserEditor
+
+#undef LOCTEXT_NAMESPACE
diff --git a/Engine/Plugins/Experimental/Chooser/Source/ChooserEditor/Private/ObjectClassColumnEditor.h b/Engine/Plugins/Experimental/Chooser/Source/ChooserEditor/Private/ObjectClassColumnEditor.h
new file mode 100644
--- /dev/null
+++ b/Engine/Plugins/Experimental/Chooser/Source/ChooserEditor/Private/ObjectClassColumnEditor.h
@@ -0,0 +1,10 @@
+// Copyright Epic Games, Inc. All Rights Reserved.
+
+#pragma once
+
+#include "CoreMinimal.h"
+
+namespace UE::ChooserEditor
+{
+	void RegisterObjectClassWidgets();
+}
