Fix chooser editor crash when dragging assets into a table with an output struct column

diff --git a/Engine/Plugins/Experimental/Chooser/Source/Chooser/Public/IChooserColumn.h b/Engine/Plugins/Experimental/Chooser/Source/Chooser/Public/IChooserColumn.h
--- a/Engine/Plugins/Experimental/Chooser/Source/Chooser/Public/IChooserColumn.h
+++ b/Engine/Plugins/Experimental/Chooser/Source/Chooser/Public/IChooserColumn.h
@@ -105,10 +105,9 @@ public:
 	}\
 	virtual void InsertRows(int Index, int Count) override\
 	{\
-		RowValuesProperty.InsertUninitialized(Index, Count);\
 		for (int i=0;i<Count;i++)\
 		{\
-			RowValuesProperty[Index + i] = DefaultRowValue;\
+			RowValuesProperty.Insert(DefaultRowValue, Index);\
 		}\
 	}\
 	virtual void DeleteRows(const TArray<uint32> & RowIndices )\
