Move Chooser Uncooked module to load in PreDefault

diff --git a/Engine/Plugins/Experimental/Chooser/Chooser.uplugin b/Engine/Plugins/Experimental/Chooser/Chooser.uplugin
--- a/Engine/Plugins/Experimental/Chooser/Chooser.uplugin
+++ b/Engine/Plugins/Experimental/Chooser/Chooser.uplugin
@@ -43,7 +43,7 @@
 		{
 			"Name": "ChooserUncooked",
 			"Type": "UncookedOnly",
-			"LoadingPhase": "Default"
+			"LoadingPhase": "PreDefault"
 		},
 		{
 			"Name" : "ChooserEditor",
