Fix NoMin/NoMax float range behavior in Choosers

diff --git a/Engine/Plugins/Experimental/Chooser/Source/Chooser/Private/FloatRangeColumn.cpp b/Engine/Plugins/Experimental/Chooser/Source/Chooser/Private/FloatRangeColumn.cpp
--- a/Engine/Plugins/Experimental/Chooser/Source/Chooser/Private/FloatRangeColumn.cpp
+++ b/Engine/Plugins/Experimental/Chooser/Source/Chooser/Private/FloatRangeColumn.cpp
@@ -77,12 +77,12 @@ void FFloatRangeColumn::Filter(FChooserEvaluationContext& Context, const FChoose
 					const FChooserFloatRangeRowData& RowValue = RowValues[IndexData.Index];
 					if (RowValue.Max < RowValue.Min) // eg for an angle range from  135 to -135  (135 to 180 or -180 to -135)
 					{
-						if (Result >= RowValue.Min || Result <= RowValue.Max)
+						if ((RowValue.bNoMin || Result >= RowValue.Min) || (RowValue.bNoMax || Result <= RowValue.Max))
 						{
 							IndexListOut.Push(IndexData);
 						}
 					}
-					else if (Result >= RowValue.Min && Result <= RowValue.Max)
+					else if ((RowValue.bNoMin || Result >= RowValue.Min) && (RowValue.bNoMax || Result <= RowValue.Max))
 					{
 						IndexListOut.Push(IndexData);
 					}
@@ -96,7 +96,7 @@ void FFloatRangeColumn::Filter(FChooserEvaluationContext& Context, const FChoose
 				if (RowValues.Num() > static_cast<int>(IndexData.Index))
 				{
 					const FChooserFloatRangeRowData& RowValue = RowValues[IndexData.Index];
-					if (Result >= RowValue.Min && Result <= RowValue.Max)
+					if ((RowValue.bNoMin || Result >= RowValue.Min) && (RowValue.bNoMax || Result <= RowValue.Max))
 					{
 						IndexListOut.Push(IndexData);
 					}
