During shutdown avoid trying to call into other modules that may be unloaded

diff --git a/Engine/Plugins/Experimental/Chooser/Source/ProxyTableEditor/Private/ProxyTableEditorModule.cpp b/Engine/Plugins/Experimental/Chooser/Source/ProxyTableEditor/Private/ProxyTableEditorModule.cpp
--- a/Engine/Plugins/Experimental/Chooser/Source/ProxyTableEditor/Private/ProxyTableEditorModule.cpp
+++ b/Engine/Plugins/Experimental/Chooser/Source/ProxyTableEditor/Private/ProxyTableEditorModule.cpp
@@ -27,12 +27,16 @@ void FModule::StartupModule()
 void FModule::ShutdownModule()
 {
 	FProxyTableEditorCommands::Unregister();
-	
-	if (FPropertyEditorModule* PropertyModule = FModuleManager::GetModulePtr<FPropertyEditorModule>("PropertyEditor"))
+
+	if (UObjectInitialized() && !IsEngineExitRequested())
 	{
-		PropertyModule->UnregisterCustomPropertyTypeLayout(FProxyStructOutput::StaticStruct()->GetFName());
+
+		if (FPropertyEditorModule* PropertyModule = FModuleManager::GetModulePtr<FPropertyEditorModule>("PropertyEditor"))
+		{
+			PropertyModule->UnregisterCustomPropertyTypeLayout(FProxyStructOutput::StaticStruct()->GetFName());
+		}
 	}
-	
+
 	FProxyTableEditorStyle::Shutdown();
 }
 
@@ -40,4 +44,4 @@ void FModule::ShutdownModule()
 
 IMPLEMENT_MODULE(UE::ProxyTableEditor::FModule, ProxyTableEditor);
 
-#undef LOCTEXT_NAMESPACE
\ No newline at end of file
+#undef LOCTEXT_NAMESPACE
