Fix manual testing of nested choosers

diff --git a/Engine/Plugins/Experimental/Chooser/Source/Chooser/Internal/Chooser.h b/Engine/Plugins/Experimental/Chooser/Source/Chooser/Internal/Chooser.h
--- a/Engine/Plugins/Experimental/Chooser/Source/Chooser/Internal/Chooser.h
+++ b/Engine/Plugins/Experimental/Chooser/Source/Chooser/Internal/Chooser.h
@@ -65,14 +65,35 @@ public:
 	void IterateRecentContextObjects(TFunction<void(const FString&)> Callback) const;
 	void UpdateDebugging(FChooserEvaluationContext& Context) const;
 	
-	// enable display of which cells pass/fail based on current TestValue for each column
-	bool bEnableDebugTesting = false;
-	mutable bool bDebugTestValuesValid = false;
-
 	static const FName PropertyNamesTag;
 	static const FString PropertyTagDelimiter;
 
-private: 
+		
+	void SetEnableDebugTesting(bool bValue)
+	{
+		GetContextOwner()->bEnableDebugTesting = bValue;
+	}
+	
+	bool GetEnableDebugTesting() const
+	{
+		return GetContextOwner()->bEnableDebugTesting;
+	}
+	
+	void SetDebugTestValuesValid(bool bValue) const
+   	{
+   		GetContextOwner()->bDebugTestValuesValid = bValue;
+   	}
+	
+	bool GetDebugTestValuesValid() const
+	{
+		return GetContextOwner()->bDebugTestValuesValid;
+	}
+
+private:
+	// enable display of which cells pass/fail based on current TestValue for each column
+	bool bEnableDebugTesting = false;
+	mutable bool bDebugTestValuesValid = false;
+	
 	// caching the OutputObjectType and ContextObjectType so that on Undo, we can tell if we should fire the changed delegate
 	UClass* CachedPreviousOutputObjectType = nullptr;
 	EObjectChooserResultType CachedPreviousResultType = EObjectChooserResultType::ObjectResult;
diff --git a/Engine/Plugins/Experimental/Chooser/Source/ChooserEditor/Private/BoolColumnEditor.cpp b/Engine/Plugins/Experimental/Chooser/Source/ChooserEditor/Private/BoolColumnEditor.cpp
--- a/Engine/Plugins/Experimental/Chooser/Source/ChooserEditor/Private/BoolColumnEditor.cpp
+++ b/Engine/Plugins/Experimental/Chooser/Source/ChooserEditor/Private/BoolColumnEditor.cpp
@@ -51,7 +51,7 @@ TSharedRef<SWidget> CreateBoolColumnWidget(UChooserTable* Chooser, FChooserColum
 				InputValueWidget ? InputValueWidget.ToSharedRef() : SNullWidget::NullWidget
 			];
 	
-		if (Chooser->bEnableDebugTesting)
+		if (Chooser->GetEnableDebugTesting())
 		{
 			ColumnHeaderWidget = SNew(SVerticalBox)
 			+ SVerticalBox::Slot()
@@ -141,7 +141,7 @@ TSharedRef<SWidget> CreateOutputBoolColumnWidget(UChooserTable* Chooser, FChoose
 				InputValueWidget ? InputValueWidget.ToSharedRef() : SNullWidget::NullWidget
 			];
 	
-		if (Chooser->bEnableDebugTesting)
+		if (Chooser->GetEnableDebugTesting())
 		{
 			ColumnHeaderWidget = SNew(SVerticalBox)
 			+ SVerticalBox::Slot()
diff --git a/Engine/Plugins/Experimental/Chooser/Source/ChooserEditor/Private/ChooserTableEditor.cpp b/Engine/Plugins/Experimental/Chooser/Source/ChooserEditor/Private/ChooserTableEditor.cpp
--- a/Engine/Plugins/Experimental/Chooser/Source/ChooserEditor/Private/ChooserTableEditor.cpp
+++ b/Engine/Plugins/Experimental/Chooser/Source/ChooserEditor/Private/ChooserTableEditor.cpp
@@ -142,10 +142,10 @@ void FChooserTableEditor::MakeDebugTargetMenu(UToolMenu* InToolMenu)
 					{
 						UChooserTable* Chooser = GetRootChooser();
 						Chooser->ResetDebugTarget();
-						if (Chooser->bEnableDebugTesting)
+						if (Chooser->GetEnableDebugTesting())
 						{
-							Chooser->bEnableDebugTesting = false;
-							Chooser->bDebugTestValuesValid = false;
+							Chooser->SetEnableDebugTesting(false);
+							Chooser->SetDebugTestValuesValid(false);
 							UpdateTableColumns();
 						}
 					}),
@@ -165,10 +165,10 @@ void FChooserTableEditor::MakeDebugTargetMenu(UToolMenu* InToolMenu)
 					{
 						UChooserTable* Chooser = GetRootChooser();
 						Chooser->ResetDebugTarget();
-						if (!Chooser->bEnableDebugTesting)
+						if (!Chooser->GetEnableDebugTesting())
 						{
-							Chooser->bEnableDebugTesting = true;
-							Chooser->bDebugTestValuesValid = true;
+							Chooser->SetEnableDebugTesting(true);
+							Chooser->SetDebugTestValuesValid(true);
 							UpdateTableColumns();
 						}
 					}),
@@ -192,10 +192,10 @@ void FChooserTableEditor::MakeDebugTargetMenu(UToolMenu* InToolMenu)
 								{
 									UChooserTable* Chooser = GetRootChooser();
 									Chooser->SetDebugTarget(ObjectName);
-									Chooser->bDebugTestValuesValid = false;
-									if (!Chooser->bEnableDebugTesting)
+									Chooser->SetDebugTestValuesValid(false);
+									if (!Chooser->GetEnableDebugTesting())
 									{
-										Chooser->bEnableDebugTesting = true;
+										Chooser->SetEnableDebugTesting(true);
 										UpdateTableColumns();
 									}
 								}),
@@ -262,7 +262,7 @@ void FChooserTableEditor::RegisterToolbar()
 							}
 							else
 							{
-								return Chooser->bEnableDebugTesting ? LOCTEXT("Manual Testing", "Manual Testing") : LOCTEXT("Debug Target", "Debug Target");
+								return Chooser->GetEnableDebugTesting() ? LOCTEXT("Manual Testing", "Manual Testing") : LOCTEXT("Debug Target", "Debug Target");
 							}
 						}),
 						LOCTEXT("Debug Target Tooltip", "Select an object that has recently been the context object for this chooser to visualize the selection results")));
diff --git a/Engine/Plugins/Experimental/Chooser/Source/ChooserEditor/Private/ChooserTrack.cpp b/Engine/Plugins/Experimental/Chooser/Source/ChooserEditor/Private/ChooserTrack.cpp
--- a/Engine/Plugins/Experimental/Chooser/Source/ChooserEditor/Private/ChooserTrack.cpp
+++ b/Engine/Plugins/Experimental/Chooser/Source/ChooserEditor/Private/ChooserTrack.cpp
@@ -112,7 +112,7 @@ bool FChooserTrack::HandleDoubleClickInternal()
 		{
 			const FObjectInfo& OwnerObjectInfo  = GameplayProvider->GetObjectInfo(ObjectId);
 			ChooserTable->SetDebugTarget(OwnerObjectInfo.Name);
-			ChooserTable->bEnableDebugTesting = true;
+			ChooserTable->SetEnableDebugTesting(true);
 		}
 		
 
diff --git a/Engine/Plugins/Experimental/Chooser/Source/ChooserEditor/Private/EnumColumnEditor.cpp b/Engine/Plugins/Experimental/Chooser/Source/ChooserEditor/Private/EnumColumnEditor.cpp
--- a/Engine/Plugins/Experimental/Chooser/Source/ChooserEditor/Private/EnumColumnEditor.cpp
+++ b/Engine/Plugins/Experimental/Chooser/Source/ChooserEditor/Private/EnumColumnEditor.cpp
@@ -162,7 +162,7 @@ TSharedRef<SWidget> CreateEnumColumnWidget(UChooserTable* Chooser, FChooserColum
 				InputValueWidget ? InputValueWidget.ToSharedRef() : SNullWidget::NullWidget
 			];
 	
-		if (Chooser->bEnableDebugTesting)
+		if (Chooser->GetEnableDebugTesting())
 		{
 			ColumnHeaderWidget = SNew(SVerticalBox)
 			+ SVerticalBox::Slot()
@@ -274,7 +274,7 @@ TSharedRef<SWidget> CreateOutputEnumColumnWidget(UChooserTable* Chooser, FChoose
 				InputValueWidget ? InputValueWidget.ToSharedRef() : SNullWidget::NullWidget
 			];
 	
-		if (Chooser->bEnableDebugTesting)
+		if (Chooser->GetEnableDebugTesting())
 		{
 			ColumnHeaderWidget = SNew(SVerticalBox)
 			+ SVerticalBox::Slot()
diff --git a/Engine/Plugins/Experimental/Chooser/Source/ChooserEditor/Private/FloatDistanceColumnEditor.cpp b/Engine/Plugins/Experimental/Chooser/Source/ChooserEditor/Private/FloatDistanceColumnEditor.cpp
--- a/Engine/Plugins/Experimental/Chooser/Source/ChooserEditor/Private/FloatDistanceColumnEditor.cpp
+++ b/Engine/Plugins/Experimental/Chooser/Source/ChooserEditor/Private/FloatDistanceColumnEditor.cpp
@@ -49,7 +49,7 @@ TSharedRef<SWidget> CreateFloatDistanceColumnWidget(UChooserTable* Chooser, FCho
 				InputValueWidget ? InputValueWidget.ToSharedRef() : SNullWidget::NullWidget
 			];
 	
-		if (Chooser->bEnableDebugTesting)
+		if (Chooser->GetEnableDebugTesting())
 		{
 			ColumnHeaderWidget = SNew(SVerticalBox)
 			+ SVerticalBox::Slot()
diff --git a/Engine/Plugins/Experimental/Chooser/Source/ChooserEditor/Private/FloatRangeColumnEditor.cpp b/Engine/Plugins/Experimental/Chooser/Source/ChooserEditor/Private/FloatRangeColumnEditor.cpp
--- a/Engine/Plugins/Experimental/Chooser/Source/ChooserEditor/Private/FloatRangeColumnEditor.cpp
+++ b/Engine/Plugins/Experimental/Chooser/Source/ChooserEditor/Private/FloatRangeColumnEditor.cpp
@@ -62,7 +62,7 @@ TSharedRef<SWidget> CreateFloatRangeColumnWidget(UChooserTable* Chooser, FChoose
 				InputValueWidget ? InputValueWidget.ToSharedRef() : SNullWidget::NullWidget
 			];
 	
-		if (Chooser->bEnableDebugTesting)
+		if (Chooser->GetEnableDebugTesting())
 		{
 			ColumnHeaderWidget = SNew(SVerticalBox)
 			+ SVerticalBox::Slot()
diff --git a/Engine/Plugins/Experimental/Chooser/Source/ChooserEditor/Private/GameplayTagColumnEditor.cpp b/Engine/Plugins/Experimental/Chooser/Source/ChooserEditor/Private/GameplayTagColumnEditor.cpp
--- a/Engine/Plugins/Experimental/Chooser/Source/ChooserEditor/Private/GameplayTagColumnEditor.cpp
+++ b/Engine/Plugins/Experimental/Chooser/Source/ChooserEditor/Private/GameplayTagColumnEditor.cpp
@@ -48,7 +48,7 @@ TSharedRef<SWidget> CreateGameplayTagColumnWidget(UChooserTable* Chooser, FChoos
 				InputValueWidget ? InputValueWidget.ToSharedRef() : SNullWidget::NullWidget
 			];
 	
-		if (Chooser->bEnableDebugTesting)
+		if (Chooser->GetEnableDebugTesting())
 		{
 			ColumnHeaderWidget = SNew(SVerticalBox)
 			+ SVerticalBox::Slot()
diff --git a/Engine/Plugins/Experimental/Chooser/Source/ChooserEditor/Private/MultiEnumColumnEditor.cpp b/Engine/Plugins/Experimental/Chooser/Source/ChooserEditor/Private/MultiEnumColumnEditor.cpp
--- a/Engine/Plugins/Experimental/Chooser/Source/ChooserEditor/Private/MultiEnumColumnEditor.cpp
+++ b/Engine/Plugins/Experimental/Chooser/Source/ChooserEditor/Private/MultiEnumColumnEditor.cpp
@@ -163,7 +163,7 @@ TSharedRef<SWidget> CreateMultiEnumColumnWidget(UChooserTable* Chooser, FChooser
 				InputValueWidget ? InputValueWidget.ToSharedRef() : SNullWidget::NullWidget
 			];
 	
-		if (Chooser->bEnableDebugTesting)
+		if (Chooser->GetEnableDebugTesting())
 		{
 			ColumnHeaderWidget = SNew(SVerticalBox)
 			+ SVerticalBox::Slot()
diff --git a/Engine/Plugins/Experimental/Chooser/Source/ChooserEditor/Private/ObjectClassColumnEditor.cpp b/Engine/Plugins/Experimental/Chooser/Source/ChooserEditor/Private/ObjectClassColumnEditor.cpp
--- a/Engine/Plugins/Experimental/Chooser/Source/ChooserEditor/Private/ObjectClassColumnEditor.cpp
+++ b/Engine/Plugins/Experimental/Chooser/Source/ChooserEditor/Private/ObjectClassColumnEditor.cpp
@@ -71,7 +71,7 @@ namespace UE::ChooserEditor
 					InputValueWidget ? InputValueWidget.ToSharedRef() : SNullWidget::NullWidget
 				];
 
-			if (Chooser->bEnableDebugTesting)
+			if (Chooser->GetEnableDebugTesting())
 			{
 				UClass* AllowedClass = UObject::StaticClass();
 				if (ObjectClassColumn->InputValue.IsValid())
diff --git a/Engine/Plugins/Experimental/Chooser/Source/ChooserEditor/Private/ObjectColumnEditor.cpp b/Engine/Plugins/Experimental/Chooser/Source/ChooserEditor/Private/ObjectColumnEditor.cpp
--- a/Engine/Plugins/Experimental/Chooser/Source/ChooserEditor/Private/ObjectColumnEditor.cpp
+++ b/Engine/Plugins/Experimental/Chooser/Source/ChooserEditor/Private/ObjectColumnEditor.cpp
@@ -125,7 +125,7 @@ namespace UE::ChooserEditor
 					InputValueWidget ? InputValueWidget.ToSharedRef() : SNullWidget::NullWidget
 				];
 		
-			if (Chooser->bEnableDebugTesting)
+			if (Chooser->GetEnableDebugTesting())
 			{
 				TSharedRef<SWidget> ObjectPicker = CreateObjectPicker(Chooser, ObjectColumn, Row);
 				ObjectPicker->SetEnabled(TAttribute<bool>::CreateLambda([Chooser]() { return !Chooser->HasDebugTarget(); }));
diff --git a/Engine/Plugins/Experimental/Chooser/Source/ChooserEditor/Private/OutputFloatColumnEditor.cpp b/Engine/Plugins/Experimental/Chooser/Source/ChooserEditor/Private/OutputFloatColumnEditor.cpp
--- a/Engine/Plugins/Experimental/Chooser/Source/ChooserEditor/Private/OutputFloatColumnEditor.cpp
+++ b/Engine/Plugins/Experimental/Chooser/Source/ChooserEditor/Private/OutputFloatColumnEditor.cpp
@@ -45,7 +45,7 @@ TSharedRef<SWidget> CreateOutputFloatColumnWidget(UChooserTable* Chooser, FChoos
 				InputValueWidget ? InputValueWidget.ToSharedRef() : SNullWidget::NullWidget
 			];
 	
-		if (Chooser->bEnableDebugTesting)
+		if (Chooser->GetEnableDebugTesting())
 		{
 			ColumnHeaderWidget = SNew(SVerticalBox)
 			+ SVerticalBox::Slot()
diff --git a/Engine/Plugins/Experimental/Chooser/Source/ChooserEditor/Private/RewindDebuggerChooser.cpp b/Engine/Plugins/Experimental/Chooser/Source/ChooserEditor/Private/RewindDebuggerChooser.cpp
--- a/Engine/Plugins/Experimental/Chooser/Source/ChooserEditor/Private/RewindDebuggerChooser.cpp
+++ b/Engine/Plugins/Experimental/Chooser/Source/ChooserEditor/Private/RewindDebuggerChooser.cpp
@@ -45,7 +45,7 @@ void FRewindDebuggerChooser::Update(float DeltaTime, IRewindDebugger* RewindDebu
 						if (ContextOwner->GetDebugTargetName() == ContextObjectInfo.Name)
 						{
 							Chooser->SetDebugSelectedRow(ChooserEvaluationData.SelectedIndex);
-							Chooser->bDebugTestValuesValid = true;
+							Chooser->SetDebugTestValuesValid(true);
 
 							ChooserProvider->ReadChooserValueTimeline(OwnerId, [StartTime, EndTime, Chooser, &ChooserEvaluationData](const FChooserProvider::ChooserValueTimeline& ValueTimeline)
 							{
diff --git a/Engine/Plugins/Experimental/Chooser/Source/ChooserEditor/Private/SChooserRowHandle.cpp b/Engine/Plugins/Experimental/Chooser/Source/ChooserEditor/Private/SChooserRowHandle.cpp
--- a/Engine/Plugins/Experimental/Chooser/Source/ChooserEditor/Private/SChooserRowHandle.cpp
+++ b/Engine/Plugins/Experimental/Chooser/Source/ChooserEditor/Private/SChooserRowHandle.cpp
@@ -46,7 +46,7 @@ namespace UE::ChooserEditor
 						SNew(SImage)
 						.Visibility_Lambda([this]()
 						{
-							return ChooserEditor->GetChooser()->bDebugTestValuesValid && RowIndex == ChooserEditor->GetChooser()->GetDebugSelectedRow() ? EVisibility::HitTestInvisible : EVisibility::Hidden;
+							return ChooserEditor->GetChooser()->GetDebugTestValuesValid() && RowIndex == ChooserEditor->GetChooser()->GetDebugSelectedRow() ? EVisibility::HitTestInvisible : EVisibility::Hidden;
 						})
 						.Image(FAppStyle::Get().GetBrush("Icons.ArrowRight"))
 					]
diff --git a/Engine/Plugins/Experimental/Chooser/Source/ChooserEditor/Private/SChooserTableRow.cpp b/Engine/Plugins/Experimental/Chooser/Source/ChooserEditor/Private/SChooserTableRow.cpp
--- a/Engine/Plugins/Experimental/Chooser/Source/ChooserEditor/Private/SChooserTableRow.cpp
+++ b/Engine/Plugins/Experimental/Chooser/Source/ChooserEditor/Private/SChooserTableRow.cpp
@@ -167,7 +167,7 @@ namespace UE::ChooserEditor
 							SNew(SColorBlock).Visibility(EVisibility::HitTestInvisible).Color_Lambda(
 									[this,Column]()
 									{
-										if (Chooser->bDebugTestValuesValid)
+										if (Chooser->GetDebugTestValuesValid())
 										{
 											if (Column->HasCosts())
 											{
