Change debug names for chooser debugging to include outer actor name

diff --git a/Engine/Plugins/Animation/GameplayInsights/Source/RewindDebugger/Private/RewindDebugger.cpp b/Engine/Plugins/Animation/GameplayInsights/Source/RewindDebugger/Private/RewindDebugger.cpp
--- a/Engine/Plugins/Animation/GameplayInsights/Source/RewindDebugger/Private/RewindDebugger.cpp
+++ b/Engine/Plugins/Animation/GameplayInsights/Source/RewindDebugger/Private/RewindDebugger.cpp
@@ -1002,7 +1002,7 @@ const TraceServices::IAnalysisSession* FRewindDebugger::GetAnalysisSession() con
 	return UnrealInsightsModule ? UnrealInsightsModule->GetAnalysisSession().Get() : nullptr;
 }
 
-const FObjectInfo* FRewindDebugger::FindOwningActorInfo(const IGameplayProvider* GameplayProvider, uint64 ObjectId)
+const FObjectInfo* FRewindDebugger::FindOwningActorInfo(const IGameplayProvider* GameplayProvider, uint64 ObjectId) const
 {
 	const FClassInfo* ActorClassInfo = GameplayProvider->FindClassInfo(*AActor::StaticClass()->GetPathName());
 	
diff --git a/Engine/Plugins/Animation/GameplayInsights/Source/RewindDebugger/Private/RewindDebugger.h b/Engine/Plugins/Animation/GameplayInsights/Source/RewindDebugger/Private/RewindDebugger.h
--- a/Engine/Plugins/Animation/GameplayInsights/Source/RewindDebugger/Private/RewindDebugger.h
+++ b/Engine/Plugins/Animation/GameplayInsights/Source/RewindDebugger/Private/RewindDebugger.h
@@ -148,7 +148,7 @@ public:
 	void SetIsDetailsPanelOpen(bool bIsOpen) { bIsDetailsPanelOpen = bIsOpen; }
 	bool IsDetailsPanelOpen(bool bIsOpen) { return bIsDetailsPanelOpen; }
 	
-	static const FObjectInfo* FindOwningActorInfo(const IGameplayProvider* GameplayProvider, uint64 ObjectId);
+	virtual const FObjectInfo* FindOwningActorInfo(const IGameplayProvider* GameplayProvider, uint64 ObjectId) const override;
 
 	TArrayView<RewindDebugger::FRewindDebuggerTrackType> GetTrackTypes() { return TrackTypes; };
 
diff --git a/Engine/Plugins/Animation/GameplayInsights/Source/RewindDebugger/Private/RewindDebuggerAnimation.cpp b/Engine/Plugins/Animation/GameplayInsights/Source/RewindDebugger/Private/RewindDebuggerAnimation.cpp
--- a/Engine/Plugins/Animation/GameplayInsights/Source/RewindDebugger/Private/RewindDebuggerAnimation.cpp
+++ b/Engine/Plugins/Animation/GameplayInsights/Source/RewindDebugger/Private/RewindDebuggerAnimation.cpp
@@ -196,7 +196,7 @@ FRewindDebuggerAnimation::FSpawnedMeshComponentInfo* FRewindDebuggerAnimation::S
 
 		MeshComponentInfo->Actor = World->SpawnActor<AActor>(ActorSpawnParameters);
 
-		if(const FObjectInfo* ActorInfo = FRewindDebugger::FindOwningActorInfo(GameplayProvider, ObjectId))
+		if(const FObjectInfo* ActorInfo = RewindDebugger->FindOwningActorInfo(GameplayProvider, ObjectId))
 		{
 			MeshComponentInfo->Actor->SetActorLabel(FString(TEXT("RewindDebugger: ") + FString(ActorInfo->Name)));
 		}
diff --git a/Engine/Plugins/Animation/GameplayInsights/Source/RewindDebugger/Private/RewindDebuggerStyle.cpp b/Engine/Plugins/Animation/GameplayInsights/Source/RewindDebugger/Private/RewindDebuggerStyle.cpp
--- a/Engine/Plugins/Animation/GameplayInsights/Source/RewindDebugger/Private/RewindDebuggerStyle.cpp
+++ b/Engine/Plugins/Animation/GameplayInsights/Source/RewindDebugger/Private/RewindDebuggerStyle.cpp
@@ -34,6 +34,7 @@ FRewindDebuggerStyle::FRewindDebuggerStyle() :
 	Set("RewindDebugger.AutoEject", new IMAGE_BRUSH("autoeject", {30.f, 24.f}));
 	Set("RewindDebugger.AutoRecord", new IMAGE_BRUSH("autorecord", {30.f, 24.f}));
 
+	Set("RewindDebugger.ConnectToSession", new IMAGE_BRUSH_SVG("Session_20", {20.f, 20.f}));
 	// Actor Picker
 	Set("RewindDebugger.SelectActor", new CORE_IMAGE_BRUSH("Editor/Slate/Icons/eyedropper_16px", Icon16x16));
 
diff --git a/Engine/Plugins/Experimental/Chooser/Source/Chooser/Private/Chooser.cpp b/Engine/Plugins/Experimental/Chooser/Source/Chooser/Private/Chooser.cpp
--- a/Engine/Plugins/Experimental/Chooser/Source/Chooser/Private/Chooser.cpp
+++ b/Engine/Plugins/Experimental/Chooser/Source/Chooser/Private/Chooser.cpp
@@ -313,9 +313,15 @@ void UChooserTable::UpdateDebugging(FChooserEvaluationContext& Context) const
 		{
 			if (UObject* ContextObject = ObjectParam->Object)
 			{
-				RecentContextObjects.Add(ContextObject->GetName());
+				FString DebugName = ContextObject->GetName();
+				if (UObject* Outer = ContextObject->GetTypedOuter(AActor::StaticClass()))
+				{
+					DebugName += " in " + Outer->GetName();					
+				}
+				
+				RecentContextObjects.Add(DebugName);
 
-				if (ContextObject->GetName() == ContextOwner->GetDebugTargetName())
+				if (DebugName == ContextOwner->GetDebugTargetName())
 				{
 					bDebugTestValuesValid = true;
 					Context.DebuggingInfo.bCurrentDebugTarget = true;
diff --git a/Engine/Plugins/Experimental/Chooser/Source/ChooserEditor/Private/ChooserTrack.cpp b/Engine/Plugins/Experimental/Chooser/Source/ChooserEditor/Private/ChooserTrack.cpp
--- a/Engine/Plugins/Experimental/Chooser/Source/ChooserEditor/Private/ChooserTrack.cpp
+++ b/Engine/Plugins/Experimental/Chooser/Source/ChooserEditor/Private/ChooserTrack.cpp
@@ -111,7 +111,13 @@ bool FChooserTrack::HandleDoubleClickInternal()
 		if (ChooserTable)
 		{
 			const FObjectInfo& OwnerObjectInfo  = GameplayProvider->GetObjectInfo(ObjectId);
-			ChooserTable->SetDebugTarget(OwnerObjectInfo.Name);
+			FString DebugName(OwnerObjectInfo.Name);
+			if (const FObjectInfo* ActorInfo = RewindDebugger->FindOwningActorInfo(GameplayProvider, ObjectId))
+			{
+				DebugName += " in " + FString(ActorInfo->Name);
+			}
+			
+			ChooserTable->SetDebugTarget(DebugName);
 			ChooserTable->SetEnableDebugTesting(true);
 		}
 		
diff --git a/Engine/Plugins/Experimental/Chooser/Source/ChooserEditor/Private/RewindDebuggerChooser.cpp b/Engine/Plugins/Experimental/Chooser/Source/ChooserEditor/Private/RewindDebuggerChooser.cpp
--- a/Engine/Plugins/Experimental/Chooser/Source/ChooserEditor/Private/RewindDebuggerChooser.cpp
+++ b/Engine/Plugins/Experimental/Chooser/Source/ChooserEditor/Private/RewindDebuggerChooser.cpp
@@ -28,21 +28,27 @@ void FRewindDebuggerChooser::Update(float DeltaTime, IRewindDebugger* RewindDebu
 			double StartTime = RewindDebugger->GetScrubTime();
 			double EndTime = StartTime + Frame.EndTime - Frame.StartTime;
 			
-			ChooserEvaluationTimeline.EnumerateEvents(StartTime, EndTime, [GameplayProvider, ChooserProvider, StartTime, EndTime, OwnerId](double InStartTime, double InEndTime, uint32 InDepth, const FChooserEvaluationData& ChooserEvaluationData)  
+			ChooserEvaluationTimeline.EnumerateEvents(StartTime, EndTime, [RewindDebugger, GameplayProvider, ChooserProvider, StartTime, EndTime, OwnerId](double InStartTime, double InEndTime, uint32 InDepth, const FChooserEvaluationData& ChooserEvaluationData)  
 			{
 				const FObjectInfo& ChooserInfo = GameplayProvider->GetObjectInfo(ChooserEvaluationData.ChooserId);
 				
 				if (UChooserTable* Chooser = FindObject<UChooserTable>(nullptr, ChooserInfo.PathName))
 				{
 					const FObjectInfo& ContextObjectInfo = GameplayProvider->GetObjectInfo(OwnerId);
+					FString DebugName = ContextObjectInfo.Name;
+
+					if (const FObjectInfo* ActorInfo = RewindDebugger->FindOwningActorInfo(GameplayProvider, OwnerId))
+					{
+						DebugName += " in " + FString(ActorInfo->Name);
+					}
 					
 					// add to recent context objects list, so that this object is selectable as a target in the chooser editor
-                 	Chooser->AddRecentContextObject(ContextObjectInfo.Name);
+                 	Chooser->AddRecentContextObject(DebugName);
 					
 					UChooserTable* ContextOwner = Chooser->GetContextOwner();
 					if (ContextOwner->HasDebugTarget())
 					{
-						if (ContextOwner->GetDebugTargetName() == ContextObjectInfo.Name)
+						if (ContextOwner->GetDebugTargetName() == DebugName)
 						{
 							Chooser->SetDebugSelectedRow(ChooserEvaluationData.SelectedIndex);
 							Chooser->SetDebugTestValuesValid(true);
diff --git a/Engine/Source/Editor/RewindDebuggerInterface/Public/IRewindDebugger.h b/Engine/Source/Editor/RewindDebuggerInterface/Public/IRewindDebugger.h
--- a/Engine/Source/Editor/RewindDebuggerInterface/Public/IRewindDebugger.h
+++ b/Engine/Source/Editor/RewindDebuggerInterface/Public/IRewindDebugger.h
@@ -101,6 +101,8 @@ public:
 
 	// opens the Rewind Debugger details panel tab
 	virtual void OpenDetailsPanel() = 0;
+	
+	virtual const struct FObjectInfo* FindOwningActorInfo(const class IGameplayProvider* GameplayProvider, uint64 ObjectId) const = 0;
 
 	// get the current IRewindDebugger instance
 	static IRewindDebugger* Instance();
