Fix crashes when changing chooser column input type

diff --git a/Engine/Plugins/Experimental/Chooser/Source/ChooserEditor/Private/ChooserTableEditor.cpp b/Engine/Plugins/Experimental/Chooser/Source/ChooserEditor/Private/ChooserTableEditor.cpp
--- a/Engine/Plugins/Experimental/Chooser/Source/ChooserEditor/Private/ChooserTableEditor.cpp
+++ b/Engine/Plugins/Experimental/Chooser/Source/ChooserEditor/Private/ChooserTableEditor.cpp
@@ -938,6 +938,13 @@ void FChooserTableEditor::UpdateTableColumns()
 							ChooserTable->Modify(true);
 							UpdateTableColumns();
 							UpdateTableRows();
+							
+							if (SelectedColumn && SelectedColumn->Column == ColumnIndex)
+							{
+								// if this column was selected, reselect to refresh the details widgets
+								SelectColumn(ChooserTable, ColumnIndex);
+							}
+                            							
 						}));
 				
 						Builder.AddWidget(Widget, FText());
