fix nested chooser duplication

diff --git a/Engine/Plugins/Experimental/Chooser/Source/ChooserEditor/Private/ChooserEditorWidgets.cpp b/Engine/Plugins/Experimental/Chooser/Source/ChooserEditor/Private/ChooserEditorWidgets.cpp
--- a/Engine/Plugins/Experimental/Chooser/Source/ChooserEditor/Private/ChooserEditorWidgets.cpp
+++ b/Engine/Plugins/Experimental/Chooser/Source/ChooserEditor/Private/ChooserEditorWidgets.cpp
@@ -182,7 +182,7 @@ TSharedRef<SWidget> CreateNestedChooserWidget(bool bReadOnly, UObject* Transacti
 												TransactionObject->Modify(true);
 											
 												FSlateApplication::Get().DismissAllMenus();
-												NestedChooser->Chooser = NewObject<UChooserTable>(TransactionObject->GetPackage(), UChooserTable::StaticClass(), FName(InText.ToString()));
+												NestedChooser->Chooser = NewObject<UChooserTable>(TransactionObject, UChooserTable::StaticClass(), FName(InText.ToString()));
 												NestedChooser->Chooser->ParentTable = Cast<UChooserTable>(TransactionObject);
 											});
 							
