Fix gameplay tag chooser column logic when referencing parent tags

diff --git a/Engine/Plugins/Experimental/Chooser/Source/Chooser/Private/GameplayTagColumn.cpp b/Engine/Plugins/Experimental/Chooser/Source/Chooser/Private/GameplayTagColumn.cpp
--- a/Engine/Plugins/Experimental/Chooser/Source/Chooser/Private/GameplayTagColumn.cpp
+++ b/Engine/Plugins/Experimental/Chooser/Source/Chooser/Private/GameplayTagColumn.cpp
@@ -33,14 +33,15 @@ bool FGameplayTagColumn::TestRow(int32 RowIndex, const FGameplayTagContainer& Va
 		{
 			if (TagMatchType == EGameplayContainerMatchType::All)
 			{
-				if (bMatchExact ? Value.HasAllExact(RowValues[RowIndex]) : Value.HasAll(RowValues[RowIndex]))
+				if (bMatchExact ? RowValues[RowIndex].HasAllExact(Value) : RowValues[RowIndex].HasAll(Value))
 				{
 					bPasses = true;
 				}
 			}
 			else
 			{
-				if (bMatchExact ? Value.HasAnyExact(RowValues[RowIndex]) : Value.HasAny(RowValues[RowIndex]))
+				if (bMatchExact ? RowValues[RowIndex].HasAnyExact(Value) : RowValues[RowIndex].HasAny(Value))
+					
 				{
 					bPasses = true;
 				}
