Fix broken outer object on nested choosers in existing content

diff --git a/Engine/Plugins/Experimental/Chooser/Source/Chooser/Private/Chooser.cpp b/Engine/Plugins/Experimental/Chooser/Source/Chooser/Private/Chooser.cpp
--- a/Engine/Plugins/Experimental/Chooser/Source/Chooser/Private/Chooser.cpp
+++ b/Engine/Plugins/Experimental/Chooser/Source/Chooser/Private/Chooser.cpp
@@ -37,6 +37,15 @@ void UChooserTable::PostLoad()
 #if WITH_EDITOR
 	CachedPreviousOutputObjectType = OutputObjectType;
 	CachedPreviousResultType = ResultType;
+
+	// fix for broken outer object on nested chooser tables
+	if (ParentTable)
+	{
+		if (GetOuter() != ParentTable)
+		{
+			Rename(nullptr, ParentTable);
+		}
+	}
 #endif
 
 	Compile();
