Add read only display in cells for Struct output  columns in choosers

diff --git a/Engine/Plugins/Experimental/Chooser/Source/ChooserEditor/Private/OutputStructColumnEditor.cpp b/Engine/Plugins/Experimental/Chooser/Source/ChooserEditor/Private/OutputStructColumnEditor.cpp
--- a/Engine/Plugins/Experimental/Chooser/Source/ChooserEditor/Private/OutputStructColumnEditor.cpp
+++ b/Engine/Plugins/Experimental/Chooser/Source/ChooserEditor/Private/OutputStructColumnEditor.cpp
@@ -4,9 +4,11 @@
 #include "OutputStructColumn.h"
 #include "SPropertyAccessChainWidget.h"
 #include "ObjectChooserWidgetFactories.h"
-#include "Widgets/SBoxPanel.h"
 #include "Widgets/Images/SImage.h"
+#include "Widgets/Layout/SBox.h"
+#include "Widgets/SBoxPanel.h"
 #include "Widgets/SNullWidget.h"
+#include "Widgets/Text/STextBlock.h"
 #include "GraphEditorSettings.h"
 
 #define LOCTEXT_NAMESPACE "StructOutputColumnEditor"
@@ -50,10 +52,44 @@ TSharedRef<SWidget> CreateOutputStructColumnWidget(UChooserTable* Chooser, FChoo
 	
 		return ColumnHeaderWidget;
 	}
+	else if (Row == ColumnWidget_SpecialIndex_Fallback)
+	{
+		return SNullWidget::NullWidget;
+	}
+
+	FOutputStructColumn* StructColumn = static_cast<FOutputStructColumn*>(Column);
 
-	return SNullWidget::NullWidget;
+	TAttribute<FText> StructValueAttribute = MakeAttributeLambda([StructColumn, Row]()
+		{
+			const FInstancedStruct& RowValue = StructColumn->RowValues[Row];
+
+			FString Value;
+			if (const UScriptStruct* ScriptStruct = RowValue.GetScriptStruct())
+			{
+				void* DefaultStructMemory = FMemory_Alloca_Aligned(ScriptStruct->GetStructureSize(), ScriptStruct->GetMinAlignment());
+				ScriptStruct->InitializeStruct(DefaultStructMemory);
+				ScriptStruct->ExportText(Value, RowValue.GetMemory(), DefaultStructMemory, nullptr, PPF_ExternalEditor, nullptr);
+				ScriptStruct->DestroyStruct(DefaultStructMemory);
+			}
+			else
+			{
+				Value = TEXT("()");
+			}
+			return FText::FromString(Value);
+		});
+
+	TSharedRef<STextBlock> TextBlock = SNew(STextBlock)
+		.Text(StructValueAttribute);
+	TextBlock->SetToolTipText(StructValueAttribute);
+
+	return SNew(SBox)
+		.HAlign(HAlign_Center)
+		.VAlign(VAlign_Center)
+		[
+			TextBlock
+		];
 }
-	
+
 TSharedRef<SWidget> CreateStructPropertyWidget(bool bReadOnly, UObject* TransactionObject, void* Value, UClass* ResultBaseClass, FChooserWidgetValueChanged ValueChanged)
 {
 	IHasContextClass* HasContextClass = Cast<IHasContextClass>(TransactionObject);
