Add fallback value support to struct output column widget

diff --git a/Engine/Plugins/Experimental/Chooser/Source/ChooserEditor/Private/OutputStructColumnEditor.cpp b/Engine/Plugins/Experimental/Chooser/Source/ChooserEditor/Private/OutputStructColumnEditor.cpp
--- a/Engine/Plugins/Experimental/Chooser/Source/ChooserEditor/Private/OutputStructColumnEditor.cpp
+++ b/Engine/Plugins/Experimental/Chooser/Source/ChooserEditor/Private/OutputStructColumnEditor.cpp
@@ -52,23 +52,28 @@ TSharedRef<SWidget> CreateOutputStructColumnWidget(UChooserTable* Chooser, FChoo
 	
 		return ColumnHeaderWidget;
 	}
-	else if (Row == ColumnWidget_SpecialIndex_Fallback)
-	{
-		return SNullWidget::NullWidget;
-	}
-
+	
 	FOutputStructColumn* StructColumn = static_cast<FOutputStructColumn*>(Column);
 
 	TAttribute<FText> StructValueAttribute = MakeAttributeLambda([StructColumn, Row]()
 		{
-			const FInstancedStruct& RowValue = StructColumn->RowValues[Row];
+			const FInstancedStruct* RowValue = nullptr;
+			if (Row == ColumnWidget_SpecialIndex_Fallback)
+			{
+				RowValue = &StructColumn->FallbackValue;
+			}
+			else
+			{
+				RowValue = &StructColumn->RowValues[Row];
+			}
+
 
 			FString Value;
-			if (const UScriptStruct* ScriptStruct = RowValue.GetScriptStruct())
+			if (const UScriptStruct* ScriptStruct = RowValue->GetScriptStruct())
 			{
 				void* DefaultStructMemory = FMemory_Alloca_Aligned(ScriptStruct->GetStructureSize(), ScriptStruct->GetMinAlignment());
 				ScriptStruct->InitializeStruct(DefaultStructMemory);
-				ScriptStruct->ExportText(Value, RowValue.GetMemory(), DefaultStructMemory, nullptr, PPF_ExternalEditor, nullptr);
+				ScriptStruct->ExportText(Value, RowValue->GetMemory(), DefaultStructMemory, nullptr, PPF_ExternalEditor, nullptr);
 				ScriptStruct->DestroyStruct(DefaultStructMemory);
 			}
 			else
