fix crash when moving a chooser table column which is currently selected

diff --git a/Engine/Plugins/Experimental/Chooser/Source/ChooserEditor/Private/ChooserTableEditor.cpp b/Engine/Plugins/Experimental/Chooser/Source/ChooserEditor/Private/ChooserTableEditor.cpp
--- a/Engine/Plugins/Experimental/Chooser/Source/ChooserEditor/Private/ChooserTableEditor.cpp
+++ b/Engine/Plugins/Experimental/Chooser/Source/ChooserEditor/Private/ChooserTableEditor.cpp
@@ -882,6 +882,10 @@ void FChooserTableEditor::UpdateTableColumns()
 								Chooser->Modify(true);
 								Chooser->ColumnsStructs.Swap(ColumnIndex, ColumnIndex - 1);
 								UpdateTableColumns();
+								if (SelectedColumn)
+								{
+									SelectColumn(Chooser, ColumnIndex - 1);
+								}
 							})
 							));
 				}
@@ -895,6 +899,10 @@ void FChooserTableEditor::UpdateTableColumns()
 								Chooser->Modify(true);
 								Chooser->ColumnsStructs.Swap(ColumnIndex, ColumnIndex + 1);
 								UpdateTableColumns();
+								if (SelectedColumn)
+								{
+									SelectColumn(Chooser, ColumnIndex + 1);
+								}
 							})
 							));
 				}
