fix undo in Nested Choosers

diff --git a/Engine/Plugins/Experimental/Chooser/Source/Chooser/Internal/Chooser.h b/Engine/Plugins/Experimental/Chooser/Source/Chooser/Internal/Chooser.h
--- a/Engine/Plugins/Experimental/Chooser/Source/Chooser/Internal/Chooser.h
+++ b/Engine/Plugins/Experimental/Chooser/Source/Chooser/Internal/Chooser.h
@@ -160,7 +160,7 @@ public:
 	UPROPERTY(EditAnywhere, Category="Output")
 	EObjectChooserResultType ResultType = EObjectChooserResultType::ObjectResult;
 
-	virtual TConstArrayView<FInstancedStruct> GetContextData() const override { return ContextData; }
+	virtual TConstArrayView<FInstancedStruct> GetContextData() const override { return GetContextOwner()->ContextData; }
 
 	bool IsCookedData() const { return !CookedResults.IsEmpty(); }
 	static FObjectChooserBase::EIteratorStatus EvaluateChooser(FChooserEvaluationContext& Context, const UChooserTable* Chooser, FObjectChooserBase::FObjectChooserIteratorCallback Callback);
diff --git a/Engine/Plugins/Experimental/Chooser/Source/Chooser/Private/Chooser.cpp b/Engine/Plugins/Experimental/Chooser/Source/Chooser/Private/Chooser.cpp
--- a/Engine/Plugins/Experimental/Chooser/Source/Chooser/Private/Chooser.cpp
+++ b/Engine/Plugins/Experimental/Chooser/Source/Chooser/Private/Chooser.cpp
@@ -41,6 +41,7 @@ void UChooserTable::PostLoad()
 	// fix for broken outer object on nested chooser tables
 	if (ParentTable)
 	{
+		SetFlags(RF_Transactional); // also fix nested chooser objects not created with Transactional flag
 		if (GetOuter() != ParentTable)
 		{
 			Rename(nullptr, ParentTable);
diff --git a/Engine/Plugins/Experimental/Chooser/Source/ChooserEditor/Private/ChooserEditorWidgets.cpp b/Engine/Plugins/Experimental/Chooser/Source/ChooserEditor/Private/ChooserEditorWidgets.cpp
--- a/Engine/Plugins/Experimental/Chooser/Source/ChooserEditor/Private/ChooserEditorWidgets.cpp
+++ b/Engine/Plugins/Experimental/Chooser/Source/ChooserEditor/Private/ChooserEditorWidgets.cpp
@@ -182,7 +182,7 @@ TSharedRef<SWidget> CreateNestedChooserWidget(bool bReadOnly, UObject* Transacti
 												TransactionObject->Modify(true);
 											
 												FSlateApplication::Get().DismissAllMenus();
-												NestedChooser->Chooser = NewObject<UChooserTable>(TransactionObject, UChooserTable::StaticClass(), FName(InText.ToString()));
+												NestedChooser->Chooser = NewObject<UChooserTable>(TransactionObject, UChooserTable::StaticClass(), FName(InText.ToString()), RF_Transactional);
 												NestedChooser->Chooser->ParentTable = Cast<UChooserTable>(TransactionObject);
 											});
 							
diff --git a/Engine/Plugins/Experimental/Chooser/Source/ChooserEditor/Private/SChooserTableRow.cpp b/Engine/Plugins/Experimental/Chooser/Source/ChooserEditor/Private/SChooserTableRow.cpp
--- a/Engine/Plugins/Experimental/Chooser/Source/ChooserEditor/Private/SChooserTableRow.cpp
+++ b/Engine/Plugins/Experimental/Chooser/Source/ChooserEditor/Private/SChooserTableRow.cpp
@@ -113,14 +113,14 @@ namespace UE::ChooserEditor
 			else if (ColumnName == Result) 
 			{
 				UChooserTable* ContextOwner = Chooser->GetContextOwner();
-				TSharedPtr<SWidget> ResultWidget = FObjectChooserWidgetFactories::CreateWidget(false, ContextOwner, FObjectChooserBase::StaticStruct(), Chooser->ResultsStructs[RowIndex->RowIndex].GetMutableMemory(), Chooser->ResultsStructs[RowIndex->RowIndex].GetScriptStruct(), ContextOwner->OutputObjectType,
+				TSharedPtr<SWidget> ResultWidget = FObjectChooserWidgetFactories::CreateWidget(false, Chooser, FObjectChooserBase::StaticStruct(), Chooser->ResultsStructs[RowIndex->RowIndex].GetMutableMemory(), Chooser->ResultsStructs[RowIndex->RowIndex].GetScriptStruct(), ContextOwner->OutputObjectType,
 				FOnStructPicked::CreateLambda([this, InRowIndex=RowIndex->RowIndex](const UScriptStruct* ChosenStruct)
 				{
 					UChooserTable* ContextOwner = Chooser->GetContextOwner();
 					const FScopedTransaction Transaction(LOCTEXT("Change Row Result Type", "Change Row Result Type"));
 					Chooser->Modify(true);
 					Chooser->ResultsStructs[InRowIndex].InitializeAs(ChosenStruct);
-					FObjectChooserWidgetFactories::CreateWidget(false, ContextOwner, FObjectChooserBase::StaticStruct(), Chooser->ResultsStructs[InRowIndex].GetMutableMemory(), ChosenStruct, ContextOwner->OutputObjectType, FOnStructPicked(), &CacheBorder);
+					FObjectChooserWidgetFactories::CreateWidget(false, Chooser, FObjectChooserBase::StaticStruct(), Chooser->ResultsStructs[InRowIndex].GetMutableMemory(), ChosenStruct, ContextOwner->OutputObjectType, FOnStructPicked(), &CacheBorder);
 				}),
 				&CacheBorder
 				);
@@ -153,7 +153,7 @@ namespace UE::ChooserEditor
 					FChooserColumnBase* Column = &Chooser->ColumnsStructs[ColumnIndex].GetMutable<FChooserColumnBase>();
 					const UStruct * ColumnStruct = Chooser->ColumnsStructs[ColumnIndex].GetScriptStruct();
 
-					TSharedPtr<SWidget> ColumnWidget = FObjectChooserWidgetFactories::CreateColumnWidget(Column, ColumnStruct, Chooser->GetContextOwner(), RowIndex->RowIndex);
+					TSharedPtr<SWidget> ColumnWidget = FObjectChooserWidgetFactories::CreateColumnWidget(Column, ColumnStruct, Chooser, RowIndex->RowIndex);
 				
 					if (ColumnWidget.IsValid())
 					{
@@ -238,14 +238,14 @@ namespace UE::ChooserEditor
 			else if (ColumnName == Result) 
 			{
 				UChooserTable* ContextOwner = Chooser->GetContextOwner();
-				TSharedPtr<SWidget> ResultWidget = FObjectChooserWidgetFactories::CreateWidget(false,ContextOwner, FObjectChooserBase::StaticStruct(), Chooser->FallbackResult.GetMutableMemory(), Chooser->FallbackResult.GetScriptStruct(), ContextOwner->OutputObjectType,
+				TSharedPtr<SWidget> ResultWidget = FObjectChooserWidgetFactories::CreateWidget(false, Chooser, FObjectChooserBase::StaticStruct(), Chooser->FallbackResult.GetMutableMemory(), Chooser->FallbackResult.GetScriptStruct(), ContextOwner->OutputObjectType,
 				FOnStructPicked::CreateLambda([this](const UScriptStruct* ChosenStruct)
 				{
 					UChooserTable* ContextOwner = Chooser->GetContextOwner();
 					const FScopedTransaction Transaction(LOCTEXT("Change Row Result Type", "Change Row Result Type"));
 					Chooser->Modify(true);
 					Chooser->FallbackResult.InitializeAs(ChosenStruct);
-					FObjectChooserWidgetFactories::CreateWidget(false, ContextOwner, FObjectChooserBase::StaticStruct(), Chooser->FallbackResult.GetMutableMemory(), ChosenStruct, ContextOwner->OutputObjectType, FOnStructPicked(), &CacheBorder
+					FObjectChooserWidgetFactories::CreateWidget(false, Chooser, FObjectChooserBase::StaticStruct(), Chooser->FallbackResult.GetMutableMemory(), ChosenStruct, ContextOwner->OutputObjectType, FOnStructPicked(), &CacheBorder
 						,FChooserWidgetValueChanged(), LOCTEXT("Fallback Result", "Fallback Result: (None)"));
 				}),
 				&CacheBorder
