Fix Delete key binding in chooser table editor to support deleting columns as well as rows

diff --git a/Engine/Plugins/Experimental/Chooser/Source/ChooserEditor/Private/ChooserTableEditor.cpp b/Engine/Plugins/Experimental/Chooser/Source/ChooserEditor/Private/ChooserTableEditor.cpp
--- a/Engine/Plugins/Experimental/Chooser/Source/ChooserEditor/Private/ChooserTableEditor.cpp
+++ b/Engine/Plugins/Experimental/Chooser/Source/ChooserEditor/Private/ChooserTableEditor.cpp
@@ -1180,7 +1180,7 @@ TSharedRef<SDockTab> FChooserTableEditor::SpawnTableTab( const FSpawnTabArgs& Ar
 					
 					if (Event.GetKey() == EKeys::Delete)
 					{
-						DeleteSelectedRows();
+						DeleteSelection();
 						return FReply::Handled();
 					}
 					return FReply::Unhandled();
