Nested Chooser Table navigation improvements

diff --git a/Engine/Plugins/Experimental/Chooser/Source/ChooserEditor/Private/ChooserEditorWidgets.cpp b/Engine/Plugins/Experimental/Chooser/Source/ChooserEditor/Private/ChooserEditorWidgets.cpp
--- a/Engine/Plugins/Experimental/Chooser/Source/ChooserEditor/Private/ChooserEditorWidgets.cpp
+++ b/Engine/Plugins/Experimental/Chooser/Source/ChooserEditor/Private/ChooserEditorWidgets.cpp
@@ -252,7 +252,7 @@ TSharedRef<SWidget> CreateNestedChooserWidget(bool bReadOnly, UObject* Transacti
 				if (NestedChooser->Chooser)
 				{
 					UObject* RootChooser = TransactionObject->GetPackage()->FindAssetInPackage();
-					IAssetEditorInstance* Editor = GEditor->GetEditorSubsystem<UAssetEditorSubsystem>()->FindEditorForAsset(RootChooser, true);
+					IAssetEditorInstance* Editor = GEditor->GetEditorSubsystem<UAssetEditorSubsystem>()->FindEditorForAsset(RootChooser, false);
 					if (Editor && Editor->GetEditorName() == FChooserTableEditor::EditorName)
 					{
 						// currently nested tables are only editable in the Chooser Table Editor
diff --git a/Engine/Plugins/Experimental/Chooser/Source/ChooserEditor/Private/ChooserTableEditor.cpp b/Engine/Plugins/Experimental/Chooser/Source/ChooserEditor/Private/ChooserTableEditor.cpp
--- a/Engine/Plugins/Experimental/Chooser/Source/ChooserEditor/Private/ChooserTableEditor.cpp
+++ b/Engine/Plugins/Experimental/Chooser/Source/ChooserEditor/Private/ChooserTableEditor.cpp
@@ -29,6 +29,7 @@
 #include "DragAndDrop/DecoratedDragDropOp.h"
 #include "Framework/Docking/TabManager.h"
 #include "Framework/MultiBox/MultiBoxBuilder.h"
+#include "Misc/TransactionObjectEvent.h"
 #include "Modules/ModuleManager.h"
 #include "Styling/AppStyle.h"
 #include "SChooserTableRow.h"
@@ -52,9 +53,87 @@ const FName FChooserTableEditor::PropertiesTabId( TEXT( "ChooserEditor_Propertie
 const FName FChooserTableEditor::FindReplaceTabId( TEXT( "ChooserEditor_FindReplace" ) );
 const FName FChooserTableEditor::TableTabId( TEXT( "ChooserEditor_Table" ) );
 
+constexpr int HistorySize = 16;	
+
+void FChooserTableEditor::AddHistory()
+{
+	// remove anything ahead of this in the history, if we had gone back
+	while (HistoryIndex !=0)
+	{
+		History.PopFront();
+		HistoryIndex--;
+	}
+	
+	if (History.Num() >= HistorySize)
+	{
+		History.Pop();
+	}
+	History.AddFront(GetChooser());
+}
+
+bool FChooserTableEditor::CanNavigateBack() const
+{
+	return HistoryIndex < History.Num() - 1;
+}
+
+void FChooserTableEditor::NavigateBack()
+{
+	if (HistoryIndex < History.Num() - 1)
+	{
+		HistoryIndex++;
+		SetChooserTableToEdit(History[HistoryIndex], false);
+	}
+}
+
+bool FChooserTableEditor::CanNavigateForward() const
+{
+	return HistoryIndex > 0;
+}
+
+void FChooserTableEditor::NavigateForward()
+{
+	if (HistoryIndex > 0)
+	{
+		HistoryIndex--;
+		SetChooserTableToEdit(History[HistoryIndex], false);
+	}
+}
+
+void FChooserTableEditor::SetChooserTableToEdit(UChooserTable* Chooser, bool bApplyToHistory)
+{
+	if (Chooser == GetChooser())
+	{
+		return;
+	}
+	
+	BreadcrumbTrail->ClearCrumbs();
+
+	TArray<UChooserTable*> OuterList;
+	OuterList.Push(Chooser);
+	
+	while(OuterList.Last() != GetRootChooser())
+	{
+		OuterList.Push(OuterList.Last()->ParentTable);
+	}
+
+	while(!OuterList.IsEmpty())
+	{
+		UChooserTable* Popped = OuterList.Pop();
+		BreadcrumbTrail->PushCrumb(FText::FromString(Popped->GetName()), Popped);
+	}
+	
+	if (bApplyToHistory)
+	{
+		AddHistory();
+	}
+	
+	RefreshAll();
+}
+	
 void FChooserTableEditor::PushChooserTableToEdit(UChooserTable* Chooser)
 {
 	BreadcrumbTrail->PushCrumb(FText::FromString(Chooser->GetName()), Chooser);
+	AddHistory();
 	RefreshAll();
 }
 	
@@ -417,6 +496,7 @@ void FChooserTableEditor::InitEditor( const EToolkitMode::Type Mode, const TShar
 {
 	EditingObjects = ObjectsToEdit;
 
+	History.Reserve(HistorySize);
 	BreadcrumbTrail = SNew(SBreadcrumbTrail<UChooserTable*>)
 		.ButtonStyle(FAppStyle::Get(), "GraphBreadcrumbButton")
 		.TextStyle(FAppStyle::Get(), "GraphBreadcrumbButtonText")
@@ -426,14 +506,20 @@ void FChooserTableEditor::InitEditor( const EToolkitMode::Type Mode, const TShar
 		{
 			RefreshAll();
 		})
-		.OnCrumbPopped_Lambda([this](UChooserTable* Table)
+		.OnCrumbClicked_Lambda([this](UChooserTable* Table)
 		{
+			AddHistory();
 			RefreshAll();
-		});
-		
+		})
+		.GetCrumbMenuContent_Lambda([this](UChooserTable* Item)
+		{
+			return MakeChoosersMenu(Item);
+		})
+	;
+
 	UChooserTable* RootTable = GetRootChooser();
 	BreadcrumbTrail->PushCrumb(FText::FromString(RootTable->GetName()), RootTable);
-	
+	AddHistory();
 	
 	FCoreUObjectDelegates::OnObjectsReplaced.AddSP(this, &FChooserTableEditor::OnObjectsReplaced);
 
@@ -484,6 +570,15 @@ void FChooserTableEditor::InitEditor( const EToolkitMode::Type Mode, const TShar
 	FCoreUObjectDelegates::OnObjectTransacted.AddSP(this, &FChooserTableEditor::OnObjectsTransacted);
 }
 
+void FChooserTableEditor::FocusWindow(UObject* ObjectToFocusOn)
+{
+	if (UChooserTable* Chooser = Cast<UChooserTable>(ObjectToFocusOn))
+	{
+		SetChooserTableToEdit(Chooser);
+	}
+	FAssetEditorToolkit::FocusWindow(ObjectToFocusOn);
+}
+
 FName FChooserTableEditor::GetToolkitFName() const
 {
 	return ToolkitFName;
@@ -545,13 +640,43 @@ void FChooserTableEditor::RefreshAll()
 	}
 }
 
+bool FChooserTableEditor::MatchesContext(const FTransactionContext& InContext, const TArray<TPair<UObject*, FTransactionObjectEvent>>& TransactionObjectContexts) const
+{
+	TArray<UObject*> ContainedObjects;
+	GetObjectsWithOuter(EditingObjects[0]->GetPackage(), ContainedObjects, true);
+	for(const TPair<UObject*, FTransactionObjectEvent>&  Entry : TransactionObjectContexts)
+	{
+		if (ContainedObjects.Contains(Entry.Key))
+		{
+			if (UChooserTable* Chooser = Cast<UChooserTable>(Entry.Key))
+			{
+				UndoChooser = Chooser;
+			}
+			return true;
+		}
+	}
+	return false;
+}
+
 void FChooserTableEditor::PostUndo(bool bSuccess)
 {
+	if (UndoChooser)
+	{
+		// browse to the chooser the undo is going to modify
+		SetChooserTableToEdit(UndoChooser);
+		UndoChooser = nullptr;
+	}
 	RefreshAll();
 }
 
 void FChooserTableEditor::PostRedo(bool bSuccess)
 {
+	if (UndoChooser)
+ 	{
+		// browse to the chooser the undo is going to modify
+ 		SetChooserTableToEdit(UndoChooser);
+ 		UndoChooser = nullptr;
+ 	}
 	RefreshAll();
 }
 
@@ -1109,6 +1234,36 @@ void FChooserTableEditor::RefreshRowSelectionDetails()
 	}
 }
 
+void FChooserTableEditor::MakeChoosersMenuRecursive(UObject* Outer, FMenuBuilder& MenuBuilder, const FString& Indent = "") 
+{
+	TArray<UObject*> ChildObjects;
+	GetObjectsWithOuter(Outer, ChildObjects, false);
+
+	FString SubIndent = Indent + "    ";
+	for (UObject* Object : ChildObjects)
+	{
+		if (UChooserTable* Chooser = Cast<UChooserTable>(Object))
+		{
+			MenuBuilder.AddMenuEntry( FText::FromString(Indent + Chooser->GetName()), LOCTEXT("Edit Chooser ToolTip", "Browse to this Nested Chooser Table"), FSlateIcon(),
+				FUIAction(FExecuteAction::CreateLambda([this, Chooser]()
+				{
+					SetChooserTableToEdit(Chooser);
+
+				})));
+
+			MakeChoosersMenuRecursive(Chooser, MenuBuilder, SubIndent);
+		}
+	}
+}
+	
+TSharedRef<SWidget> FChooserTableEditor::MakeChoosersMenu(UObject* RootObject)
+{
+	FMenuBuilder MenuBuilder(true, nullptr);
+
+	MakeChoosersMenuRecursive(RootObject, MenuBuilder);
+
+	return MenuBuilder.MakeWidget();
+}
 
 TSharedRef<SDockTab> FChooserTableEditor::SpawnTableTab( const FSpawnTabArgs& Args )
 {
@@ -1204,37 +1359,10 @@ TSharedRef<SDockTab> FChooserTableEditor::SpawnTableTab( const FSpawnTabArgs& Ar
 	
 	EditChooserTableButton->SetOnGetMenuContent(
     		FOnGetContent::CreateLambda(
-    			[ EditChooserTableButton, this]()
-                			{
-    							FMenuBuilder MenuBuilder(true, nullptr);
-                            
-								UObject* RootChooser = GetRootChooser();
-								TArray<UObject*> ObjectsInPackage;
-								GetObjectsWithOuter(RootChooser->GetPackage(), ObjectsInPackage);
-
-								for (UObject* Object : ObjectsInPackage)
-								{
-									if (UChooserTable* Chooser = Cast<UChooserTable>(Object))
-									{
-										MenuBuilder.AddMenuEntry( FText::FromString(Chooser->GetName()), LOCTEXT("AddExistingObjectTooltip", "Add a reference to this existing Chooser Table."), FSlateIcon(),
-											FUIAction(FExecuteAction::CreateLambda([this, EditChooserTableButton, Chooser, RootChooser]()
-											{
-												while(GetChooser() != RootChooser)
-												{
-													PopChooserTableToEdit();
-												}
-												if (Chooser != RootChooser)
-												{
-													PushChooserTableToEdit(Chooser);
-												}
-												EditChooserTableButton->SetIsOpen(false);
-											})));
-									}
-                            	}
-    
-    							return MenuBuilder.MakeWidget();
-    
-                			})
+    			[this]()
+                		{
+							return MakeChoosersMenu(GetRootChooser()->GetPackage());
+                		})
     		);
 
 	return SNew(SDockTab)
@@ -1247,6 +1375,38 @@ TSharedRef<SDockTab> FChooserTableEditor::SpawnTableTab( const FSpawnTabArgs& Ar
 			[
 				SNew(SHorizontalBox)
 				+ SHorizontalBox::Slot().AutoWidth()
+				[
+					SNew(SButton)
+						.ButtonStyle(FAppStyle::Get(), "GraphBreadcrumbButton")
+						.IsEnabled_Raw(this, &FChooserTableEditor::CanNavigateBack)
+						.OnClicked_Lambda([this]()
+						{
+							NavigateBack();
+							return FReply::Handled();
+						})
+						.Content()
+						[
+							SNew(SImage)
+							.Image(FAppStyle::Get().GetBrush("Icons.ArrowLeft"))
+						]
+				]
+				+ SHorizontalBox::Slot().AutoWidth()
+				[
+					SNew(SButton)
+						.ButtonStyle(FAppStyle::Get(), "GraphBreadcrumbButton")
+						.IsEnabled_Raw(this, &FChooserTableEditor::CanNavigateForward)
+						.OnClicked_Lambda([this]()
+						{
+							NavigateForward();
+							return FReply::Handled();
+						})
+						.Content()
+						[
+							SNew(SImage)
+							.Image(FAppStyle::Get().GetBrush("Icons.ArrowRight") )
+						]
+				]
+				+ SHorizontalBox::Slot().AutoWidth()
 				[
 					EditChooserTableButton
 				]
@@ -1254,6 +1414,8 @@ TSharedRef<SDockTab> FChooserTableEditor::SpawnTableTab( const FSpawnTabArgs& Ar
 				[
 					BreadcrumbTrail.ToSharedRef()
 				]
+
+				
 			]
 			+ SVerticalBox::Slot().FillHeight(1)
 			[
diff --git a/Engine/Plugins/Experimental/Chooser/Source/ChooserEditor/Private/ChooserTableEditor.h b/Engine/Plugins/Experimental/Chooser/Source/ChooserEditor/Private/ChooserTableEditor.h
--- a/Engine/Plugins/Experimental/Chooser/Source/ChooserEditor/Private/ChooserTableEditor.h
+++ b/Engine/Plugins/Experimental/Chooser/Source/ChooserEditor/Private/ChooserTableEditor.h
@@ -6,6 +6,7 @@
 #include "Chooser.h"
 #include "EditorUndoClient.h"
 #include "PropertyEditorDelegates.h"
+#include "Containers/RingBuffer.h"
 #include "Misc/NotifyHook.h"
 #include "Toolkits/AssetEditorToolkit.h"
 #include "Toolkits/IToolkitHost.h"
@@ -44,7 +45,6 @@ namespace UE::ChooserEditor
 		virtual void RegisterTabSpawners(const TSharedRef<class FTabManager>& TabManager) override;
 		virtual void UnregisterTabSpawners(const TSharedRef<class FTabManager>& TabManager) override;
 
-
 		/**
 		* Edits the specified asset object
 		*
@@ -55,6 +55,8 @@ namespace UE::ChooserEditor
 		*/
 		void InitEditor( const EToolkitMode::Type Mode, const TSharedPtr< class IToolkitHost >& InitToolkitHost, const TArray<UObject*>& ObjectsToEdit, FGetDetailsViewObjects GetDetailsViewObjects );
 
+		virtual void FocusWindow(UObject* ObjectToFocusOn) override;
+
 		/** Destructor */
 		virtual ~FChooserTableEditor();
 
@@ -72,6 +74,7 @@ namespace UE::ChooserEditor
 		virtual void InitToolMenuContext(FToolMenuContext& MenuContext) override;
 
 		/** FEditorUndoClient Interface */
+		virtual bool MatchesContext( const FTransactionContext& InContext, const TArray<TPair<UObject*, FTransactionObjectEvent>>& TransactionObjectContexts ) const override;
 		virtual void PostUndo(bool bSuccess) override;
 		virtual void PostRedo(bool bSuccess) override;
 		
@@ -108,6 +111,8 @@ namespace UE::ChooserEditor
 		void DeleteColumn(int Index);
 		void AddColumn(const UScriptStruct* ColumnType);
 		void RefreshRowSelectionDetails();
+		TSharedRef<SWidget> MakeChoosersMenu(UObject* RootObject);
+		void MakeChoosersMenuRecursive(UObject* Outer, FMenuBuilder& MenuBuilder, const FString& Indent);
 		void DeleteSelectedRows();
 		int MoveRow(int SourceRowIndex, int TargetIndex);
 		void SelectRow(int32 RowIndex, bool bClear = true);
@@ -153,14 +158,25 @@ namespace UE::ChooserEditor
 		static const FName PropertiesTabId;
 		static const FName FindReplaceTabId;
 		static const FName TableTabId;
+		
+		void AddHistory();
+		bool CanNavigateBack() const;
+		void NavigateBack();
+		bool CanNavigateForward() const;
+		void NavigateForward();
+		void SetChooserTableToEdit(UChooserTable* Chooser, bool bApplyToHistory = true);
 
 		/** The objects open within this editor */
 		TArray<UObject*> EditingObjects;
 
+		mutable UChooserTable* UndoChooser = nullptr;
+
 		UChooserColumnDetails* SelectedColumn = nullptr;
 		TArray<TObjectPtr<UChooserRowDetails>> SelectedRows;
 
 		TSharedPtr<SBreadcrumbTrail<UChooserTable*>> BreadcrumbTrail;
+		TRingBuffer<UChooserTable*> History;
+		int32 HistoryIndex = 0;
 		
 		void UpdateTableColumns();
 		TArray<TSharedPtr<FChooserTableRow>> TableRows;
