Fix fallback behavior in choosers so that the fallback will execute if some rows passed, but all passing rows contain nested choosers which fail to find any match (or some other type of FAssetChooserBase which fails)

diff --git a/Engine/Plugins/Experimental/Chooser/Source/Chooser/Private/Chooser.cpp b/Engine/Plugins/Experimental/Chooser/Source/Chooser/Private/Chooser.cpp
--- a/Engine/Plugins/Experimental/Chooser/Source/Chooser/Private/Chooser.cpp
+++ b/Engine/Plugins/Experimental/Chooser/Source/Chooser/Private/Chooser.cpp
@@ -424,10 +424,42 @@ FObjectChooserBase::EIteratorStatus UChooserTable::EvaluateChooser(FChooserEvalu
 	}
 	
 	bool bSetOutputs = false;
-	if (IndicesOut->IsEmpty())
+	
+	// of the rows that passed all column filters, iterate through them calling the callback until it returns Stop
+	for (FChooserIndexArray::FIndexData& SelectedIndexData : *IndicesOut)
+	{
+		if (ResultsArray->IsValidIndex(SelectedIndexData.Index))
+		{
+			const FObjectChooserBase& SelectedResult = (*ResultsArray)[SelectedIndexData.Index].Get<FObjectChooserBase>();
+			FObjectChooserBase::EIteratorStatus Status = SelectedResult.ChooseMulti(Context, Callback);
+			if (Status != FObjectChooserBase::EIteratorStatus::Continue)
+			{
+				bSetOutputs = true;
+				// trigger all output columns
+				for (const FInstancedStruct& ColumnData : Chooser->ColumnsStructs)
+				{
+					const FChooserColumnBase& Column = ColumnData.Get<FChooserColumnBase>();
+					Column.SetOutputs(Context, SelectedIndexData.Index);
+				}
+				#if WITH_EDITOR
+				if (Context.DebuggingInfo.bCurrentDebugTarget)
+				{
+					Chooser->SetDebugSelectedRow(SelectedIndexData.Index);
+				}
+				#endif
+				TRACE_CHOOSER_EVALUATION(Chooser, Context, SelectedIndexData.Index);
+			}
+			if (Status == FObjectChooserBase::EIteratorStatus::Stop)
+			{
+				return FObjectChooserBase::EIteratorStatus::Stop;
+			}
+		}
+	}
+
+	// if no rows passed, or the ones which passed contained nested choosers which failed,
+	// then return the FallbackResult and output the FallbackValue from each output column
+	if (!bSetOutputs)
 	{
-		// if no rows passed the filter columns then return the FallbackResult and output the FallbackValue from each output column
-		
 		#if WITH_EDITOR
 		if (Context.DebuggingInfo.bCurrentDebugTarget)
 		{
@@ -456,39 +488,6 @@ FObjectChooserBase::EIteratorStatus UChooserTable::EvaluateChooser(FChooserEvalu
 			}
 		}
 	}
-	else
-	{
-		// of the rows that passed all column filters, iterate through them calling the callback until it returns Stop
-		for (FChooserIndexArray::FIndexData& SelectedIndexData : *IndicesOut)
-		{
-			if (ResultsArray->IsValidIndex(SelectedIndexData.Index))
-			{
-				const FObjectChooserBase& SelectedResult = (*ResultsArray)[SelectedIndexData.Index].Get<FObjectChooserBase>();
-				FObjectChooserBase::EIteratorStatus Status = SelectedResult.ChooseMulti(Context, Callback);
-				if (Status != FObjectChooserBase::EIteratorStatus::Continue)
-				{
-					bSetOutputs = true;
-					// trigger all output columns
-					for (const FInstancedStruct& ColumnData : Chooser->ColumnsStructs)
-					{
-						const FChooserColumnBase& Column = ColumnData.Get<FChooserColumnBase>();
-						Column.SetOutputs(Context, SelectedIndexData.Index);
-					}
-					#if WITH_EDITOR
-					if (Context.DebuggingInfo.bCurrentDebugTarget)
-					{
-						Chooser->SetDebugSelectedRow(SelectedIndexData.Index);
-					}
-					#endif
-					TRACE_CHOOSER_EVALUATION(Chooser, Context, SelectedIndexData.Index);
-				}
-				if (Status == FObjectChooserBase::EIteratorStatus::Stop)
-				{
-					return FObjectChooserBase::EIteratorStatus::Stop;
-				}
-			}
-		}
-	}
 	// If this is a nested chooser make sure the parent also sets the output values from the row that contained this chooser
 	return bSetOutputs ? FObjectChooserBase::EIteratorStatus::ContinueWithOutputs : FObjectChooserBase::EIteratorStatus::Continue;
 }
