Add nested chooser hierarchy tab to chooser table editor, and support for deleting and renaming nested choosers

diff --git a/Engine/Plugins/Experimental/Chooser/Source/Chooser/Internal/Chooser.h b/Engine/Plugins/Experimental/Chooser/Source/Chooser/Internal/Chooser.h
--- a/Engine/Plugins/Experimental/Chooser/Source/Chooser/Internal/Chooser.h
+++ b/Engine/Plugins/Experimental/Chooser/Source/Chooser/Internal/Chooser.h
@@ -71,22 +71,22 @@ public:
 		
 	void SetEnableDebugTesting(bool bValue)
 	{
-		GetContextOwner()->bEnableDebugTesting = bValue;
+		GetRootChooser()->bEnableDebugTesting = bValue;
 	}
 	
 	bool GetEnableDebugTesting() const
 	{
-		return GetContextOwner()->bEnableDebugTesting;
+		return GetRootChooser()->bEnableDebugTesting;
 	}
 	
 	void SetDebugTestValuesValid(bool bValue) const
    	{
-   		GetContextOwner()->bDebugTestValuesValid = bValue;
+   		GetRootChooser()->bDebugTestValuesValid = bValue;
    	}
 	
 	bool GetDebugTestValuesValid() const
 	{
-		return GetContextOwner()->bDebugTestValuesValid;
+		return GetRootChooser()->bDebugTestValuesValid;
 	}
 
 private:
@@ -125,12 +125,16 @@ public:
 	UPROPERTY()
 	TArray<TScriptInterface<IChooserColumn>> Columns_DEPRECATED;
 #endif
+	
+	
+	UChooserTable* GetRootChooser() { return RootChooser ? RootChooser.Get() : this; }
+	const UChooserTable* GetRootChooser() const { return RootChooser ? RootChooser.Get() : this; }
 
-	UChooserTable* GetContextOwner() { return ParentTable ? ParentTable.Get() : this; }
-	const UChooserTable* GetContextOwner() const { return ParentTable ? ParentTable.Get() : this; }
+	UChooserTable* GetContextOwner() { return RootChooser ? RootChooser.Get() : this; }
+	const UChooserTable* GetContextOwner() const { return RootChooser ? RootChooser.Get() : this; }
 
 	UPROPERTY()
-	TObjectPtr<UChooserTable> ParentTable;
+	TObjectPtr<UChooserTable> RootChooser;
 	
 	// FallbackResult will be used as the Result if there are no rows in the chooser which pass all filters.  If FallbackResult is not assigned, then the Chooser will return null in that case.
 	UPROPERTY(EditAnywhere, Meta = (ExcludeBaseStruct, BaseStruct = "/Script/Chooser.ObjectChooserBase"), Category = "Fallback")
@@ -143,6 +147,20 @@ public:
 
 	UPROPERTY()
 	TArray<bool> DisabledRows;
+
+	UPROPERTY()
+	TArray<TObjectPtr<UChooserTable>> NestedChoosers;
+
+	void AddNestedChooser(UChooserTable* Chooser);
+	void RemoveNestedChooser(UChooserTable* Chooser);
+	FSimpleMulticastDelegate NestedChoosersChanged;
+
+	// deprecated in favor of RootChooser
+	UPROPERTY()
+	TObjectPtr<UChooserTable> ParentTable;
+
+	UPROPERTY()
+	uint32 Version = 0;
 #endif
 	
 	UPROPERTY()
@@ -160,7 +178,7 @@ public:
 	UPROPERTY(EditAnywhere, Category="Output")
 	EObjectChooserResultType ResultType = EObjectChooserResultType::ObjectResult;
 
-	virtual TConstArrayView<FInstancedStruct> GetContextData() const override { return GetContextOwner()->ContextData; }
+	virtual TConstArrayView<FInstancedStruct> GetContextData() const override { return GetRootChooser()->ContextData; }
 
 	bool IsCookedData() const { return !CookedResults.IsEmpty(); }
 	static FObjectChooserBase::EIteratorStatus EvaluateChooser(FChooserEvaluationContext& Context, const UChooserTable* Chooser, FObjectChooserBase::FObjectChooserIteratorCallback Callback);
diff --git a/Engine/Plugins/Experimental/Chooser/Source/Chooser/Private/Chooser.cpp b/Engine/Plugins/Experimental/Chooser/Source/Chooser/Private/Chooser.cpp
--- a/Engine/Plugins/Experimental/Chooser/Source/Chooser/Private/Chooser.cpp
+++ b/Engine/Plugins/Experimental/Chooser/Source/Chooser/Private/Chooser.cpp
@@ -34,18 +34,40 @@ void UChooserTable::PostLoad()
 {
 	Super::PostLoad();
 
-#if WITH_EDITOR
+#if WITH_EDITORONLY_DATA
 	CachedPreviousOutputObjectType = OutputObjectType;
 	CachedPreviousResultType = ResultType;
 
-	// fix for broken outer object on nested chooser tables
-	if (ParentTable)
+	if (Version != 1)
 	{
-		SetFlags(RF_Transactional); // also fix nested chooser objects not created with Transactional flag
-		if (GetOuter() != ParentTable)
+		if (ParentTable)
 		{
-			Rename(nullptr, ParentTable);
+			RootChooser = ParentTable;
+			ParentTable = nullptr;
+		
+			SetFlags(RF_Transactional); // fix nested chooser objects not created with Transactional flag
+		
+			// fix for broken outer object on nested chooser tables
+			if (GetOuter() == GetPackage())
+			{
+				Rename(nullptr, RootChooser);
+			}
 		}
+
+		if(RootChooser == nullptr && NestedChoosers.IsEmpty())
+		{
+			// data upgrade for root tables: add elements to NestedTables list
+			TArray<UObject*> ChildObjects;
+			GetObjectsWithOuter(GetPackage(), ChildObjects, true);
+			for(UObject* ChildObject : ChildObjects)
+			{
+				if(UChooserTable* Chooser = Cast<UChooserTable>(ChildObject))
+				{
+					NestedChoosers.Add(Chooser);
+				}
+			}
+		}
+		Version = 1;
 	}
 #endif
 
@@ -118,7 +140,7 @@ void UChooserTable::AddCompileDependency(const UStruct* InStructType)
 
 void UChooserTable::Compile(bool bForce)
 {
-	IHasContextClass* ContextOwner = GetContextOwner();
+	IHasContextClass* ContextOwner = GetRootChooser();
 	
 	for (FInstancedStruct& ColumnData : ColumnsStructs)
 	{
@@ -315,7 +337,7 @@ void UChooserTable::UpdateDebugging(FChooserEvaluationContext& Context) const
 {
 	FScopeLock Lock(&DebugLock);
 
-	const UChooserTable* ContextOwner = GetContextOwner();
+	const UChooserTable* RootTable = GetRootChooser();
 
 	for (const FStructView& Param : Context.Params)
 	{
@@ -331,7 +353,7 @@ void UChooserTable::UpdateDebugging(FChooserEvaluationContext& Context) const
 				
 				RecentContextObjects.Add(DebugName);
 
-				if (DebugName == ContextOwner->GetDebugTargetName())
+				if (DebugName == RootTable->GetDebugTargetName())
 				{
 					bDebugTestValuesValid = true;
 					Context.DebuggingInfo.bCurrentDebugTarget = true;
@@ -439,6 +461,14 @@ FObjectChooserBase::EIteratorStatus UChooserTable::EvaluateChooser(FChooserEvalu
 				for (const FInstancedStruct& ColumnData : Chooser->ColumnsStructs)
 				{
 					const FChooserColumnBase& Column = ColumnData.Get<FChooserColumnBase>();
+
+					#if WITH_EDITORONLY_DATA
+					if (Column.bDisabled)
+					{
+						continue;
+					}
+					#endif
+					
 					Column.SetOutputs(Context, SelectedIndexData.Index);
 				}
 				#if WITH_EDITOR
@@ -560,4 +590,17 @@ static FAutoConsoleCommand CCmdTestCookChoosers(
 	TEXT(""),
 	FConsoleCommandDelegate::CreateStatic(TestCook));
 
+
+void UChooserTable::AddNestedChooser(UChooserTable* Chooser)
+{
+	NestedChoosers.Add(Chooser);
+	NestedChoosersChanged.Broadcast();
+}
+
+void UChooserTable::RemoveNestedChooser(UChooserTable* Chooser)
+{
+	NestedChoosers.Remove(Chooser);
+	NestedChoosersChanged.Broadcast();
+}
+
 #endif
\ No newline at end of file
diff --git a/Engine/Plugins/Experimental/Chooser/Source/ChooserEditor/Private/ChooserEditorWidgets.cpp b/Engine/Plugins/Experimental/Chooser/Source/ChooserEditor/Private/ChooserEditorWidgets.cpp
--- a/Engine/Plugins/Experimental/Chooser/Source/ChooserEditor/Private/ChooserEditorWidgets.cpp
+++ b/Engine/Plugins/Experimental/Chooser/Source/ChooserEditor/Private/ChooserEditorWidgets.cpp
@@ -100,7 +100,7 @@ TSharedRef<SWidget> CreateEvaluateChooserWidget(bool bReadOnly, UObject* Transac
 		 	{
 		 		if (UChooserTable* Chooser = Cast<UChooserTable>(InAssetData.GetAsset()))
 		 		{
-		 			UChooserTable* ContextOwner = Chooser->GetContextOwner();
+		 			UChooserTable* ContextOwner = Chooser->GetRootChooser();
 		 			return !(ContextOwner->OutputObjectType && ContextOwner->OutputObjectType->IsChildOf(ResultBaseClass));
 		 		}
 		 	}
@@ -178,12 +178,51 @@ TSharedRef<SWidget> CreateNestedChooserWidget(bool bReadOnly, UObject* Transacti
 											.Label(LOCTEXT("NewNestedChooserLabel", "New Chooser Name"))
 											.OnTextCommitted_Lambda([NestedChooser, TransactionObject](FText InText, ETextCommit::Type InCommitType)
 											{
-												const FScopedTransaction Transaction(LOCTEXT("Assign New Nested Chooser", "Assign New Nested Chooser"));
-												TransactionObject->Modify(true);
-											
-												FSlateApplication::Get().DismissAllMenus();
-												NestedChooser->Chooser = NewObject<UChooserTable>(TransactionObject, UChooserTable::StaticClass(), FName(InText.ToString()), RF_Transactional);
-												NestedChooser->Chooser->ParentTable = Cast<UChooserTable>(TransactionObject);
+												if (InCommitType == ETextCommit::OnEnter)
+												{
+													FSlateApplication::Get().DismissAllMenus();
+													FName NewChooserName = FName(InText.ToString());
+													if (NewChooserName.IsNone())
+													{
+														NewChooserName = "NestedChooser";
+													}
+													FString NewChooserPlainName = NewChooserName.GetPlainNameString();
+													UChooserTable* ParentChooser = Cast<UChooserTable>(TransactionObject);
+													UChooserTable* RootChooser = ParentChooser->GetRootChooser();
+
+													bool bNameConflict = false;
+													int MaxNumber = NewChooserName.GetNumber();
+													for(UChooserTable* NestedChooserIterator : RootChooser->NestedChoosers)
+													{
+														FName NestedChooserName = NestedChooserIterator->GetFName();
+														if (NestedChooserName.GetPlainNameString() == NewChooserPlainName)
+														{
+															MaxNumber = FMath::Max(MaxNumber, NestedChooserName.GetNumber());
+															if(NestedChooserName.GetNumber() == NewChooserName.GetNumber())
+															{
+																bNameConflict = true;
+															}
+														}
+													}
+													if (bNameConflict)
+													{
+														NewChooserName.SetNumber(MaxNumber + 1);
+													}
+													
+													UChooserTable* NewChooser = NewObject<UChooserTable>(TransactionObject, UChooserTable::StaticClass(), NewChooserName, RF_Transactional);
+												
+													const FScopedTransaction Transaction(LOCTEXT("Assign New Nested Chooser", "Assign New Nested Chooser"));
+												
+													ParentChooser->Modify();
+													NestedChooser->Chooser = NewChooser;
+													
+													if (ParentChooser)
+													{
+														NestedChooser->Chooser->RootChooser = RootChooser;
+														RootChooser->Modify();
+														RootChooser->AddNestedChooser(NestedChooser->Chooser);
+													}
+												}
 											});
 							
 										FSlateApplication& SlateApp = FSlateApplication::Get();
@@ -202,27 +241,22 @@ TSharedRef<SWidget> CreateNestedChooserWidget(bool bReadOnly, UObject* Transacti
 									{
 										SubMenuBuilder.BeginSection("Existing", LOCTEXT("Existing Choosers", "Existing Choosers"));
 
-										TArray<UObject*> ObjectsInPackage;
-										GetObjectsWithOuter(TransactionObject->GetPackage(), ObjectsInPackage);
+										UChooserTable* OuterChooser = Cast<UChooserTable>(TransactionObject);
+										UChooserTable* RootTable = OuterChooser->GetRootChooser();
 
-										UObject* RootChooser = TransactionObject->GetPackage()->FindAssetInPackage();
-
-										for (UObject* Object : ObjectsInPackage)
+										for (UChooserTable* Chooser : RootTable->NestedChoosers)
 										{
 											 // exclude the main chooser 
-											if (Object != RootChooser)
+											if (Chooser != RootTable)
 											{
-												if (UChooserTable* Chooser = Cast<UChooserTable>(Object))
-												{
-													SubMenuBuilder.AddMenuEntry( FText::FromString(Chooser->GetName()), LOCTEXT("AddExistingObjectTooltip", "Add a reference to this existing Chooser Table."), FSlateIcon(),
-														FUIAction(FExecuteAction::CreateLambda([Chooser, Button, NestedChooser, TransactionObject]()
-														{
-															const FScopedTransaction Transaction(LOCTEXT("Set Nested Chooser", "Set Nested Chooser"));
-															TransactionObject->Modify(true);
-															Button->SetIsOpen(false);
-															NestedChooser->Chooser = Chooser;
-														})));
-												}
+												SubMenuBuilder.AddMenuEntry( FText::FromString(Chooser->GetName()), LOCTEXT("AddExistingObjectTooltip", "Add a reference to this existing Chooser Table."), FSlateIcon(),
+													FUIAction(FExecuteAction::CreateLambda([Chooser, Button, NestedChooser, TransactionObject]()
+													{
+														const FScopedTransaction Transaction(LOCTEXT("Set Nested Chooser", "Set Nested Chooser"));
+														TransactionObject->Modify(true);
+														Button->SetIsOpen(false);
+														NestedChooser->Chooser = Chooser;
+													})));
 											}
 										}
 										
diff --git a/Engine/Plugins/Experimental/Chooser/Source/ChooserEditor/Private/ChooserTableEditor.cpp b/Engine/Plugins/Experimental/Chooser/Source/ChooserEditor/Private/ChooserTableEditor.cpp
--- a/Engine/Plugins/Experimental/Chooser/Source/ChooserEditor/Private/ChooserTableEditor.cpp
+++ b/Engine/Plugins/Experimental/Chooser/Source/ChooserEditor/Private/ChooserTableEditor.cpp
@@ -24,6 +24,7 @@
 #include "SAssetDropTarget.h"
 #include "SClassViewer.h"
 #include "ScopedTransaction.h"
+#include "SNestedChooserTree.h"
 #include "SourceCodeNavigation.h"
 #include "StructViewerModule.h"
 #include "DragAndDrop/DecoratedDragDropOp.h"
@@ -52,6 +53,7 @@ const FName FChooserTableEditor::ToolkitFName( TEXT( "ChooserTableEditor" ) );
 const FName FChooserTableEditor::PropertiesTabId( TEXT( "ChooserEditor_Properties" ) );
 const FName FChooserTableEditor::FindReplaceTabId( TEXT( "ChooserEditor_FindReplace" ) );
 const FName FChooserTableEditor::TableTabId( TEXT( "ChooserEditor_Table" ) );
+const FName FChooserTableEditor::NestedTablesTreeTabId( TEXT( "ChooserEditor_NestedTables" ) );
 
 constexpr int HistorySize = 16;	
 
@@ -113,7 +115,7 @@ void FChooserTableEditor::SetChooserTableToEdit(UChooserTable* Chooser, bool bAp
 	
 	while(OuterList.Last() != GetRootChooser())
 	{
-		OuterList.Push(OuterList.Last()->ParentTable);
+		OuterList.Push(Cast<UChooserTable>(OuterList.Last()->GetOuter()));
 	}
 
 	while(!OuterList.IsEmpty())
@@ -161,6 +163,12 @@ void FChooserTableEditor::RegisterTabSpawners(const TSharedRef<class FTabManager
 		.SetDisplayName( LOCTEXT("TableTab", "Chooser Table") )
 		.SetGroup(WorkspaceMenuCategory.ToSharedRef())
 		.SetIcon(FSlateIcon("ChooserEditorStyle", "ChooserEditor.ChooserTableIconSmall"));
+		
+	InTabManager->RegisterTabSpawner( NestedTablesTreeTabId, FOnSpawnTab::CreateSP(this, &FChooserTableEditor::SpawnNestedTablesTreeTab) )
+		.SetDisplayName( LOCTEXT("NestedTablesTab", "Nested Choosers") )
+		.SetGroup(WorkspaceMenuCategory.ToSharedRef())
+		.SetIcon(FSlateIcon("ChooserEditorStyle", "ChooserEditor.ChooserTableIconSmall"));
+
 
 	InTabManager->RegisterTabSpawner( FindReplaceTabId, FOnSpawnTab::CreateSP(this, &FChooserTableEditor::SpawnFindReplaceTab) )
 		.SetDisplayName( LOCTEXT("FindReplaceTab", "Find/Replace") )
@@ -528,7 +536,7 @@ void FChooserTableEditor::InitEditor( const EToolkitMode::Type Mode, const TShar
 	DetailsViewArgs.NotifyHook = this;
 	DetailsViewArgs.NameAreaSettings = FDetailsViewArgs::HideNameArea;
 	DetailsView = PropertyEditorModule.CreateDetailView( DetailsViewArgs );
-	const TSharedRef<FTabManager::FLayout> StandaloneDefaultLayout = FTabManager::NewLayout( "Standalone_ChooserTableEditor_Layout_v1" )
+	const TSharedRef<FTabManager::FLayout> StandaloneDefaultLayout = FTabManager::NewLayout( "Standalone_ChooserTableEditor_Layout_v1.6" )
 	->AddArea
 	(
 		FTabManager::NewPrimaryArea() ->SetOrientation(Orient_Vertical)
@@ -543,9 +551,20 @@ void FChooserTableEditor::InitEditor( const EToolkitMode::Type Mode, const TShar
 			)
 			->Split
 			(
-				FTabManager::NewStack()
+			FTabManager::NewSplitter()->SetOrientation(Orient_Vertical)
 				->SetSizeCoefficient(0.3)
-				->AddTab( PropertiesTabId, ETabState::OpenedTab )
+				->Split
+				(
+					FTabManager::NewStack()
+					->SetSizeCoefficient(0.5)
+					->AddTab( PropertiesTabId, ETabState::OpenedTab)
+				)
+				->Split
+				(
+					FTabManager::NewStack()
+					->SetSizeCoefficient(0.5)
+					->AddTab( NestedTablesTreeTabId, ETabState::OpenedTab )
+				)
 			)
 		)
 	);
@@ -636,7 +655,16 @@ void FChooserTableEditor::RefreshAll()
 				SelectRow(Row, false);
 			}
 		}
-		
+	}
+	
+	RefreshNestedChoosers();
+}
+
+void FChooserTableEditor::RefreshNestedChoosers()
+{
+	if (NestedChooserTree.IsValid())
+	{
+		NestedChooserTree->RefreshAll();
 	}
 }
 
@@ -644,14 +672,12 @@ bool FChooserTableEditor::MatchesContext(const FTransactionContext& InContext, c
 {
 	TArray<UObject*> ContainedObjects;
 	GetObjectsWithOuter(EditingObjects[0]->GetPackage(), ContainedObjects, true);
+	
+	const UChooserTable* RootChooser = GetRootChooser();
 	for(const TPair<UObject*, FTransactionObjectEvent>&  Entry : TransactionObjectContexts)
 	{
 		if (ContainedObjects.Contains(Entry.Key))
 		{
-			if (UChooserTable* Chooser = Cast<UChooserTable>(Entry.Key))
-			{
-				UndoChooser = Chooser;
-			}
 			return true;
 		}
 	}
@@ -660,23 +686,11 @@ bool FChooserTableEditor::MatchesContext(const FTransactionContext& InContext, c
 
 void FChooserTableEditor::PostUndo(bool bSuccess)
 {
-	if (UndoChooser)
-	{
-		// browse to the chooser the undo is going to modify
-		SetChooserTableToEdit(UndoChooser);
-		UndoChooser = nullptr;
-	}
 	RefreshAll();
 }
 
 void FChooserTableEditor::PostRedo(bool bSuccess)
 {
-	if (UndoChooser)
- 	{
-		// browse to the chooser the undo is going to modify
- 		SetChooserTableToEdit(UndoChooser);
- 		UndoChooser = nullptr;
- 	}
 	RefreshAll();
 }
 
@@ -970,7 +984,7 @@ void FChooserTableEditor::UpdateTableColumns()
 	{
 		FChooserColumnBase& Column = Chooser->ColumnsStructs[ColumnIndex].GetMutable<FChooserColumnBase>();
 
-		TSharedPtr<SWidget> HeaderWidget = FObjectChooserWidgetFactories::CreateColumnWidget(&Column, Chooser->ColumnsStructs[ColumnIndex].GetScriptStruct(), Chooser->GetContextOwner(), -1);
+		TSharedPtr<SWidget> HeaderWidget = FObjectChooserWidgetFactories::CreateColumnWidget(&Column, Chooser->ColumnsStructs[ColumnIndex].GetScriptStruct(), Chooser->GetRootChooser(), -1);
 		if (!HeaderWidget.IsValid())
 		{
 			HeaderWidget = SNullWidget::NullWidget;
@@ -1244,14 +1258,17 @@ void FChooserTableEditor::MakeChoosersMenuRecursive(UObject* Outer, FMenuBuilder
 	{
 		if (UChooserTable* Chooser = Cast<UChooserTable>(Object))
 		{
-			MenuBuilder.AddMenuEntry( FText::FromString(Indent + Chooser->GetName()), LOCTEXT("Edit Chooser ToolTip", "Browse to this Nested Chooser Table"), FSlateIcon(),
-				FUIAction(FExecuteAction::CreateLambda([this, Chooser]()
-				{
-					SetChooserTableToEdit(Chooser);
-
-				})));
+			UChooserTable* RootChooser = Chooser->GetRootChooser();
+			if (Chooser == RootChooser || Chooser->GetRootChooser()->NestedChoosers.Contains(Chooser))
+			{
+				MenuBuilder.AddMenuEntry( FText::FromString(Indent + Chooser->GetName()), LOCTEXT("Edit Chooser ToolTip", "Browse to this Nested Chooser Table"), FSlateIcon(),
+					FUIAction(FExecuteAction::CreateLambda([this, Chooser]()
+					{
+						SetChooserTableToEdit(Chooser);
+					})));
 
-			MakeChoosersMenuRecursive(Chooser, MenuBuilder, SubIndent);
+				MakeChoosersMenuRecursive(Chooser, MenuBuilder, SubIndent);
+			}
 		}
 	}
 }
@@ -1265,6 +1282,17 @@ TSharedRef<SWidget> FChooserTableEditor::MakeChoosersMenu(UObject* RootObject)
 	return MenuBuilder.MakeWidget();
 }
 
+TSharedRef<SDockTab> FChooserTableEditor::SpawnNestedTablesTreeTab( const FSpawnTabArgs& Args )
+{
+	check( Args.GetTabId() == NestedTablesTreeTabId );
+
+	return SNew(SDockTab)
+		.Label( LOCTEXT("NestedChooserTreeTitle", "Nested Choosers") )
+		[
+			SAssignNew(NestedChooserTree, SNestedChooserTree).ChooserEditor(this)
+		];
+}
+
 TSharedRef<SDockTab> FChooserTableEditor::SpawnTableTab( const FSpawnTabArgs& Args )
 {
 	check( Args.GetTabId() == TableTabId );
diff --git a/Engine/Plugins/Experimental/Chooser/Source/ChooserEditor/Private/ChooserTableEditor.h b/Engine/Plugins/Experimental/Chooser/Source/ChooserEditor/Private/ChooserTableEditor.h
--- a/Engine/Plugins/Experimental/Chooser/Source/ChooserEditor/Private/ChooserTableEditor.h
+++ b/Engine/Plugins/Experimental/Chooser/Source/ChooserEditor/Private/ChooserTableEditor.h
@@ -6,6 +6,7 @@
 #include "Chooser.h"
 #include "EditorUndoClient.h"
 #include "PropertyEditorDelegates.h"
+#include "SNestedChooserTree.h"
 #include "Containers/RingBuffer.h"
 #include "Misc/NotifyHook.h"
 #include "Toolkits/AssetEditorToolkit.h"
@@ -34,6 +35,7 @@ public:
 
 namespace UE::ChooserEditor
 {
+	class SNestedChooserTree;
 	struct FChooserTableRow;
 	
 	class FChooserTableEditor : public FAssetEditorToolkit, public FSelfRegisteringEditorUndoClient, public FNotifyHook
@@ -82,6 +84,7 @@ namespace UE::ChooserEditor
 		virtual void NotifyPreChange( FProperty* PropertyAboutToChange ) override;
 		virtual void NotifyPostChange( const FPropertyChangedEvent& PropertyChangedEvent, FProperty* PropertyThatChanged) override;
 
+		const UChooserTable* GetRootChooser() const { return Cast<UChooserTable>(EditingObjects[0]); }
 		UChooserTable* GetRootChooser() { return Cast<UChooserTable>(EditingObjects[0]); }
 		UChooserTable* GetChooser() { return BreadcrumbTrail->PeekCrumb(); }
 		const UChooserTable* GetChooser() const { return BreadcrumbTrail->PeekCrumb(); }
@@ -89,6 +92,7 @@ namespace UE::ChooserEditor
 		void PushChooserTableToEdit(UChooserTable* Chooser);
 		void PopChooserTableToEdit();
 		void RefreshAll();
+		void RefreshNestedChoosers();
 	
 		/** Used to show or hide certain properties */
 		void SetPropertyVisibilityDelegate(FIsPropertyVisible InVisibilityDelegate);
@@ -118,6 +122,8 @@ namespace UE::ChooserEditor
 		void SelectRow(int32 RowIndex, bool bClear = true);
 		void ClearSelectedRows(); 
 		bool IsRowSelected(int32 RowIndex);
+		
+		void SetChooserTableToEdit(UChooserTable* Chooser, bool bApplyToHistory = true);
 
 		enum class ESelectionType
 		{
@@ -142,6 +148,8 @@ namespace UE::ChooserEditor
 		TSharedRef<SDockTab> SpawnTableTab( const FSpawnTabArgs& Args );
 		/** Create the find/replace tab and its content */
 		TSharedRef<SDockTab> SpawnFindReplaceTab( const FSpawnTabArgs& Args );
+		/** Create the nested tables list tab and its content */
+		TSharedRef<SDockTab> SpawnNestedTablesTreeTab( const FSpawnTabArgs& Args );
 	
 		TSharedRef<ITableRow> GenerateTableRow(TSharedPtr<FChooserTableRow> InItem, const TSharedRef<STableViewBase>& OwnerTable);
 
@@ -158,19 +166,17 @@ namespace UE::ChooserEditor
 		static const FName PropertiesTabId;
 		static const FName FindReplaceTabId;
 		static const FName TableTabId;
+		static const FName NestedTablesTreeTabId;
 		
 		void AddHistory();
 		bool CanNavigateBack() const;
 		void NavigateBack();
 		bool CanNavigateForward() const;
 		void NavigateForward();
-		void SetChooserTableToEdit(UChooserTable* Chooser, bool bApplyToHistory = true);
 
 		/** The objects open within this editor */
 		TArray<UObject*> EditingObjects;
 
-		mutable UChooserTable* UndoChooser = nullptr;
-
 		UChooserColumnDetails* SelectedColumn = nullptr;
 		TArray<TObjectPtr<UChooserRowDetails>> SelectedRows;
 
@@ -192,6 +198,7 @@ namespace UE::ChooserEditor
 	public:
 
 		TSharedPtr<SComboButton>& GetCreateRowComboButton() { return CreateRowComboButton; };
+		TSharedPtr<SNestedChooserTree> NestedChooserTree;
 
 		/** The name given to all instances of this type of editor */
 		static const FName ToolkitFName;
diff --git a/Engine/Plugins/Experimental/Chooser/Source/ChooserEditor/Private/OutputObjectColumnEditor.cpp b/Engine/Plugins/Experimental/Chooser/Source/ChooserEditor/Private/OutputObjectColumnEditor.cpp
--- a/Engine/Plugins/Experimental/Chooser/Source/ChooserEditor/Private/OutputObjectColumnEditor.cpp
+++ b/Engine/Plugins/Experimental/Chooser/Source/ChooserEditor/Private/OutputObjectColumnEditor.cpp
@@ -55,7 +55,7 @@ namespace UE::ChooserEditor
     		if (Column)
     		{
 				int Row = RowIndex.Get();
-				UChooserTable* ContextOwner = Chooser->GetContextOwner();
+				UChooserTable* ContextOwner = Chooser->GetRootChooser();
 
 				if (Row == ColumnWidget_SpecialIndex_Fallback)
 				{
diff --git a/Engine/Plugins/Experimental/Chooser/Source/ChooserEditor/Private/RewindDebuggerChooser.cpp b/Engine/Plugins/Experimental/Chooser/Source/ChooserEditor/Private/RewindDebuggerChooser.cpp
--- a/Engine/Plugins/Experimental/Chooser/Source/ChooserEditor/Private/RewindDebuggerChooser.cpp
+++ b/Engine/Plugins/Experimental/Chooser/Source/ChooserEditor/Private/RewindDebuggerChooser.cpp
@@ -45,10 +45,10 @@ void FRewindDebuggerChooser::Update(float DeltaTime, IRewindDebugger* RewindDebu
 					// add to recent context objects list, so that this object is selectable as a target in the chooser editor
                  	Chooser->AddRecentContextObject(DebugName);
 					
-					UChooserTable* ContextOwner = Chooser->GetContextOwner();
-					if (ContextOwner->HasDebugTarget())
+					UChooserTable* RootChooser = Chooser->GetRootChooser();
+					if (RootChooser->HasDebugTarget())
 					{
-						if (ContextOwner->GetDebugTargetName() == DebugName)
+						if (RootChooser->GetDebugTargetName() == DebugName)
 						{
 							Chooser->SetDebugSelectedRow(ChooserEvaluationData.SelectedIndex);
 							Chooser->SetDebugTestValuesValid(true);
diff --git a/Engine/Plugins/Experimental/Chooser/Source/ChooserEditor/Private/SChooserTableRow.cpp b/Engine/Plugins/Experimental/Chooser/Source/ChooserEditor/Private/SChooserTableRow.cpp
--- a/Engine/Plugins/Experimental/Chooser/Source/ChooserEditor/Private/SChooserTableRow.cpp
+++ b/Engine/Plugins/Experimental/Chooser/Source/ChooserEditor/Private/SChooserTableRow.cpp
@@ -112,11 +112,11 @@ namespace UE::ChooserEditor
 			}
 			else if (ColumnName == Result) 
 			{
-				UChooserTable* ContextOwner = Chooser->GetContextOwner();
+				UChooserTable* ContextOwner = Chooser->GetRootChooser();
 				TSharedPtr<SWidget> ResultWidget = FObjectChooserWidgetFactories::CreateWidget(false, Chooser, FObjectChooserBase::StaticStruct(), Chooser->ResultsStructs[RowIndex->RowIndex].GetMutableMemory(), Chooser->ResultsStructs[RowIndex->RowIndex].GetScriptStruct(), ContextOwner->OutputObjectType,
 				FOnStructPicked::CreateLambda([this, InRowIndex=RowIndex->RowIndex](const UScriptStruct* ChosenStruct)
 				{
-					UChooserTable* ContextOwner = Chooser->GetContextOwner();
+					UChooserTable* ContextOwner = Chooser->GetRootChooser();
 					const FScopedTransaction Transaction(LOCTEXT("Change Row Result Type", "Change Row Result Type"));
 					Chooser->Modify(true);
 					Chooser->ResultsStructs[InRowIndex].InitializeAs(ChosenStruct);
@@ -237,11 +237,11 @@ namespace UE::ChooserEditor
 			}
 			else if (ColumnName == Result) 
 			{
-				UChooserTable* ContextOwner = Chooser->GetContextOwner();
+				UChooserTable* ContextOwner = Chooser->GetRootChooser();
 				TSharedPtr<SWidget> ResultWidget = FObjectChooserWidgetFactories::CreateWidget(false, Chooser, FObjectChooserBase::StaticStruct(), Chooser->FallbackResult.GetMutableMemory(), Chooser->FallbackResult.GetScriptStruct(), ContextOwner->OutputObjectType,
 				FOnStructPicked::CreateLambda([this](const UScriptStruct* ChosenStruct)
 				{
-					UChooserTable* ContextOwner = Chooser->GetContextOwner();
+					UChooserTable* ContextOwner = Chooser->GetRootChooser();
 					const FScopedTransaction Transaction(LOCTEXT("Change Row Result Type", "Change Row Result Type"));
 					Chooser->Modify(true);
 					Chooser->FallbackResult.InitializeAs(ChosenStruct);
@@ -262,7 +262,7 @@ namespace UE::ChooserEditor
 					FChooserColumnBase* Column = &Chooser->ColumnsStructs[ColumnIndex].GetMutable<FChooserColumnBase>();
 					const UStruct * ColumnStruct = Chooser->ColumnsStructs[ColumnIndex].GetScriptStruct();
 
-					TSharedPtr<SWidget> ColumnWidget = FObjectChooserWidgetFactories::CreateColumnWidget(Column, ColumnStruct, Chooser->GetContextOwner(), -2);
+					TSharedPtr<SWidget> ColumnWidget = FObjectChooserWidgetFactories::CreateColumnWidget(Column, ColumnStruct, Chooser->GetRootChooser(), -2);
 				
 					if (ColumnWidget.IsValid())
 					{
@@ -296,7 +296,7 @@ namespace UE::ChooserEditor
 			{
 				bDropSupported = true;
 			
-				UChooserTable* ContextOwner = Chooser->GetContextOwner();
+				UChooserTable* ContextOwner = Chooser->GetRootChooser();
 				if (ContextOwner->OutputObjectType) // if OutputObjectType is null, then any kind of object is supported, don't need to check all of them
 				{
 					for (const FAssetData& Asset : ContentDragDropOp->GetAssets())
@@ -419,7 +419,7 @@ namespace UE::ChooserEditor
 						}
 						else
 						{
-							UChooserTable* ContextOwner = Chooser->GetContextOwner();
+							UChooserTable* ContextOwner = Chooser->GetRootChooser();
 							if (ContextOwner->OutputObjectType == nullptr || ensure(AssetClass->IsChildOf(ContextOwner->OutputObjectType)))
 							{
 								NewResult.InitializeAs(FAssetChooser::StaticStruct());
diff --git a/Engine/Plugins/Experimental/Chooser/Source/ChooserEditor/Private/SNestedChooserTree.cpp b/Engine/Plugins/Experimental/Chooser/Source/ChooserEditor/Private/SNestedChooserTree.cpp
new file mode 100644
--- /dev/null
+++ b/Engine/Plugins/Experimental/Chooser/Source/ChooserEditor/Private/SNestedChooserTree.cpp
@@ -0,0 +1,287 @@
+// Copyright Epic Games, Inc. All Rights Reserved.
+
+#include "SNestedChooserTree.h"
+#include "ChooserTableEditor.h"
+#include "Chooser.h"
+#include "Framework/Application/SlateApplication.h"
+#include "ObjectTools.h"
+#include "OutputObjectColumn.h"
+#include "ScopedTransaction.h"
+#include "Widgets/Input/STextEntryPopup.h"
+
+#define LOCTEXT_NAMESPACE "NestedChooserTree"
+
+namespace UE::ChooserEditor
+{
+	
+	TSharedRef<ITableRow> SNestedChooserTree::TreeViewGenerateRow(TSharedPtr<FNestedChooserTreeEntry> InItem, const TSharedRef<STableViewBase>& OwnerTable)
+	{
+			return SNew(STableRow<TSharedPtr<FNestedChooserTreeEntry>>, OwnerTable)
+						.Content()
+						[
+							SNew(STextBlock).Text(FText::FromString(InItem->Chooser->GetName()))
+						];
+	}
+	
+	void SNestedChooserTree::TreeViewGetChildren(TSharedPtr<FNestedChooserTreeEntry> InItem, TArray<TSharedPtr<FNestedChooserTreeEntry>>& OutChildren)
+    {
+		for(TSharedPtr<FNestedChooserTreeEntry>& Entry : AllChoosers)
+		{
+			if (Entry->Chooser->GetOuter() == InItem->Chooser)
+			{
+				OutChildren.Add(Entry);
+			}
+		}
+    }
+	
+	void SNestedChooserTree::TreeViewDoubleClicked(TSharedPtr<FNestedChooserTreeEntry> SelectedObject)
+	{
+		ChooserEditor->SetChooserTableToEdit(SelectedObject->Chooser);
+	}
+
+	// Replace any references to a nested chooser, from either the Result column, or an OutputObject column
+	static void ReplaceReferencesInTable(UChooserTable* ChooserToReplace, UChooserTable* ReplaceWith, UChooserTable* Table)
+	{
+		Table->Modify();
+		
+		for(FInstancedStruct& ResultData : Table->ResultsStructs)
+		{
+			if (FNestedChooser* NestedChooserResult = ResultData.GetMutablePtr<FNestedChooser>())
+			{
+				if (NestedChooserResult->Chooser == ChooserToReplace)
+				{
+					NestedChooserResult->Chooser = ReplaceWith;
+				}
+			}
+
+			if (FNestedChooser* NestedChooserFallback = Table->FallbackResult.GetMutablePtr<FNestedChooser>())
+			{
+				if (NestedChooserFallback->Chooser == ChooserToReplace)
+				{
+					NestedChooserFallback->Chooser = ReplaceWith;
+				}
+			}
+		}
+		
+		for (FInstancedStruct& ColumnData : Table->ColumnsStructs)
+		{
+			if (FOutputObjectColumn* OutputObjectColumn = ColumnData.GetMutablePtr<FOutputObjectColumn>())
+			{
+				for(FChooserOutputObjectRowData& RowData : OutputObjectColumn->RowValues)
+				{
+					if (FNestedChooser* NestedChooserResult = RowData.Value.GetMutablePtr<FNestedChooser>())
+					{
+						if (NestedChooserResult->Chooser == ChooserToReplace)
+						{
+							NestedChooserResult->Chooser = ReplaceWith;
+						}
+					}	
+				}
+				
+				if (FNestedChooser* NestedChooserFallback = OutputObjectColumn->FallbackValue.Value.GetMutablePtr<FNestedChooser>())
+				{
+					if (NestedChooserFallback->Chooser == ChooserToReplace)
+					{
+						NestedChooserFallback->Chooser = ReplaceWith;
+					}
+				}
+				
+				if (FNestedChooser* NestedChooserDefault = OutputObjectColumn->DefaultRowValue.Value.GetMutablePtr<FNestedChooser>())
+				{
+					if (NestedChooserDefault->Chooser == ChooserToReplace)
+					{
+						NestedChooserDefault->Chooser = ReplaceWith;
+					}
+				}	
+			}
+		}
+
+	}
+	
+	static void ReplaceReferences(UChooserTable* ChooserToReplace, UChooserTable* ReplaceWith, UChooserTable* RootTable)
+	{
+		ReplaceReferencesInTable(ChooserToReplace, ReplaceWith, RootTable);
+
+		for(UChooserTable* NestedChooser : RootTable->NestedChoosers)
+		{
+			ReplaceReferencesInTable(ChooserToReplace, ReplaceWith, NestedChooser);
+
+			// reparent any child tables to the root
+			if (NestedChooser->GetOuter() == ChooserToReplace)
+			{
+				NestedChooser->Rename(nullptr, RootTable);
+			}
+		}
+	}
+
+	void SNestedChooserTree::DeleteNestedChooser()
+	{
+		TArray<TSharedPtr<FNestedChooserTreeEntry>>SelectedItems = TreeView->GetSelectedItems();
+		if (SelectedItems.Num() == 1)
+		{
+			UChooserTable* ChooserToDelete = SelectedItems[0]->Chooser;
+
+			if (ChooserToDelete != RootChooser)
+			{
+				FScopedTransaction Transaction(LOCTEXT("Delete Nested Choosers", "Delete Nested Choosers"));
+				static FName DeletedChooserName = "DeletedNestedChooser";
+				DeletedChooserName.SetNumber(DeletedChooserName.GetNumber() + 1);
+				ChooserToDelete->Rename(*DeletedChooserName.ToString());
+				RootChooser->Modify(true);
+				RootChooser->RemoveNestedChooser(ChooserToDelete);
+				ReplaceReferences(ChooserToDelete, nullptr, RootChooser);
+			}
+		}
+	}
+
+	void SNestedChooserTree::RenameNestedChooser()
+	{
+		TArray<TSharedPtr<FNestedChooserTreeEntry>>SelectedItems = TreeView->GetSelectedItems();
+		if (SelectedItems.Num() == 1)
+		{
+			UChooserTable* ChooserToRename = SelectedItems[0]->Chooser;
+
+			if (ChooserToRename != RootChooser)
+			{
+				TSharedRef<STextEntryPopup> TextEntry =
+					SNew(STextEntryPopup)
+					.Label(LOCTEXT("RenameNestedChooserLabel", "Rename Chooser"))
+					.OnTextCommitted_Lambda([ChooserToRename, this](FText InText, ETextCommit::Type InCommitType)
+					{
+						if (InCommitType == ETextCommit::OnEnter)
+						{
+							FSlateApplication::Get().DismissAllMenus();
+
+							FString NewName = InText.ToString();
+							
+							UChooserTable* RootChooser = ChooserToRename->GetRootChooser();
+							for(UChooserTable* NestedChooser : RootChooser->NestedChoosers)
+							{
+								if (NestedChooser->GetName() == NewName)
+								{
+									return;
+								}
+							}
+							
+							const FScopedTransaction Transaction(LOCTEXT("Rename Nested Chooser", "Rename Nested Chooser"));
+				
+							ChooserToRename->Modify(true);
+							ChooserToRename->Rename(*NewName);
+							RefreshAll();
+						}
+					});
+
+				FSlateApplication& SlateApp = FSlateApplication::Get();
+				SlateApp.PushMenu(
+					SlateApp.GetInteractiveTopLevelWindows()[0],
+					FWidgetPath(),
+					TextEntry,
+					SlateApp.GetCursorPos(),
+					FPopupTransitionEffect::TypeInPopup
+					);
+			}
+		}
+	}
+
+	void SNestedChooserTree::Construct(const FArguments& InArgs)
+	{
+		ChooserEditor = InArgs._ChooserEditor;
+		RootChooser = ChooserEditor->GetRootChooser();
+		RootChooser->NestedChoosersChanged.AddRaw(this, &SNestedChooserTree::RefreshAll);
+		TreeEntries.Add(MakeShared<FNestedChooserTreeEntry>(FNestedChooserTreeEntry({RootChooser})));
+
+		TreeView = SNew(STreeView<TSharedPtr<FNestedChooserTreeEntry>>)
+							.OnExpansionChanged_Lambda([](TSharedPtr<FNestedChooserTreeEntry> Entry, bool bExpanded){ Entry->bExpanded = bExpanded; })
+							.OnKeyDownHandler_Lambda([this](const FGeometry&, const FKeyEvent& Event)
+							{
+								if (Event.GetKey() == EKeys::Delete)
+								{
+									DeleteNestedChooser();
+									return FReply::Handled();
+								}
+								else if (Event.GetKey() == EKeys::F2)
+								{
+									RenameNestedChooser();
+									return FReply::Handled();
+								}
+								return FReply::Unhandled();
+							})
+							.OnContextMenuOpening(this, &SNestedChooserTree::TreeViewContextMenuOpening)
+							.TreeItemsSource(&TreeEntries)
+							.SelectionMode(ESelectionMode::Single)
+							.OnGenerateRow(this, &SNestedChooserTree::TreeViewGenerateRow)
+							.OnGetChildren(this, &SNestedChooserTree::TreeViewGetChildren)
+							.OnMouseButtonDoubleClick(this, &SNestedChooserTree::TreeViewDoubleClicked)
+							;
+		
+		RefreshAll();
+	
+		ChildSlot
+		[
+			TreeView.ToSharedRef()
+		];
+
+	}
+
+	SNestedChooserTree::~SNestedChooserTree()
+	{
+		if (RootChooser)
+		{
+			RootChooser->NestedChoosersChanged.RemoveAll(this);
+		}
+	}
+
+	static TSharedPtr<FNestedChooserTreeEntry> MakeEntry(TArray<TSharedPtr<FNestedChooserTreeEntry>>& OldValues, UChooserTable* Chooser)
+	{
+		for(TSharedPtr<FNestedChooserTreeEntry>& Entry : OldValues)
+		{
+			if (Entry->Chooser == Chooser)
+			{
+				return Entry;
+			}
+		}
+
+		return MakeShared<FNestedChooserTreeEntry>(FNestedChooserTreeEntry({Chooser, true}));
+	}
+
+	void SNestedChooserTree::RefreshAll()
+	{
+		TArray<TSharedPtr<FNestedChooserTreeEntry>> OldValues = AllChoosers;
+		AllChoosers.SetNum(0);
+		RootChooser = ChooserEditor->GetRootChooser();
+		AllChoosers.Add(MakeEntry(OldValues, RootChooser));
+		for(UChooserTable* Chooser : RootChooser->NestedChoosers)
+		{
+			AllChoosers.Add(MakeEntry(OldValues, Chooser));
+		}
+
+		if (TreeView.IsValid())
+		{
+			TreeView->RebuildList();
+			
+			for(TSharedPtr<FNestedChooserTreeEntry>& Entry : AllChoosers)
+			{
+				TreeView->SetItemExpansion(Entry, Entry->bExpanded);
+			}
+		}
+	}
+
+	TSharedPtr<SWidget> SNestedChooserTree::TreeViewContextMenuOpening()
+	{
+		FMenuBuilder MenuBuilder(true, nullptr);
+		MenuBuilder.AddMenuEntry(LOCTEXT("Delete", "Delete"), LOCTEXT("Delete Tooltip", "Delete Nested Chooser Table"), FSlateIcon(), FUIAction
+			(
+				FExecuteAction::CreateRaw(this, &SNestedChooserTree::DeleteNestedChooser)
+				)
+				);
+		MenuBuilder.AddMenuEntry(LOCTEXT("Rename", "Rename"), LOCTEXT("Rename Tooltip", "Rename Nested Chooser Table"), FSlateIcon(), FUIAction
+			(
+				FExecuteAction::CreateRaw(this, &SNestedChooserTree::RenameNestedChooser)
+				)
+				);
+
+		return MenuBuilder.MakeWidget();
+	}
+}
+
+#undef LOCTEXT_NAMESPACE
diff --git a/Engine/Plugins/Experimental/Chooser/Source/ChooserEditor/Private/SNestedChooserTree.h b/Engine/Plugins/Experimental/Chooser/Source/ChooserEditor/Private/SNestedChooserTree.h
new file mode 100644
--- /dev/null
+++ b/Engine/Plugins/Experimental/Chooser/Source/ChooserEditor/Private/SNestedChooserTree.h
@@ -0,0 +1,48 @@
+// Copyright Epic Games, Inc. All Rights Reserved.
+#pragma once
+
+#include "Widgets/SCompoundWidget.h"
+#include "Widgets/Views/STreeView.h"
+
+class UChooserTable;
+
+namespace UE::ChooserEditor
+{
+
+class FChooserTableEditor;
+
+struct FNestedChooserTreeEntry
+{
+	UChooserTable* Chooser = nullptr;
+	bool bExpanded = true;
+};
+
+class SNestedChooserTree : public SCompoundWidget
+{
+public:
+	SLATE_BEGIN_ARGS(SNestedChooserTree)
+	{}
+	SLATE_ARGUMENT(FChooserTableEditor*, ChooserEditor)
+	SLATE_END_ARGS()
+
+	void Construct(const FArguments& InArgs);
+	virtual ~SNestedChooserTree() override;
+
+	void RefreshAll();
+private:
+	TSharedPtr<SWidget> TreeViewContextMenuOpening();
+	TSharedRef<ITableRow> TreeViewGenerateRow(TSharedPtr<FNestedChooserTreeEntry> InItem, const TSharedRef<STableViewBase>& OwnerTable);
+	void TreeViewGetChildren(TSharedPtr<FNestedChooserTreeEntry> InItem, TArray<TSharedPtr<FNestedChooserTreeEntry>>& OutChildren);
+	void TreeViewDoubleClicked(TSharedPtr<FNestedChooserTreeEntry> SelectedObject);
+	void DeleteNestedChooser();
+	void RenameNestedChooser();
+	
+	FChooserTableEditor* ChooserEditor = nullptr;
+	UChooserTable* RootChooser = nullptr;
+	TSharedPtr<STreeView<TSharedPtr<FNestedChooserTreeEntry>>> TreeView;
+
+	TArray<TSharedPtr<FNestedChooserTreeEntry>> TreeEntries;
+	TArray<TSharedPtr<FNestedChooserTreeEntry>> AllChoosers;
+};
+
+}
