Bulid Fix for variable shadowing a class member

diff --git a/Engine/Plugins/Experimental/Chooser/Source/ChooserEditor/Private/SNestedChooserTree.cpp b/Engine/Plugins/Experimental/Chooser/Source/ChooserEditor/Private/SNestedChooserTree.cpp
--- a/Engine/Plugins/Experimental/Chooser/Source/ChooserEditor/Private/SNestedChooserTree.cpp
+++ b/Engine/Plugins/Experimental/Chooser/Source/ChooserEditor/Private/SNestedChooserTree.cpp
@@ -154,7 +154,6 @@ namespace UE::ChooserEditor
 
 							FString NewName = InText.ToString();
 							
-							UChooserTable* RootChooser = ChooserToRename->GetRootChooser();
 							for(UChooserTable* NestedChooser : RootChooser->NestedChoosers)
 							{
 								if (NestedChooser->GetName() == NewName)
