Crash fixes if you create a nested chooser in a proxy table

diff --git a/Engine/Plugins/Experimental/Chooser/Source/ChooserEditor/Private/ChooserEditorWidgets.cpp b/Engine/Plugins/Experimental/Chooser/Source/ChooserEditor/Private/ChooserEditorWidgets.cpp
--- a/Engine/Plugins/Experimental/Chooser/Source/ChooserEditor/Private/ChooserEditorWidgets.cpp
+++ b/Engine/Plugins/Experimental/Chooser/Source/ChooserEditor/Private/ChooserEditorWidgets.cpp
@@ -188,35 +188,40 @@ TSharedRef<SWidget> CreateNestedChooserWidget(bool bReadOnly, UObject* Transacti
 													}
 													FString NewChooserPlainName = NewChooserName.GetPlainNameString();
 													UChooserTable* ParentChooser = Cast<UChooserTable>(TransactionObject);
-													UChooserTable* RootChooser = ParentChooser->GetRootChooser();
-
-													bool bNameConflict = false;
-													int MaxNumber = NewChooserName.GetNumber();
-													for(UChooserTable* NestedChooserIterator : RootChooser->NestedChoosers)
+													UChooserTable* RootChooser = nullptr;
+													
+													if (ParentChooser)
 													{
-														FName NestedChooserName = NestedChooserIterator->GetFName();
-														if (NestedChooserName.GetPlainNameString() == NewChooserPlainName)
+														RootChooser = ParentChooser->GetRootChooser();
+
+														bool bNameConflict = false;
+														int MaxNumber = NewChooserName.GetNumber();
+														for(UChooserTable* NestedChooserIterator : RootChooser->NestedChoosers)
 														{
-															MaxNumber = FMath::Max(MaxNumber, NestedChooserName.GetNumber());
-															if(NestedChooserName.GetNumber() == NewChooserName.GetNumber())
+															FName NestedChooserName = NestedChooserIterator->GetFName();
+															if (NestedChooserName.GetPlainNameString() == NewChooserPlainName)
 															{
-																bNameConflict = true;
+																MaxNumber = FMath::Max(MaxNumber, NestedChooserName.GetNumber());
+																if(NestedChooserName.GetNumber() == NewChooserName.GetNumber())
+																{
+																	bNameConflict = true;
+																}
 															}
 														}
-													}
-													if (bNameConflict)
-													{
-														NewChooserName.SetNumber(MaxNumber + 1);
+														if (bNameConflict)
+														{
+															NewChooserName.SetNumber(MaxNumber + 1);
+														}
 													}
 													
 													UChooserTable* NewChooser = NewObject<UChooserTable>(TransactionObject, UChooserTable::StaticClass(), NewChooserName, RF_Transactional);
 												
 													const FScopedTransaction Transaction(LOCTEXT("Assign New Nested Chooser", "Assign New Nested Chooser"));
 												
-													ParentChooser->Modify();
+													TransactionObject->Modify();
 													NestedChooser->Chooser = NewChooser;
-													
-													if (ParentChooser)
+
+													if (RootChooser)
 													{
 														NestedChooser->Chooser->RootChooser = RootChooser;
 														RootChooser->Modify();
@@ -242,21 +247,25 @@ TSharedRef<SWidget> CreateNestedChooserWidget(bool bReadOnly, UObject* Transacti
 										SubMenuBuilder.BeginSection("Existing", LOCTEXT("Existing Choosers", "Existing Choosers"));
 
 										UChooserTable* OuterChooser = Cast<UChooserTable>(TransactionObject);
-										UChooserTable* RootTable = OuterChooser->GetRootChooser();
-
-										for (UChooserTable* Chooser : RootTable->NestedChoosers)
+										UChooserTable* RootTable = nullptr;
+										if (OuterChooser)
 										{
-											 // exclude the main chooser 
-											if (Chooser != RootTable)
+											RootTable = OuterChooser->GetRootChooser();
+
+											for (UChooserTable* Chooser : RootTable->NestedChoosers)
 											{
-												SubMenuBuilder.AddMenuEntry( FText::FromString(Chooser->GetName()), LOCTEXT("AddExistingObjectTooltip", "Add a reference to this existing Chooser Table."), FSlateIcon(),
-													FUIAction(FExecuteAction::CreateLambda([Chooser, Button, NestedChooser, TransactionObject]()
-													{
-														const FScopedTransaction Transaction(LOCTEXT("Set Nested Chooser", "Set Nested Chooser"));
-														TransactionObject->Modify(true);
-														Button->SetIsOpen(false);
-														NestedChooser->Chooser = Chooser;
-													})));
+												 // exclude the main chooser 
+												if (Chooser != RootTable)
+												{
+													SubMenuBuilder.AddMenuEntry( FText::FromString(Chooser->GetName()), LOCTEXT("AddExistingObjectTooltip", "Add a reference to this existing Chooser Table."), FSlateIcon(),
+														FUIAction(FExecuteAction::CreateLambda([Chooser, Button, NestedChooser, TransactionObject]()
+														{
+															const FScopedTransaction Transaction(LOCTEXT("Set Nested Chooser", "Set Nested Chooser"));
+															TransactionObject->Modify(true);
+															Button->SetIsOpen(false);
+															NestedChooser->Chooser = Chooser;
+														})));
+												}
 											}
 										}
 										
