Multiplayer support for chooser debugging

diff --git a/Engine/Plugins/Experimental/Chooser/Source/Chooser/Internal/Chooser.h b/Engine/Plugins/Experimental/Chooser/Source/Chooser/Internal/Chooser.h
--- a/Engine/Plugins/Experimental/Chooser/Source/Chooser/Internal/Chooser.h
+++ b/Engine/Plugins/Experimental/Chooser/Source/Chooser/Internal/Chooser.h
@@ -61,7 +61,11 @@ public:
 	const FString& GetDebugTargetName() const { return DebugTargetName; }
 	
 
-	void AddRecentContextObject(const FString& ObjectName) const { RecentContextObjects.Add(ObjectName); }
+	void AddRecentContextObject(const FString& ObjectName) const
+	{
+		RecentContextObjects.AddUnique(ObjectName);
+		RecentContextObjects.StableSort();
+	}
 	void IterateRecentContextObjects(TFunction<void(const FString&)> Callback) const;
 	void UpdateDebugging(FChooserEvaluationContext& Context) const;
 	
@@ -99,7 +103,7 @@ private:
 	EObjectChooserResultType CachedPreviousResultType = EObjectChooserResultType::ObjectResult;
 	
 	// objects this chooser has been recently evaluated on
-	mutable TSet<FString> RecentContextObjects;
+	mutable TArray<FString> RecentContextObjects;
 	mutable FCriticalSection DebugLock;
 	// reference to the UObject in PIE  which we want to get debug info for
 	mutable TWeakObjectPtr<const UObject> DebugTarget;
diff --git a/Engine/Plugins/Experimental/Chooser/Source/Chooser/Private/Chooser.cpp b/Engine/Plugins/Experimental/Chooser/Source/Chooser/Private/Chooser.cpp
--- a/Engine/Plugins/Experimental/Chooser/Source/Chooser/Private/Chooser.cpp
+++ b/Engine/Plugins/Experimental/Chooser/Source/Chooser/Private/Chooser.cpp
@@ -350,8 +350,20 @@ void UChooserTable::UpdateDebugging(FChooserEvaluationContext& Context) const
 				{
 					DebugName += " in " + Outer->GetName();					
 				}
-				
-				RecentContextObjects.Add(DebugName);
+
+				if (UWorld* World = ContextObject->GetWorld())
+				{
+					if (World->GetNetMode() == ENetMode::NM_DedicatedServer)
+					{
+						DebugName += " (Server)";
+					}
+					else if (World->GetNetMode() == ENetMode::NM_Client)
+					{
+						DebugName += FString(" (Client ") + FString::FromInt(World->GetOutermost()->GetPIEInstanceID()) + ")";
+					}
+				}
+
+				AddRecentContextObject(DebugName);
 
 				if (DebugName == RootTable->GetDebugTargetName())
 				{
diff --git a/Engine/Plugins/Experimental/Chooser/Source/ChooserEditor/Private/ChooserTableEditor.cpp b/Engine/Plugins/Experimental/Chooser/Source/ChooserEditor/Private/ChooserTableEditor.cpp
--- a/Engine/Plugins/Experimental/Chooser/Source/ChooserEditor/Private/ChooserTableEditor.cpp
+++ b/Engine/Plugins/Experimental/Chooser/Source/ChooserEditor/Private/ChooserTableEditor.cpp
@@ -216,6 +216,7 @@ FName FChooserTableEditor::GetEditorName() const
 void FChooserTableEditor::MakeDebugTargetMenu(UToolMenu* InToolMenu) 
 {
 	static FName SectionName = "Select Debug Target";
+	InToolMenu->bSearchable = true;
 		
 	InToolMenu->AddMenuEntry(
 			SectionName,
diff --git a/Engine/Plugins/Experimental/Chooser/Source/ChooserEditor/Private/ChooserTrack.cpp b/Engine/Plugins/Experimental/Chooser/Source/ChooserEditor/Private/ChooserTrack.cpp
--- a/Engine/Plugins/Experimental/Chooser/Source/ChooserEditor/Private/ChooserTrack.cpp
+++ b/Engine/Plugins/Experimental/Chooser/Source/ChooserEditor/Private/ChooserTrack.cpp
@@ -117,6 +117,20 @@ bool FChooserTrack::HandleDoubleClickInternal()
 				DebugName += " in " + FString(ActorInfo->Name);
 			}
 			
+			if (const FWorldInfo* WorldInfo = GameplayProvider->FindWorldInfoFromObject(ObjectId))
+			{
+				const FObjectInfo& WorldObjectInfo = GameplayProvider->GetObjectInfo(WorldInfo->Id);
+
+				if (WorldInfo->NetMode == FWorldInfo::ENetMode::DedicatedServer)
+				{
+					DebugName += " (Server)";
+				}
+				else if (WorldInfo->NetMode == FWorldInfo::ENetMode::Client)
+				{
+					DebugName += FString(" (Client ") + FString::FromInt(WorldInfo->PIEInstanceId) + ")";
+				}
+			}
+			
 			ChooserTable->SetDebugTarget(DebugName);
 			ChooserTable->SetEnableDebugTesting(true);
 		}
diff --git a/Engine/Plugins/Experimental/Chooser/Source/ChooserEditor/Private/RewindDebuggerChooser.cpp b/Engine/Plugins/Experimental/Chooser/Source/ChooserEditor/Private/RewindDebuggerChooser.cpp
--- a/Engine/Plugins/Experimental/Chooser/Source/ChooserEditor/Private/RewindDebuggerChooser.cpp
+++ b/Engine/Plugins/Experimental/Chooser/Source/ChooserEditor/Private/RewindDebuggerChooser.cpp
@@ -41,6 +41,20 @@ void FRewindDebuggerChooser::Update(float DeltaTime, IRewindDebugger* RewindDebu
 					{
 						DebugName += " in " + FString(ActorInfo->Name);
 					}
+
+					if (const FWorldInfo* WorldInfo = GameplayProvider->FindWorldInfoFromObject(OwnerId))
+					{
+						const FObjectInfo& WorldObjectInfo = GameplayProvider->GetObjectInfo(WorldInfo->Id);
+
+						if (WorldInfo->NetMode == FWorldInfo::ENetMode::DedicatedServer)
+						{
+							DebugName += " (Server)";
+						}
+						else if (WorldInfo->NetMode == FWorldInfo::ENetMode::Client)
+						{
+							DebugName += FString(" (Client ") + FString::FromInt(WorldInfo->PIEInstanceId) + ")";
+						}
+					}
 					
 					// add to recent context objects list, so that this object is selectable as a target in the chooser editor
                  	Chooser->AddRecentContextObject(DebugName);
