add direction setting for gameplay tag comparisons in choosers

diff --git a/Engine/Plugins/Experimental/Chooser/Source/Chooser/Internal/GameplayTagColumn.h b/Engine/Plugins/Experimental/Chooser/Source/Chooser/Internal/GameplayTagColumn.h
--- a/Engine/Plugins/Experimental/Chooser/Source/Chooser/Internal/GameplayTagColumn.h
+++ b/Engine/Plugins/Experimental/Chooser/Source/Chooser/Internal/GameplayTagColumn.h
@@ -38,6 +38,15 @@ public:
 	CHOOSER_PARAMETER_BOILERPLATE();
 };
 
+UENUM()
+enum class EGameplayTagMatchDirection : uint8
+{
+	// Row will pass if the Column Input Tags contains the Row Tag(s)
+	RowValueInInput,
+	// Row will pass if the Row Tags contain the Column Input Tag(s)
+	InputInRowValue,
+};
+
 USTRUCT()
 struct CHOOSER_API FGameplayTagColumn : public FChooserColumnBase
 {
@@ -51,6 +60,10 @@ struct CHOOSER_API FGameplayTagColumn : public FChooserColumnBase
 	UPROPERTY(EditAnywhere, Category="Data")
 	EGameplayContainerMatchType	TagMatchType = EGameplayContainerMatchType::Any;
 
+	
+	UPROPERTY(EditAnywhere, Category="Data")
+	EGameplayTagMatchDirection	TagMatchDirection = EGameplayTagMatchDirection::RowValueInInput;
+
 	//	If true, leaf tags must match exactly.
 	UPROPERTY(EditAnywhere, Category="Data")
 	bool bMatchExact = false;
diff --git a/Engine/Plugins/Experimental/Chooser/Source/Chooser/Private/GameplayTagColumn.cpp b/Engine/Plugins/Experimental/Chooser/Source/Chooser/Private/GameplayTagColumn.cpp
--- a/Engine/Plugins/Experimental/Chooser/Source/Chooser/Private/GameplayTagColumn.cpp
+++ b/Engine/Plugins/Experimental/Chooser/Source/Chooser/Private/GameplayTagColumn.cpp
@@ -31,17 +31,30 @@ bool FGameplayTagColumn::TestRow(int32 RowIndex, const FGameplayTagContainer& Va
 		}
 		else
 		{
+			const FGameplayTagContainer* A;
+			const FGameplayTagContainer* B;
+			if (TagMatchDirection == EGameplayTagMatchDirection::RowValueInInput)
+			{
+				A = &Value;
+				B = &RowValues[RowIndex];
+			}
+			else
+			{
+				A = &RowValues[RowIndex];
+				B = &Value;
+			}
+			
+			
 			if (TagMatchType == EGameplayContainerMatchType::All)
 			{
-				if (bMatchExact ? RowValues[RowIndex].HasAllExact(Value) : RowValues[RowIndex].HasAll(Value))
+				if (bMatchExact ? A->HasAllExact(*B) : A->HasAll(*B))
 				{
 					bPasses = true;
 				}
 			}
 			else
 			{
-				if (bMatchExact ? RowValues[RowIndex].HasAnyExact(Value) : RowValues[RowIndex].HasAny(Value))
-					
+				if (bMatchExact ? A->HasAnyExact(*B) : A->HasAny(*B))
 				{
 					bPasses = true;
 				}
