Fix multi enum column editor for enums with more than 8 entries

diff --git a/Engine/Plugins/Experimental/Chooser/Source/ChooserEditor/Private/MultiEnumColumnEditor.cpp b/Engine/Plugins/Experimental/Chooser/Source/ChooserEditor/Private/MultiEnumColumnEditor.cpp
--- a/Engine/Plugins/Experimental/Chooser/Source/ChooserEditor/Private/MultiEnumColumnEditor.cpp
+++ b/Engine/Plugins/Experimental/Chooser/Source/ChooserEditor/Private/MultiEnumColumnEditor.cpp
@@ -189,7 +189,7 @@ TSharedRef<SWidget> CreateMultiEnumColumnWidget(UChooserTable* Chooser, FChooser
 			{
 				if (MultiEnumColumn->RowValues.IsValidIndex(Row))
 				{
-					MultiEnumColumn->RowValues[Row].Value = static_cast<uint8>(Value);
+					MultiEnumColumn->RowValues[Row].Value = static_cast<uint32>(Value);
 				}
 			})
 			.EnumValue_Lambda([MultiEnumColumn, Row]()
