Fix double click on nested choosers in rewind debugger

diff --git a/Engine/Plugins/Experimental/Chooser/Source/ChooserEditor/Private/ChooserTableEditor.cpp b/Engine/Plugins/Experimental/Chooser/Source/ChooserEditor/Private/ChooserTableEditor.cpp
--- a/Engine/Plugins/Experimental/Chooser/Source/ChooserEditor/Private/ChooserTableEditor.cpp
+++ b/Engine/Plugins/Experimental/Chooser/Source/ChooserEditor/Private/ChooserTableEditor.cpp
@@ -503,7 +503,9 @@ void FChooserTableEditor::OnObjectsTransacted(UObject* Object, const FTransactio
 
 void FChooserTableEditor::InitEditor( const EToolkitMode::Type Mode, const TSharedPtr< class IToolkitHost >& InitToolkitHost, const TArray<UObject*>& ObjectsToEdit, FGetDetailsViewObjects GetDetailsViewObjects )
 {
-	EditingObjects = ObjectsToEdit;
+	UChooserTable* Chooser = Cast<UChooserTable>(ObjectsToEdit[0]);
+	RootChooser = Chooser->GetRootChooser();
+	check(RootChooser);
 
 	History.Reserve(HistorySize);
 	BreadcrumbTrail = SNew(SBreadcrumbTrail<UChooserTable*>)
@@ -526,8 +528,7 @@ void FChooserTableEditor::InitEditor( const EToolkitMode::Type Mode, const TShar
 		})
 	;
 
-	UChooserTable* RootTable = GetRootChooser();
-	BreadcrumbTrail->PushCrumb(FText::FromString(RootTable->GetName()), RootTable);
+	BreadcrumbTrail->PushCrumb(FText::FromString(RootChooser->GetName()), RootChooser);
 	AddHistory();
 	
 	FCoreUObjectDelegates::OnObjectsReplaced.AddSP(this, &FChooserTableEditor::OnObjectsReplaced);
@@ -583,6 +584,7 @@ void FChooserTableEditor::InitEditor( const EToolkitMode::Type Mode, const TShar
 	RegenerateMenusAndToolbars();
 
 	SelectRootProperties();
+	SetChooserTableToEdit(Chooser);
 		
 	FAnimAssetFindReplaceConfig FindReplaceConfig;
 	FindReplaceConfig.InitialProcessorClass = UChooserFindProperties::StaticClass();
@@ -672,9 +674,8 @@ void FChooserTableEditor::RefreshNestedChoosers()
 bool FChooserTableEditor::MatchesContext(const FTransactionContext& InContext, const TArray<TPair<UObject*, FTransactionObjectEvent>>& TransactionObjectContexts) const
 {
 	TArray<UObject*> ContainedObjects;
-	GetObjectsWithOuter(EditingObjects[0]->GetPackage(), ContainedObjects, true);
+	GetObjectsWithOuter(RootChooser->GetPackage(), ContainedObjects, true);
 	
-	const UChooserTable* RootChooser = GetRootChooser();
 	for(const TPair<UObject*, FTransactionObjectEvent>&  Entry : TransactionObjectContexts)
 	{
 		if (ContainedObjects.Contains(Entry.Key))
@@ -724,107 +725,16 @@ void FChooserTableEditor::NotifyPostChange(const FPropertyChangedEvent& Property
 	// editing row data should not require any refreshing
 }
 
-
 FText FChooserTableEditor::GetToolkitName() const
 {
-	const TArray<UObject*>& EditingObjs = GetEditingObjects();
-
-	check( EditingObjs.Num() > 0 );
-
-	FFormatNamedArguments Args;
-	Args.Add( TEXT("ToolkitName"), GetBaseToolkitName() );
-
-	if( EditingObjs.Num() == 1 )
-	{
-		const UObject* EditingObject = EditingObjs[ 0 ];
-		return FText::FromString(EditingObject->GetName());
-	}
-	else
-	{
-		UClass* SharedBaseClass = nullptr;
-		for( int32 x = 0; x < EditingObjs.Num(); ++x )
-		{
-			UObject* Obj = EditingObjs[ x ];
-			check( Obj );
-
-			UClass* ObjClass = Cast<UClass>(Obj);
-			if (ObjClass == nullptr)
-			{
-				ObjClass = Obj->GetClass();
-			}
-			check( ObjClass );
-
-			// Initialize with the class of the first object we encounter.
-			if( SharedBaseClass == nullptr )
-			{
-				SharedBaseClass = ObjClass;
-			}
-
-			// If we've encountered an object that's not a subclass of the current best baseclass,
-			// climb up a step in the class hierarchy.
-			while( !ObjClass->IsChildOf( SharedBaseClass ) )
-			{
-				SharedBaseClass = SharedBaseClass->GetSuperClass();
-			}
-		}
-
-		check(SharedBaseClass);
-
-		Args.Add( TEXT("NumberOfObjects"), EditingObjs.Num() );
-		Args.Add( TEXT("ClassName"), FText::FromString( SharedBaseClass->GetName() ) );
-		return FText::Format( LOCTEXT("ToolkitTitle_EditingMultiple", "{NumberOfObjects} {ClassName} - {ToolkitName}"), Args );
-	}
+	check( RootChooser );
+	return FText::FromString(RootChooser->GetName());
 }
 
 FText FChooserTableEditor::GetToolkitToolTipText() const
 {
-	const TArray<UObject*>& EditingObjs = GetEditingObjects();
-
-	check( EditingObjs.Num() > 0 );
-
-	FFormatNamedArguments Args;
-	Args.Add( TEXT("ToolkitName"), GetBaseToolkitName() );
-
-	if( EditingObjs.Num() == 1 )
-	{
-		const UObject* EditingObject = EditingObjs[ 0 ];
-		return FAssetEditorToolkit::GetToolTipTextForObject(EditingObject);
-	}
-	else
-	{
-		UClass* SharedBaseClass = NULL;
-		for( int32 x = 0; x < EditingObjs.Num(); ++x )
-		{
-			UObject* Obj = EditingObjs[ x ];
-			check( Obj );
-
-			UClass* ObjClass = Cast<UClass>(Obj);
-			if (ObjClass == nullptr)
-			{
-				ObjClass = Obj->GetClass();
-			}
-			check( ObjClass );
-
-			// Initialize with the class of the first object we encounter.
-			if( SharedBaseClass == nullptr )
-			{
-				SharedBaseClass = ObjClass;
-			}
-
-			// If we've encountered an object that's not a subclass of the current best baseclass,
-			// climb up a step in the class hierarchy.
-			while( !ObjClass->IsChildOf( SharedBaseClass ) )
-			{
-				SharedBaseClass = SharedBaseClass->GetSuperClass();
-			}
-		}
-
-		check(SharedBaseClass);
-
-		Args.Add( TEXT("NumberOfObjects"), EditingObjs.Num() );
-		Args.Add( TEXT("ClassName"), FText::FromString( SharedBaseClass->GetName() ) );
-		return FText::Format( LOCTEXT("ToolkitTitle_EditingMultipleToolTip", "{NumberOfObjects} {ClassName} - {ToolkitName}"), Args );
-	}
+	check( RootChooser );
+	return FAssetEditorToolkit::GetToolTipTextForObject(RootChooser);
 }
 
 FLinearColor FChooserTableEditor::GetWorldCentricTabColorScale() const
@@ -1259,7 +1169,6 @@ void FChooserTableEditor::MakeChoosersMenuRecursive(UObject* Outer, FMenuBuilder
 	{
 		if (UChooserTable* Chooser = Cast<UChooserTable>(Object))
 		{
-			UChooserTable* RootChooser = Chooser->GetRootChooser();
 			if (Chooser == RootChooser || Chooser->GetRootChooser()->NestedChoosers.Contains(Chooser))
 			{
 				MenuBuilder.AddMenuEntry( FText::FromString(Indent + Chooser->GetName()), LOCTEXT("Edit Chooser ToolTip", "Browse to this Nested Chooser Table"), FSlateIcon(),
@@ -1492,21 +1401,12 @@ void FChooserTableEditor::OnObjectsReplaced(const TMap<UObject*, UObject*>& Repl
 {
 	bool bChangedAny = false;
 
-	// Refresh our details view if one of the objects replaced was in the map. This gets called before the reinstance GC fixup, so we might as well fixup EditingObjects now too
-	for (int32 i = 0; i < EditingObjects.Num(); i++)
-	{
-		UObject* SourceObject = EditingObjects[i];
-		UObject* ReplacedObject = ReplacementMap.FindRef(SourceObject);
-
-		if (ReplacedObject && ReplacedObject != SourceObject)
-		{
-			EditingObjects[i] = ReplacedObject;
-			bChangedAny = true;
-		}
-	}
+	UObject* ReplacedObject = ReplacementMap.FindRef(RootChooser);
 
-	if (bChangedAny)
+	if (ReplacedObject && ReplacedObject != RootChooser)
 	{
+		RootChooser = Cast<UChooserTable>(ReplacedObject);
+		SetChooserTableToEdit(RootChooser);
 		SelectRootProperties();
 	}
 }
diff --git a/Engine/Plugins/Experimental/Chooser/Source/ChooserEditor/Private/ChooserTableEditor.h b/Engine/Plugins/Experimental/Chooser/Source/ChooserEditor/Private/ChooserTableEditor.h
--- a/Engine/Plugins/Experimental/Chooser/Source/ChooserEditor/Private/ChooserTableEditor.h
+++ b/Engine/Plugins/Experimental/Chooser/Source/ChooserEditor/Private/ChooserTableEditor.h
@@ -84,8 +84,8 @@ namespace UE::ChooserEditor
 		virtual void NotifyPreChange( FProperty* PropertyAboutToChange ) override;
 		virtual void NotifyPostChange( const FPropertyChangedEvent& PropertyChangedEvent, FProperty* PropertyThatChanged) override;
 
-		const UChooserTable* GetRootChooser() const { return Cast<UChooserTable>(EditingObjects[0]); }
-		UChooserTable* GetRootChooser() { return Cast<UChooserTable>(EditingObjects[0]); }
+		const UChooserTable* GetRootChooser() const { return RootChooser; }
+		UChooserTable* GetRootChooser() { return RootChooser; }
 		UChooserTable* GetChooser() { return BreadcrumbTrail->PeekCrumb(); }
 		const UChooserTable* GetChooser() const { return BreadcrumbTrail->PeekCrumb(); }
 
@@ -174,8 +174,8 @@ namespace UE::ChooserEditor
 		bool CanNavigateForward() const;
 		void NavigateForward();
 
-		/** The objects open within this editor */
-		TArray<UObject*> EditingObjects;
+		/** The root chooser asset being edited in this editor */
+		UChooserTable* RootChooser = nullptr;
 
 		UChooserColumnDetails* SelectedColumn = nullptr;
 		TArray<TObjectPtr<UChooserRowDetails>> SelectedRows;
diff --git a/Engine/Plugins/Experimental/Chooser/Source/ChooserEditor/Private/ChooserTrack.cpp b/Engine/Plugins/Experimental/Chooser/Source/ChooserEditor/Private/ChooserTrack.cpp
--- a/Engine/Plugins/Experimental/Chooser/Source/ChooserEditor/Private/ChooserTrack.cpp
+++ b/Engine/Plugins/Experimental/Chooser/Source/ChooserEditor/Private/ChooserTrack.cpp
@@ -110,6 +110,8 @@ bool FChooserTrack::HandleDoubleClickInternal()
 		// attach chooser table editor debugging
 		if (ChooserTable)
 		{
+			UChooserTable* RootChooser = ChooserTable->GetRootChooser();
+			
 			const FObjectInfo& OwnerObjectInfo  = GameplayProvider->GetObjectInfo(ObjectId);
 			FString DebugName(OwnerObjectInfo.Name);
 			if (const FObjectInfo* ActorInfo = RewindDebugger->FindOwningActorInfo(GameplayProvider, ObjectId))
@@ -131,12 +133,15 @@ bool FChooserTrack::HandleDoubleClickInternal()
 				}
 			}
 			
-			ChooserTable->SetDebugTarget(DebugName);
-			ChooserTable->SetEnableDebugTesting(true);
+			RootChooser->SetDebugTarget(DebugName);
+			RootChooser->SetEnableDebugTesting(true);
+			
+			GEditor->GetEditorSubsystem<UAssetEditorSubsystem>()->OpenEditorForAsset(ChooserTable);
+		}
+		else
+		{
+			GEditor->GetEditorSubsystem<UAssetEditorSubsystem>()->OpenEditorForAsset(AssetInfo.PathName);
 		}
-		
-
-		GEditor->GetEditorSubsystem<UAssetEditorSubsystem>()->OpenEditorForAsset(AssetInfo.PathName);
 
 		return true;
 	}
