Initialize nested chooser rename dialog with the current name

diff --git a/Engine/Plugins/Experimental/Chooser/Source/ChooserEditor/Private/SNestedChooserTree.cpp b/Engine/Plugins/Experimental/Chooser/Source/ChooserEditor/Private/SNestedChooserTree.cpp
--- a/Engine/Plugins/Experimental/Chooser/Source/ChooserEditor/Private/SNestedChooserTree.cpp
+++ b/Engine/Plugins/Experimental/Chooser/Source/ChooserEditor/Private/SNestedChooserTree.cpp
@@ -145,6 +145,7 @@ namespace UE::ChooserEditor
 			{
 				TSharedRef<STextEntryPopup> TextEntry =
 					SNew(STextEntryPopup)
+					.DefaultText(FText::FromString(ChooserToRename->GetName()))
 					.Label(LOCTEXT("RenameNestedChooserLabel", "Rename Chooser"))
 					.OnTextCommitted_Lambda([ChooserToRename, this](FText InText, ETextCommit::Type InCommitType)
 					{
