Chooser table editor copy/paste support

diff --git a/Engine/Plugins/Experimental/Chooser/Source/Chooser/Public/IChooserColumn.h b/Engine/Plugins/Experimental/Chooser/Source/Chooser/Public/IChooserColumn.h
--- a/Engine/Plugins/Experimental/Chooser/Source/Chooser/Public/IChooserColumn.h
+++ b/Engine/Plugins/Experimental/Chooser/Source/Chooser/Public/IChooserColumn.h
@@ -62,6 +62,7 @@ public:
 	virtual void DeleteRows(const TArray<uint32> & RowIndices) {}
 	virtual void MoveRow(int SourceIndex, int TargetIndex) {}
 	virtual void InsertRows(int Index, int Count) {}
+	virtual void CopyRow(FChooserColumnBase& SourceColumn, int SourceIndex, int TargetIndex) {}
 	
 	virtual UScriptStruct* GetInputBaseType() const { return nullptr; };
 	virtual const UScriptStruct* GetInputType() const { return nullptr; };
@@ -110,20 +111,24 @@ public:
 			RowValuesProperty.Insert(DefaultRowValue, Index);\
 		}\
 	}\
-	virtual void DeleteRows(const TArray<uint32> & RowIndices )\
+	virtual void DeleteRows(const TArray<uint32> & RowIndices ) override\
 	{\
 		for(uint32 Index : RowIndices)\
 		{\
 			RowValuesProperty.RemoveAt(Index);\
 		}\
 	}\
-	virtual void MoveRow(int SourceRowIndex, int TargetRowIndex)\
+	virtual void MoveRow(int SourceRowIndex, int TargetRowIndex) override\
 	{\
 		auto RowData = RowValuesProperty[SourceRowIndex];\
     	RowValuesProperty.RemoveAt(SourceRowIndex);\
     	if (SourceRowIndex < TargetRowIndex) { TargetRowIndex--; }\
     	RowValuesProperty.Insert(RowData, TargetRowIndex);\
 	}\
+	virtual void CopyRow(FChooserColumnBase& SourceColumn, int SourceRowIndex, int TargetRowIndex) override\
+	{\
+		RowValuesProperty.Insert(static_cast<decltype(this)>(&SourceColumn)->RowValuesProperty[SourceRowIndex], TargetRowIndex);\
+	}\
 	virtual UScriptStruct* GetInputBaseType() const override { return ParameterType::StaticStruct(); };\
 	virtual const UScriptStruct* GetInputType() const override { return InputValue.IsValid() ? InputValue.GetScriptStruct() : nullptr; };\
 	virtual FChooserParameterBase* GetInputValue() override { return InputValue.GetMutablePtr<FChooserParameterBase>(); };\
diff --git a/Engine/Plugins/Experimental/Chooser/Source/ChooserEditor/ChooserEditor.Build.cs b/Engine/Plugins/Experimental/Chooser/Source/ChooserEditor/ChooserEditor.Build.cs
--- a/Engine/Plugins/Experimental/Chooser/Source/ChooserEditor/ChooserEditor.Build.cs
+++ b/Engine/Plugins/Experimental/Chooser/Source/ChooserEditor/ChooserEditor.Build.cs
@@ -10,6 +10,7 @@ namespace UnrealBuildTool.Rules
 			PrivateDependencyModuleNames.AddRange(
 				new string[]
 				{
+					"ApplicationCore",
 					"Core",
 					"CoreUObject",
 					"Engine",
diff --git a/Engine/Plugins/Experimental/Chooser/Source/ChooserEditor/Private/ChooserTableEditor.cpp b/Engine/Plugins/Experimental/Chooser/Source/ChooserEditor/Private/ChooserTableEditor.cpp
--- a/Engine/Plugins/Experimental/Chooser/Source/ChooserEditor/Private/ChooserTableEditor.cpp
+++ b/Engine/Plugins/Experimental/Chooser/Source/ChooserEditor/Private/ChooserTableEditor.cpp
@@ -9,6 +9,7 @@
 #include "ChooserTableEditorCommands.h"
 #include "ClassViewerFilter.h"
 #include "DetailCategoryBuilder.h"
+#include "Factories.h"
 #include "GraphEditorSettings.h"
 #include "IDetailsView.h"
 #include "IPropertyAccessEditor.h"
@@ -43,6 +44,10 @@
 #include "Widgets/Layout/SWidgetSwitcher.h"
 #include "Widgets/Views/SListView.h"
 #include "ToolMenus.h"
+#include "UnrealExporter.h"
+#include "Exporters/Exporter.h"
+#include "HAL/PlatformApplicationMisc.h"
+#include "Framework/Commands/GenericCommands.h"
 
 #define LOCTEXT_NAMESPACE "ChooserEditor"
 
@@ -381,7 +386,10 @@ void FChooserTableEditor::RegisterMenus()
 	{
 		FToolMenuSection& Section = ToolMenu->AddSection("ChooserTableRow", TAttribute<FText>());
 	
-		Section.AddEntry(FToolMenuEntry::InitMenuEntry(Commands.Delete));
+		Section.AddEntry(FToolMenuEntry::InitMenuEntry(FGenericCommands::Get().Copy));
+		Section.AddEntry(FToolMenuEntry::InitMenuEntry(FGenericCommands::Get().Cut));
+		Section.AddEntry(FToolMenuEntry::InitMenuEntry(FGenericCommands::Get().Paste));
+		Section.AddEntry(FToolMenuEntry::InitMenuEntry(FGenericCommands::Get().Delete));
 		Section.AddEntry(FToolMenuEntry::InitMenuEntry(Commands.Disable));
 		Section.AddEntry(FToolMenuEntry::InitMenuEntry(Commands.AutoPopulateSelection));
 	}
@@ -393,9 +401,12 @@ void FChooserTableEditor::RegisterMenus()
 		{
 			MenuBuilder.BeginSection("ChooserEditing", LOCTEXT("Chooser Table Editing", "Chooser Table"));
 			{
-				MenuBuilder.AddMenuEntry(FChooserTableEditorCommands::Get().RemoveDisabledData, NAME_None);
-				MenuBuilder.AddMenuEntry(FChooserTableEditorCommands::Get().Delete, NAME_None, LOCTEXT("Delete Selection", "Delete Selection"));
+				MenuBuilder.AddMenuEntry(FGenericCommands::Get().Copy, NAME_None);
+				MenuBuilder.AddMenuEntry(FGenericCommands::Get().Cut, NAME_None);
+				MenuBuilder.AddMenuEntry(FGenericCommands::Get().Paste, NAME_None);
+				MenuBuilder.AddMenuEntry(FGenericCommands::Get().Delete, NAME_None, LOCTEXT("Delete Selection", "Delete Selection"));
 				MenuBuilder.AddMenuEntry(FChooserTableEditorCommands::Get().Disable, NAME_None, LOCTEXT("Disable Selection", "Disable Selection"));
+				MenuBuilder.AddMenuEntry(FChooserTableEditorCommands::Get().RemoveDisabledData, NAME_None);
 			}
 			MenuBuilder.EndSection();
 		}
@@ -423,6 +434,7 @@ void FChooserTableEditor::InitToolMenuContext(FToolMenuContext& MenuContext)
 	MenuContext.AddObject(Context);
 }
 
+
 void FChooserTableEditor::BindCommands()
 {
 	const FChooserTableEditorCommands& Commands = FChooserTableEditorCommands::Get();
@@ -440,7 +452,7 @@ void FChooserTableEditor::BindCommands()
 		FExecuteAction::CreateSP(this, &FChooserTableEditor::RemoveDisabledData));
 	
 	ToolkitCommands->MapAction(
-		Commands.Delete,
+		FGenericCommands::Get().Delete,
 		FExecuteAction::CreateSP(this, &FChooserTableEditor::DeleteSelection),
 		FCanExecuteAction::CreateSP(this, &FChooserTableEditor::HasSelection)
 		);
@@ -457,7 +469,21 @@ void FChooserTableEditor::BindCommands()
 		FCanExecuteAction::CreateSP(this, &FChooserTableEditor::HasSelection),
 		FIsActionChecked::CreateSP(this, &FChooserTableEditor::IsSelectionDisabled)
 		);
-
+	
+	ToolkitCommands->MapAction(
+		FGenericCommands::Get().Copy,
+		FExecuteAction::CreateSP(this, &FChooserTableEditor::CopySelection),
+		FCanExecuteAction::CreateSP(this, &FChooserTableEditor::HasSelection));
+	
+	ToolkitCommands->MapAction(
+		FGenericCommands::Get().Cut,
+		FExecuteAction::CreateSP(this, &FChooserTableEditor::CutSelection),
+		FCanExecuteAction::CreateSP(this, &FChooserTableEditor::HasSelection));
+	
+	ToolkitCommands->MapAction(
+		FGenericCommands::Get().Paste,
+		FExecuteAction::CreateSP(this, &FChooserTableEditor::Paste),
+		FCanExecuteAction::CreateSP(this, &FChooserTableEditor::CanPaste));
 }
 
 void FChooserTableEditor::OnObjectsTransacted(UObject* Object, const FTransactionObjectEvent& Event)
@@ -1268,17 +1294,6 @@ TSharedRef<SDockTab> FChooserTableEditor::SpawnTableTab( const FSpawnTabArgs& Ar
 	TableView = SNew(SListView<TSharedPtr<FChooserTableRow>>)
     			.ListItemsSource(&TableRows)
 				.OnContextMenuOpening_Raw(this, &FChooserTableEditor::GenerateRowContextMenu)
-				.OnKeyDownHandler_Lambda([this](const FGeometry&, const FKeyEvent& Event)
-				{
-					
-					if (Event.GetKey() == EKeys::Delete)
-					{
-						DeleteSelection();
-						return FReply::Handled();
-					}
-					return FReply::Unhandled();
-				}
-				)
 				.OnSelectionChanged_Lambda([this](TSharedPtr<FChooserTableRow>,  ESelectInfo::Type SelectInfo)
 				{
 					// deselect any selected column
@@ -1674,6 +1689,248 @@ void FChooserTableEditor::DeleteSelection()
 		DeleteSelectedRows();
 	}
 }
+	
+	
+void FChooserTableEditor::CopySelection()
+{
+	UChooserTable* CopyData = NewObject<UChooserTable>(GetTransientPackage());
+
+	UChooserTable* Chooser = GetChooser();
+
+	// copy context data from root table
+	CopyData->OutputObjectType = RootChooser->OutputObjectType;
+	CopyData->ResultType = RootChooser->ResultType;
+	CopyData->ContextData = RootChooser->ContextData;
+
+	if (CurrentSelectionType == ESelectionType::Column)
+	{
+		// add selected column including all the cell data
+		CopyData->ColumnsStructs.Add(Chooser->ColumnsStructs[SelectedColumn->Column]);
+	}
+	else if (CurrentSelectionType == ESelectionType::Rows)
+	{
+		// copy all columns
+		CopyData->ColumnsStructs = Chooser->ColumnsStructs;
+		
+		// clear all column's cell data
+		for (FInstancedStruct& ColumnData : CopyData->ColumnsStructs)
+		{
+			FChooserColumnBase& Column = ColumnData.GetMutable<FChooserColumnBase>();
+			Column.SetNumRows(0);
+			Column.SetNumRows(SelectedRows.Num());
+		}
+
+		CopyData->ResultsStructs.SetNum(SelectedRows.Num());
+
+		// add the selected results and column data
+		
+		for (int RowIndex = 0; RowIndex < SelectedRows.Num(); RowIndex++)
+		{
+			CopyData->ResultsStructs[RowIndex] = Chooser->ResultsStructs[SelectedRows[RowIndex]->Row];
+
+			for (int ColumnIndex = 0; ColumnIndex < CopyData->ColumnsStructs.Num(); ColumnIndex++)
+			{
+				FChooserColumnBase& SourceColumn = Chooser->ColumnsStructs[ColumnIndex].GetMutable<FChooserColumnBase>();
+				FChooserColumnBase& TargetColumn = CopyData->ColumnsStructs[ColumnIndex].GetMutable<FChooserColumnBase>();
+				TargetColumn.CopyRow(SourceColumn, SelectedRows[RowIndex]->Row, RowIndex);
+			}
+
+		}
+	}
+
+	// Clear the mark state for saving.
+	UnMarkAllObjects(EObjectMark(OBJECTMARK_TagExp | OBJECTMARK_TagImp));
+	
+	// Export the clipboard to text.
+	FStringOutputDevice Archive;
+	const FExportObjectInnerContext Context;
+	UExporter::ExportToOutputDevice(&Context,CopyData, nullptr, Archive, TEXT("copy"), 0, PPF_ExportsNotFullyQualified | PPF_Copy | PPF_Delimited, false, CopyData->GetOuter());
+	FPlatformApplicationMisc::ClipboardCopy(*Archive);
+}
+	
+void FChooserTableEditor::CutSelection()
+{
+	CopySelection();
+	DeleteSelection();
+}
+
+class FChooserClipboardFactory : public FCustomizableTextObjectFactory
+{
+public:
+	
+	FChooserClipboardFactory()
+		: FCustomizableTextObjectFactory(GWarn)
+		, ClipboardContent(nullptr) 
+	{
+	}
+
+	UChooserTable* ClipboardContent;
+	
+protected:
+
+	virtual bool CanCreateClass(UClass* InObjectClass, bool& bOmitSubObjs) const override
+	{
+		if (InObjectClass->IsChildOf(UChooserTable::StaticClass()))
+		{
+			return true;
+		}
+		return false;
+	}
+	
+	virtual void ProcessConstructedObject(UObject* CreatedObject) override
+	{
+		if (CreatedObject->IsA<UChooserTable>())
+		{
+			ClipboardContent = CastChecked<UChooserTable>(CreatedObject);
+		}
+	}
+};
+
+static FString GetColumnName(FChooserColumnBase& Column)
+{
+	if (FChooserParameterBase* InputValue = Column.GetInputValue())
+	{
+		return InputValue->GetDebugName();
+	}
+	return FString();
+}
+	
+void FChooserTableEditor::Paste()
+{
+	// Get the text from the clipboard.
+	FString ClipboardText;
+	FPlatformApplicationMisc::ClipboardPaste(ClipboardText);
+	FChooserClipboardFactory Factory;
+
+	UChooserTable* Chooser = GetChooser();
+
+	if (Factory.CanCreateObjectsFromText(ClipboardText))
+	{
+		Factory.ProcessBuffer((UObject*)GetTransientPackage(), RF_Transactional, ClipboardText);
+		if (UChooserTable* PastedContent = Factory.ClipboardContent)
+		{
+			FScopedTransaction Transaction(LOCTEXT("Paste Chooser Data", "Paste Chooser Data"));
+			Chooser->Modify();
+			
+			if (PastedContent->ResultsStructs.IsEmpty())
+			{
+				// pasting a column
+				Chooser->ColumnsStructs.Append(PastedContent->ColumnsStructs);
+			}
+			else
+			{
+				// pasting rows
+				int RowsToPaste = PastedContent->ResultsStructs.Num();
+
+				// figure out where to start inserting
+				int InsertIndex = Chooser->ResultsStructs.Num();
+				if (SelectedRows.Num() > 0)
+				{
+					InsertIndex = SelectedRows[0]->Row;
+					for(int SelectedRowIndex = 1; SelectedRowIndex < SelectedRows.Num(); SelectedRowIndex ++)
+					{
+						InsertIndex = FMath::Max(InsertIndex, SelectedRows[SelectedRowIndex]->Row);
+					}
+					InsertIndex++;
+				}
+
+				Chooser->ResultsStructs.Insert(PastedContent->ResultsStructs, InsertIndex);
+				
+				// Make sure each column has the same number of row datas as there are results
+				for(FInstancedStruct& ColumnData : Chooser->ColumnsStructs)
+				{
+					FChooserColumnBase& Column = ColumnData.GetMutable<FChooserColumnBase>();
+					Column.InsertRows(InsertIndex, RowsToPaste);
+				}
+
+
+				// try to also paste column data from columns in the paste buffer which match the columns in the current chooser
+				// -- matching by column type and input value name
+
+				// keep track of target columns that have already been matched, to avoid matching multiple source columns with the same target column
+				TArray<bool> MatchedTargetColumns;
+				MatchedTargetColumns.SetNum(Chooser->ColumnsStructs.Num());
+
+				// keep track of which source columns were matched, so we can add new columns for the unmatched ones after
+				TArray<bool> MatchedSourceColumns;
+				MatchedSourceColumns.SetNum(PastedContent->ColumnsStructs.Num());
+				
+				for(int SourceColumnIndex = 0; SourceColumnIndex < PastedContent->ColumnsStructs.Num(); SourceColumnIndex++)
+				{
+					FInstancedStruct& PastedColumnData = PastedContent->ColumnsStructs[SourceColumnIndex];
+					FChooserColumnBase& PastedColumn = PastedColumnData.GetMutable<FChooserColumnBase>();
+					FString PastedColumnName = GetColumnName(PastedColumn);
+					if (FChooserParameterBase* PastedColumnInput = PastedColumn.GetInputValue())
+					{
+						for(int TargetColumnIndex = 0; TargetColumnIndex < Chooser->ColumnsStructs.Num(); TargetColumnIndex++)
+						{
+							if (!MatchedTargetColumns[TargetColumnIndex])
+							{
+								FInstancedStruct& ColumnData = Chooser->ColumnsStructs[TargetColumnIndex];
+								if (ColumnData.GetScriptStruct() == PastedColumnData.GetScriptStruct())
+								{
+									FChooserColumnBase& Column = ColumnData.GetMutable<FChooserColumnBase>();
+									FString ColumnName = GetColumnName(Column);
+
+									if (ColumnName == PastedColumnName)
+									{
+										MatchedTargetColumns[TargetColumnIndex] = true;
+										MatchedSourceColumns[SourceColumnIndex] = true;
+
+										// found a match, copy the data over
+										for (int i = 0; i < RowsToPaste; i++)
+										{
+											Column.CopyRow(PastedColumn, i, InsertIndex + i);
+										}
+										break;
+									}
+								}
+							}
+						}
+					}
+				}
+				
+				// add new columns for any source columns that were unmatched
+				for(int SourceColumnIndex = 0; SourceColumnIndex < PastedContent->ColumnsStructs.Num(); SourceColumnIndex++)
+				{
+					if (!MatchedSourceColumns[SourceColumnIndex])
+					{
+						FInstancedStruct& PastedColumnData = PastedContent->ColumnsStructs[SourceColumnIndex];
+						FChooserColumnBase& PastedColumn = PastedColumnData.GetMutable<FChooserColumnBase>();
+						// if we couldn't find a match, paste a new column
+						Chooser->ColumnsStructs.Add(PastedColumnData);
+						FChooserColumnBase& Column = Chooser->ColumnsStructs.Last().GetMutable<FChooserColumnBase>();
+						Column.SetNumRows(0);
+						Column.SetNumRows(Chooser->ResultsStructs.Num());
+						for (int i = 0; i < RowsToPaste; i++)
+						{
+							Column.CopyRow(PastedColumn, i, InsertIndex + i);
+						}
+					}
+				}
+
+				// select the inserted rows
+				ClearSelectedRows();
+				for (int i = 0; i < RowsToPaste; i++)
+				{
+					SelectRow(InsertIndex + i);
+				}
+			}
+			
+			RefreshAll();
+		}
+	}
+}
+
+	
+bool FChooserTableEditor::CanPaste() const
+{
+	// Get the text from the clipboard.
+	FString ClipboardText;
+	FPlatformApplicationMisc::ClipboardPaste(ClipboardText);
+	FChooserClipboardFactory Factory;
+	return Factory.CanCreateObjectsFromText(ClipboardText); 
+}
 
 void FChooserTableEditor::RegisterWidgets()
 {
diff --git a/Engine/Plugins/Experimental/Chooser/Source/ChooserEditor/Private/ChooserTableEditor.h b/Engine/Plugins/Experimental/Chooser/Source/ChooserEditor/Private/ChooserTableEditor.h
--- a/Engine/Plugins/Experimental/Chooser/Source/ChooserEditor/Private/ChooserTableEditor.h
+++ b/Engine/Plugins/Experimental/Chooser/Source/ChooserEditor/Private/ChooserTableEditor.h
@@ -103,6 +103,10 @@ namespace UE::ChooserEditor
 		bool IsSelectionDisabled();
 		void ToggleDisableSelection();
 		void DeleteSelection();
+		
+		void CopySelection();
+        void CutSelection();
+        		
 
 		void AutoPopulateColumn(FChooserColumnBase& Column);
 		void AutoPopulateRow(int Index);
@@ -137,6 +141,10 @@ namespace UE::ChooserEditor
 		void RemoveDisabledData();
 		void RegisterToolbar();
 		void RegisterMenus();
+
+
+		void Paste();
+		bool CanPaste() const;
 		void BindCommands();
 		void OnObjectsTransacted(UObject* Object, const FTransactionObjectEvent& Event);
 		void MakeDebugTargetMenu(UToolMenu* InToolMenu);
diff --git a/Engine/Plugins/Experimental/Chooser/Source/ChooserEditor/Private/ChooserTableEditorCommands.h b/Engine/Plugins/Experimental/Chooser/Source/ChooserEditor/Private/ChooserTableEditorCommands.h
--- a/Engine/Plugins/Experimental/Chooser/Source/ChooserEditor/Private/ChooserTableEditorCommands.h
+++ b/Engine/Plugins/Experimental/Chooser/Source/ChooserEditor/Private/ChooserTableEditorCommands.h
@@ -21,7 +21,6 @@ public:
 	TSharedPtr<FUICommandInfo> AutoPopulateAll;
 	TSharedPtr<FUICommandInfo> AutoPopulateSelection;
 	TSharedPtr<FUICommandInfo> RemoveDisabledData;
-	TSharedPtr<FUICommandInfo> Delete;
 	TSharedPtr<FUICommandInfo> Disable;
 
 	/** Initialize commands */
@@ -31,7 +30,6 @@ public:
 		UI_COMMAND(AutoPopulateAll, "AutoPopulate All", "Auto Populate cell data for all supported columns", EUserInterfaceActionType::Button, FInputChord())
 		UI_COMMAND(AutoPopulateSelection, "AutoPopulate", "Auto Populate cell data for selected rows for all supported columns", EUserInterfaceActionType::Button, FInputChord())
 		UI_COMMAND(RemoveDisabledData, "Remove Disabled Data", "Delete all data that's marked as disabled.", EUserInterfaceActionType::Button, FInputChord())
-		UI_COMMAND(Delete, "Delete", "Delete the selected Rows or Column.", EUserInterfaceActionType::Button, FInputChord())
 		UI_COMMAND(Disable, "Disable", "Disable the selected Rows or Column.", EUserInterfaceActionType::Check, FInputChord())
 	}
 };
