Prevent chooser table copy/paste from putting columns after the randomize column

diff --git a/Engine/Plugins/Experimental/Chooser/Source/ChooserEditor/Private/ChooserTableEditor.cpp b/Engine/Plugins/Experimental/Chooser/Source/ChooserEditor/Private/ChooserTableEditor.cpp
--- a/Engine/Plugins/Experimental/Chooser/Source/ChooserEditor/Private/ChooserTableEditor.cpp
+++ b/Engine/Plugins/Experimental/Chooser/Source/ChooserEditor/Private/ChooserTableEditor.cpp
@@ -22,6 +22,7 @@
 #include "StructUtils/PropertyBag.h"
 #include "PropertyCustomizationHelpers.h"
 #include "PropertyEditorModule.h"
+#include "RandomizeColumn.h"
 #include "SAssetDropTarget.h"
 #include "SClassViewer.h"
 #include "ScopedTransaction.h"
@@ -460,7 +461,7 @@ void FChooserTableEditor::BindCommands()
 	ToolkitCommands->MapAction(
 		Commands.AutoPopulateSelection,
 		FExecuteAction::CreateSP(this, &FChooserTableEditor::AutoPopulateSelection),
-		FCanExecuteAction::CreateSP(this, &FChooserTableEditor::HasSelection)
+		FCanExecuteAction::CreateSP(this, &FChooserTableEditor::CanAutoPopulateSelection)
 		);
 	
 	ToolkitCommands->MapAction(
@@ -1569,6 +1570,26 @@ void FChooserTableEditor::AutoPopulateRow(int Index)
 	}
 }
 	
+bool FChooserTableEditor::CanAutoPopulateSelection()
+{
+	if (UChooserTable* Chooser = GetChooser())
+	{
+		if (SelectedColumn)
+		{
+			if (Chooser->ColumnsStructs.IsValidIndex(SelectedColumn->Column))
+			{
+				return Chooser->ColumnsStructs[SelectedColumn->Column].Get<FChooserColumnBase>().AutoPopulates();
+			}	
+		}
+		else
+		{
+			return !SelectedRows.IsEmpty();
+		}
+	}
+	
+	return false;
+}
+	
 void FChooserTableEditor::AutoPopulateSelection()
 {
 	if (UChooserTable* Chooser = GetChooser())
@@ -1815,7 +1836,22 @@ void FChooserTableEditor::Paste()
 			if (PastedContent->ResultsStructs.IsEmpty())
 			{
 				// pasting a column
-				Chooser->ColumnsStructs.Append(PastedContent->ColumnsStructs);
+				int InsertColumnIndex = Chooser->ColumnsStructs.Num();
+				if (CurrentSelectionType == ESelectionType::Column && SelectedColumn)
+				{
+					InsertColumnIndex = FMath::Min(InsertColumnIndex, SelectedColumn->Column + 1);
+				}
+				
+				if (Chooser->ColumnsStructs.Num() > 0 && Chooser->ColumnsStructs.Num() == InsertColumnIndex)
+				{
+					// if were inserting at the end, there is a randomize column, insert new columns before it
+					if (Chooser->ColumnsStructs.Last().GetPtr<FRandomizeColumn>())
+					{
+						InsertColumnIndex--;
+					}
+				}
+				Chooser->ColumnsStructs.Insert(PastedContent->ColumnsStructs, InsertColumnIndex);
+				SelectColumn(Chooser, InsertColumnIndex);
 			}
 			else
 			{
@@ -1891,6 +1927,18 @@ void FChooserTableEditor::Paste()
 				}
 				
 				// add new columns for any source columns that were unmatched
+				
+				int InsertColumnIndex = Chooser->ColumnsStructs.Num();
+				if (Chooser->ColumnsStructs.Num() > 0)
+				{
+					// if there is a randomize column, insert new columns before it
+					if (Chooser->ColumnsStructs.Last().GetPtr<FRandomizeColumn>())
+					{
+						InsertColumnIndex--;
+					}
+				}
+
+				
 				for(int SourceColumnIndex = 0; SourceColumnIndex < PastedContent->ColumnsStructs.Num(); SourceColumnIndex++)
 				{
 					if (!MatchedSourceColumns[SourceColumnIndex])
@@ -1898,8 +1946,9 @@ void FChooserTableEditor::Paste()
 						FInstancedStruct& PastedColumnData = PastedContent->ColumnsStructs[SourceColumnIndex];
 						FChooserColumnBase& PastedColumn = PastedColumnData.GetMutable<FChooserColumnBase>();
 						// if we couldn't find a match, paste a new column
-						Chooser->ColumnsStructs.Add(PastedColumnData);
-						FChooserColumnBase& Column = Chooser->ColumnsStructs.Last().GetMutable<FChooserColumnBase>();
+						Chooser->ColumnsStructs.Insert(PastedColumnData, InsertColumnIndex);
+						FChooserColumnBase& Column = Chooser->ColumnsStructs[InsertColumnIndex].GetMutable<FChooserColumnBase>();
+						InsertColumnIndex++;
 						Column.SetNumRows(0);
 						Column.SetNumRows(Chooser->ResultsStructs.Num());
 						for (int i = 0; i < RowsToPaste; i++)
diff --git a/Engine/Plugins/Experimental/Chooser/Source/ChooserEditor/Private/ChooserTableEditor.h b/Engine/Plugins/Experimental/Chooser/Source/ChooserEditor/Private/ChooserTableEditor.h
--- a/Engine/Plugins/Experimental/Chooser/Source/ChooserEditor/Private/ChooserTableEditor.h
+++ b/Engine/Plugins/Experimental/Chooser/Source/ChooserEditor/Private/ChooserTableEditor.h
@@ -110,6 +110,7 @@ namespace UE::ChooserEditor
 
 		void AutoPopulateColumn(FChooserColumnBase& Column);
 		void AutoPopulateRow(int Index);
+		bool CanAutoPopulateSelection();
 		void AutoPopulateSelection();
 		void AutoPopulateAll();
 		
@@ -221,7 +222,6 @@ namespace UE::ChooserEditor
 		static FName ContextMenuName;
 	};
 
-
 }
 
 // todo: for menus to actually be extensible this needs to be somewhere public
