Chooser Column Drag/Drop, and other usability improvements

diff --git a/Engine/Plugins/Experimental/Chooser/Source/ChooserEditor/Private/ChooserTableEditor.cpp b/Engine/Plugins/Experimental/Chooser/Source/ChooserEditor/Private/ChooserTableEditor.cpp
--- a/Engine/Plugins/Experimental/Chooser/Source/ChooserEditor/Private/ChooserTableEditor.cpp
+++ b/Engine/Plugins/Experimental/Chooser/Source/ChooserEditor/Private/ChooserTableEditor.cpp
@@ -24,6 +24,7 @@
 #include "PropertyEditorModule.h"
 #include "RandomizeColumn.h"
 #include "SAssetDropTarget.h"
+#include "SChooserColumnHandle.h"
 #include "SClassViewer.h"
 #include "ScopedTransaction.h"
 #include "SNestedChooserTree.h"
@@ -385,7 +386,7 @@ void FChooserTableEditor::RegisterMenus()
 
 	if (ToolMenu)
 	{
-		FToolMenuSection& Section = ToolMenu->AddSection("ChooserTableRow", TAttribute<FText>());
+		FToolMenuSection& Section = ToolMenu->AddSection("ChooserEditorContext", TAttribute<FText>());
 	
 		Section.AddEntry(FToolMenuEntry::InitMenuEntry(FGenericCommands::Get().Copy));
 		Section.AddEntry(FToolMenuEntry::InitMenuEntry(FGenericCommands::Get().Cut));
@@ -393,6 +394,47 @@ void FChooserTableEditor::RegisterMenus()
 		Section.AddEntry(FToolMenuEntry::InitMenuEntry(FGenericCommands::Get().Delete));
 		Section.AddEntry(FToolMenuEntry::InitMenuEntry(Commands.Disable));
 		Section.AddEntry(FToolMenuEntry::InitMenuEntry(Commands.AutoPopulateSelection));
+		
+
+		Section.AddDynamicEntry("ColumnInputType", FNewToolMenuSectionDelegate::CreateLambda([this] (FToolMenuSection& Section)
+		{
+			if (CurrentSelectionType == ESelectionType::Column && SelectedColumn)
+			{
+				Section.AddSubMenu("ParameterType", LOCTEXT("Parameter Type", "Parameter Type"),
+					LOCTEXT("Parameter Type Tooltip", "Change the type of input/output parameter for this column"),
+					FNewToolMenuChoice(FOnGetContent::CreateLambda([this]()
+					{
+						int ColumnIndex = SelectedColumn->Column;
+						UChooserTable* Chooser = SelectedColumn->Chooser;
+						FStructViewerInitializationOptions Options;
+						Options.StructFilter = MakeShared<FStructFilter>(Chooser->ColumnsStructs[ColumnIndex].Get<FChooserColumnBase>().GetInputBaseType());
+						Options.bAllowViewOptions = false;
+						Options.bShowNoneOption = false;
+						Options.NameTypeToDisplay = EStructViewerNameTypeToDisplay::DisplayName;
+					
+						// Add class filter for columns here
+						TSharedRef<SWidget> Widget = FModuleManager::LoadModuleChecked<FStructViewerModule>("StructViewer").CreateStructViewer(Options, FOnStructPicked::CreateLambda([this, ColumnIndex](const UScriptStruct* ChosenStruct)
+						{
+							const FScopedTransaction Transaction(LOCTEXT("SetColumnInputType", "Set Column Input Type"));
+							UChooserTable* ChooserTable = GetChooser();
+							ChooserTable->ColumnsStructs[ColumnIndex].GetMutable<FChooserColumnBase>().SetInputType(ChosenStruct);
+							ChooserTable->Modify(true);
+							UpdateTableColumns();
+							UpdateTableRows();
+							
+							if (SelectedColumn && SelectedColumn->Column == ColumnIndex)
+							{
+								// if this column was selected, reselect to refresh the details widgets
+								SelectColumn(ChooserTable, ColumnIndex);
+							}
+														
+						}));
+				
+						return Widget;
+					}))
+					);
+			}
+		}));
 	}
 
 	
@@ -461,14 +503,16 @@ void FChooserTableEditor::BindCommands()
 	ToolkitCommands->MapAction(
 		Commands.AutoPopulateSelection,
 		FExecuteAction::CreateSP(this, &FChooserTableEditor::AutoPopulateSelection),
-		FCanExecuteAction::CreateSP(this, &FChooserTableEditor::CanAutoPopulateSelection)
+		FCanExecuteAction::CreateSP(this, &FChooserTableEditor::CanAutoPopulateSelection),
+		FGetActionCheckState(),FIsActionButtonVisible::CreateSP(this, &FChooserTableEditor::HasSelection)
 		);
 	
 	ToolkitCommands->MapAction(
 		Commands.Disable,
 		FExecuteAction::CreateSP(this, &FChooserTableEditor::ToggleDisableSelection),
 		FCanExecuteAction::CreateSP(this, &FChooserTableEditor::HasSelection),
-		FIsActionChecked::CreateSP(this, &FChooserTableEditor::IsSelectionDisabled)
+		FIsActionChecked::CreateSP(this, &FChooserTableEditor::IsSelectionDisabled),
+		FIsActionButtonVisible::CreateSP(this, &FChooserTableEditor::HasSelection)
 		);
 	
 	ToolkitCommands->MapAction(
@@ -831,13 +875,50 @@ void FChooserTableEditor::SelectRootProperties()
 void FChooserTableEditor::RemoveDisabledData()
 {
 	UChooserTable* Chooser = GetChooser();
-	const FScopedTransaction Transaction(LOCTEXT("Move Row", "Move Row"));
+	const FScopedTransaction Transaction(LOCTEXT("Remove Disabled Data", "Remove Disabled Data"));
 
 	Chooser->Modify(true);
 	Chooser->RemoveDisabledData();
 	RefreshAll();
 }
 
+int FChooserTableEditor::MoveColumn(int SourceIndex, int TargetIndex)
+{
+	UChooserTable* Chooser = GetChooser();
+	TargetIndex = FMath::Clamp(TargetIndex, 0, Chooser->ResultsStructs.Num());
+	
+	if (SourceIndex < 0 || SourceIndex == TargetIndex)
+	{
+		return TargetIndex;
+	}
+
+	const FScopedTransaction Transaction(LOCTEXT("Move Row", "Move Row"));
+
+	Chooser->Modify(true);
+
+	FInstancedStruct ColumnData = Chooser->ColumnsStructs[SourceIndex];
+	Chooser->ColumnsStructs.RemoveAt(SourceIndex);
+	if (SourceIndex < TargetIndex)
+	{
+		TargetIndex--;
+	}
+
+	if (TargetIndex == Chooser->ColumnsStructs.Num())
+	{
+		if (Chooser->ColumnsStructs.Last().GetPtr<FRandomizeColumn>())
+		{
+			// never drop after a Randomize Column;
+			TargetIndex--;
+		}
+	}
+
+	Chooser->ColumnsStructs.Insert(ColumnData, TargetIndex);
+
+	RefreshAll();
+
+	return TargetIndex;
+}
+
 int FChooserTableEditor::MoveRow(int SourceRowIndex, int TargetRowIndex)
 {
 	UChooserTable* Chooser = GetChooser();
@@ -901,6 +982,11 @@ bool FChooserTableEditor::IsRowSelected(int32 RowIndex)
  	}
 	return false;
 }
+	
+bool FChooserTableEditor::IsColumnSelected(int32 ColumnIndex)
+{
+	return  (CurrentSelectionType == ESelectionType::Column && SelectedColumn);
+}
 
 void FChooserTableEditor::UpdateTableColumns()
 {
@@ -931,132 +1017,25 @@ void FChooserTableEditor::UpdateTableColumns()
 		HeaderRow->AddColumn(SHeaderRow::FColumn::FArguments()
 			.ColumnId(ColumnId)
 			.ManualWidth(200)
-			.OnGetMenuContent_Lambda([this, &Column, Chooser, ColumnIndex, ColumnId]()
-			{
-				UChooserColumnMenuContext* MenuContext = NewObject<UChooserColumnMenuContext>();
-				MenuContext->Editor = this;
-				MenuContext->Chooser = Chooser;
-				MenuContext->ColumnIndex = ColumnIndex;
-
-				FMenuBuilder MenuBuilder(true,nullptr);
-
-				MenuBuilder.AddMenuEntry(LOCTEXT("Column Properties","Properties"),LOCTEXT("Select Column ToolTip", "Select this Column, and show its properties in the Details panel"),FSlateIcon(),
-					FUIAction(
-						FExecuteAction::CreateLambda([this,Chooser,ColumnIndex, ColumnId, &Column]()
-						{
-							SelectColumn(Chooser, ColumnId.GetNumber() - 1);
-						})
-						)
-					);
-
-				if (ColumnIndex > 0 && !Chooser->ColumnsStructs[ColumnIndex].Get<FChooserColumnBase>().IsRandomizeColumn())
-				{
-					MenuBuilder.AddMenuEntry(LOCTEXT("Move Left","Move Left"),LOCTEXT("Move Left ToolTip", "Move this column to the left."),FSlateIcon(),
-						FUIAction(
-							FExecuteAction::CreateLambda([this,Chooser,ColumnIndex, &Column]()
-							{
-								const FScopedTransaction Transaction(LOCTEXT("Move Column Left Transaction", "Move Column Left"));
-								Chooser->Modify(true);
-								Chooser->ColumnsStructs.Swap(ColumnIndex, ColumnIndex - 1);
-								UpdateTableColumns();
-								if (SelectedColumn)
-								{
-									SelectColumn(Chooser, ColumnIndex - 1);
-								}
-							})
-							));
-				}
-				if (ColumnIndex < Chooser->ColumnsStructs.Num() - 1 && !Chooser->ColumnsStructs[ColumnIndex+1].Get<FChooserColumnBase>().IsRandomizeColumn())
-				{
-					MenuBuilder.AddMenuEntry(LOCTEXT("Move Right","Move Right"),LOCTEXT("Move Right ToolTip", "Move this column to the right."),FSlateIcon(),
-						FUIAction(
-							FExecuteAction::CreateLambda([this,Chooser,ColumnIndex, &Column]()
-							{
-								const FScopedTransaction Transaction(LOCTEXT("Move Column Right Transaction", "Move Column Right"));
-								Chooser->Modify(true);
-								Chooser->ColumnsStructs.Swap(ColumnIndex, ColumnIndex + 1);
-								UpdateTableColumns();
-								if (SelectedColumn)
-								{
-									SelectColumn(Chooser, ColumnIndex + 1);
-								}
-							})
-							));
-				}
-
-				MenuBuilder.AddMenuEntry(LOCTEXT("Delete Column", "Delete"), LOCTEXT("Delete Column ToolTip", "Remove this column and all its data from the table"), FSlateIcon(),
-					FUIAction(
-						FExecuteAction::CreateLambda([this, Chooser, ColumnIndex, &Column]()
-						{
-							DeleteColumn(ColumnIndex);
-						})
-						));
-
-				if(Column.AutoPopulates())
-				{
-					MenuBuilder.AddMenuEntry(LOCTEXT("Auto Populate", "Auto Populate"), LOCTEXT("Auto Populate ToolTip", "Auto populate cell values for this colun"), FSlateIcon(),
-						FUIAction(
-							FExecuteAction::CreateLambda([this, &Column]()
-							{
-								const FScopedTransaction Transaction(LOCTEXT("Auto Populate Column", "Auto Populate Column"));
-								AutoPopulateColumn(Column);
-							})
-							));
-				}
-			
-				MenuBuilder.AddSubMenu(LOCTEXT("Input Type", "Input Type"),
-					LOCTEXT("InputTypeToolTip", "Change input parameter type"),
-					FNewMenuDelegate::CreateLambda([this, Chooser, &ColumnIndex](FMenuBuilder& Builder)
-					{
-						FStructViewerInitializationOptions Options;
-						Options.StructFilter = MakeShared<FStructFilter>(Chooser->ColumnsStructs[ColumnIndex].Get<FChooserColumnBase>().GetInputBaseType());
-						Options.bAllowViewOptions = false;
-						Options.bShowNoneOption = false;
-						Options.NameTypeToDisplay = EStructViewerNameTypeToDisplay::DisplayName;
-					
-						// Add class filter for columns here
-						TSharedRef<SWidget> Widget = FModuleManager::LoadModuleChecked<FStructViewerModule>("StructViewer").CreateStructViewer(Options, FOnStructPicked::CreateLambda([this, ColumnIndex](const UScriptStruct* ChosenStruct)
-						{
-							const FScopedTransaction Transaction(LOCTEXT("SetColumnInputType", "Set Column Input Type"));
-							UChooserTable* ChooserTable = GetChooser();
-							ChooserTable->ColumnsStructs[ColumnIndex].GetMutable<FChooserColumnBase>().SetInputType(ChosenStruct);
-							ChooserTable->Modify(true);
-							UpdateTableColumns();
-							UpdateTableRows();
-							
-							if (SelectedColumn && SelectedColumn->Column == ColumnIndex)
-							{
-								// if this column was selected, reselect to refresh the details widgets
-								SelectColumn(ChooserTable, ColumnIndex);
-							}
-                            							
-						}));
-				
-						Builder.AddWidget(Widget, FText());
-					}));
-			
-				return MenuBuilder.MakeWidget();
-			})
 			.HeaderComboVisibility(EHeaderComboVisibility::Ghosted)
 			.HeaderContent()
 			[
-				SNew(SBorder)
-				.VAlign(VAlign_Center)
-				.Padding(3)
-				.BorderBackgroundColor_Lambda([this, ColumnId] ()
-				{
-					// unclear why this color is coming out much darker
-					return (SelectedColumn && SelectedColumn->Column == ColumnId.GetNumber() - 1) ? FSlateColor(FColor(0x00, 0x70, 0xe0, 0xFF)) : FSlateColor(FLinearColor(0.05f,0.05f,0.05f));
-				})
-				.OnMouseButtonDown_Lambda([this, Chooser, ColumnIndex, ColumnId](	const FGeometry&, const FPointerEvent& PointerEvent)
-				{
-					TableView->ClearSelection();
-				
-					SelectColumn(Chooser, ColumnId.GetNumber() - 1);
-					return FReply::Handled();
-				})
+				SNew(SChooserColumnHandle)
+					.ChooserEditor(this)
+					.ColumnIndex(ColumnIndex)
+					.NoDropAfter(Chooser->ColumnsStructs[ColumnIndex].GetPtr<FRandomizeColumn>() != nullptr)
 				[
-					HeaderWidget.ToSharedRef()
+					SNew(SBorder)
+					.VAlign(VAlign_Center)
+					.Padding(3)
+					.BorderBackgroundColor_Lambda([this, ColumnIndex] ()
+					{
+						// unclear why this color is coming out much darker
+						return (SelectedColumn && SelectedColumn->Column == ColumnIndex) ? FSlateColor(FColor(0x00, 0x70, 0xe0, 0xFF)) : FSlateColor(FLinearColor(0.05f,0.05f,0.05f));
+					})
+					[
+						HeaderWidget.ToSharedRef()
+					]
 				]
 			
 			]);
@@ -1452,6 +1431,8 @@ TSharedRef<FChooserTableEditor> FChooserTableEditor::CreateEditor( const EToolki
 	
 void FChooserTableEditor::SelectColumn(UChooserTable* ChooserEditor, int Index)
 {
+	ClearSelectedRows();
+	
 	UChooserTable* Chooser = GetChooser();
    	if (Index < Chooser->ColumnsStructs.Num())
    	{
@@ -1574,8 +1555,9 @@ bool FChooserTableEditor::CanAutoPopulateSelection()
 {
 	if (UChooserTable* Chooser = GetChooser())
 	{
-		if (SelectedColumn)
+		if (CurrentSelectionType == ESelectionType::Column && SelectedColumn)
 		{
+			// when a column is selected, return true if that column supports auto populate
 			if (Chooser->ColumnsStructs.IsValidIndex(SelectedColumn->Column))
 			{
 				return Chooser->ColumnsStructs[SelectedColumn->Column].Get<FChooserColumnBase>().AutoPopulates();
@@ -1583,7 +1565,20 @@ bool FChooserTableEditor::CanAutoPopulateSelection()
 		}
 		else
 		{
-			return !SelectedRows.IsEmpty();
+			if (SelectedRows.IsEmpty())
+			{
+				return false;
+			}
+			
+			// when rows are selected, return true if any column supports auto populate
+			for (FInstancedStruct& ColumnData : Chooser->ColumnsStructs)
+			{
+				const FChooserColumnBase& Column = ColumnData.Get<FChooserColumnBase>();
+				if (Column.AutoPopulates())
+				{
+					return true;
+				}
+			}
 		}
 	}
 	
@@ -1962,7 +1957,7 @@ void FChooserTableEditor::Paste()
 				ClearSelectedRows();
 				for (int i = 0; i < RowsToPaste; i++)
 				{
-					SelectRow(InsertIndex + i);
+					SelectRow(InsertIndex + i, false);
 				}
 			}
 			
diff --git a/Engine/Plugins/Experimental/Chooser/Source/ChooserEditor/Private/ChooserTableEditor.h b/Engine/Plugins/Experimental/Chooser/Source/ChooserEditor/Private/ChooserTableEditor.h
--- a/Engine/Plugins/Experimental/Chooser/Source/ChooserEditor/Private/ChooserTableEditor.h
+++ b/Engine/Plugins/Experimental/Chooser/Source/ChooserEditor/Private/ChooserTableEditor.h
@@ -124,9 +124,11 @@ namespace UE::ChooserEditor
 		void MakeChoosersMenuRecursive(UObject* Outer, FMenuBuilder& MenuBuilder, const FString& Indent);
 		void DeleteSelectedRows();
 		int MoveRow(int SourceRowIndex, int TargetIndex);
+		int MoveColumn(int SourceIndex, int TargetIndex);
 		void SelectRow(int32 RowIndex, bool bClear = true);
 		void ClearSelectedRows(); 
 		bool IsRowSelected(int32 RowIndex);
+		bool IsColumnSelected(int32 ColumnIndex);
 		
 		void SetChooserTableToEdit(UChooserTable* Chooser, bool bApplyToHistory = true);
 
diff --git a/Engine/Plugins/Experimental/Chooser/Source/ChooserEditor/Private/ChooserTableEditorCommands.h b/Engine/Plugins/Experimental/Chooser/Source/ChooserEditor/Private/ChooserTableEditorCommands.h
--- a/Engine/Plugins/Experimental/Chooser/Source/ChooserEditor/Private/ChooserTableEditorCommands.h
+++ b/Engine/Plugins/Experimental/Chooser/Source/ChooserEditor/Private/ChooserTableEditorCommands.h
@@ -28,7 +28,7 @@ public:
 	{
 		UI_COMMAND(EditChooserSettings, "Table Settings", "Edit the root properties of the ChooserTable asset.", EUserInterfaceActionType::Button, FInputChord())
 		UI_COMMAND(AutoPopulateAll, "AutoPopulate All", "Auto Populate cell data for all supported columns", EUserInterfaceActionType::Button, FInputChord())
-		UI_COMMAND(AutoPopulateSelection, "AutoPopulate", "Auto Populate cell data for selected rows for all supported columns", EUserInterfaceActionType::Button, FInputChord())
+		UI_COMMAND(AutoPopulateSelection, "AutoPopulate", "Auto Populate cell data for selection (requires Columns that support Auto Populate)", EUserInterfaceActionType::Button, FInputChord())
 		UI_COMMAND(RemoveDisabledData, "Remove Disabled Data", "Delete all data that's marked as disabled.", EUserInterfaceActionType::Button, FInputChord())
 		UI_COMMAND(Disable, "Disable", "Disable the selected Rows or Column.", EUserInterfaceActionType::Check, FInputChord())
 	}
diff --git a/Engine/Plugins/Experimental/Chooser/Source/ChooserEditor/Private/SChooserColumnHandle.cpp b/Engine/Plugins/Experimental/Chooser/Source/ChooserEditor/Private/SChooserColumnHandle.cpp
new file mode 100644
--- /dev/null
+++ b/Engine/Plugins/Experimental/Chooser/Source/ChooserEditor/Private/SChooserColumnHandle.cpp
@@ -0,0 +1,153 @@
+// Copyright Epic Games, Inc. All Rights Reserved.
+
+#include "SChooserColumnHandle.h"
+#include "ChooserTableEditor.h"
+#include "RandomizeColumn.h"
+#include "Widgets/Layout/SSeparator.h"
+
+#define LOCTEXT_NAMESPACE "ChooserColumnHandle"
+
+namespace UE::ChooserEditor
+{
+
+	TSharedRef<FChooserColumnDragDropOp> FChooserColumnDragDropOp::New(FChooserTableEditor* InEditor, uint32 InColumnIndex)
+	{
+		TSharedRef<FChooserColumnDragDropOp> Operation = MakeShareable(new FChooserColumnDragDropOp());
+		Operation->ChooserEditor = InEditor;
+		Operation->ColumnIndex = InColumnIndex;
+		Operation->DefaultHoverText = LOCTEXT("Chooser Column", "Chooser Column");
+		Operation->CurrentHoverText = Operation->DefaultHoverText;
+		
+		Operation->Construct();
+
+		return Operation;
+	}
+	
+	void SChooserColumnHandle::Construct(const FArguments& InArgs)
+	{
+		ChooserEditor = InArgs._ChooserEditor;
+		ColumnIndex = InArgs._ColumnIndex;
+		bNoDropAfter = InArgs._NoDropAfter;
+
+		ChildSlot
+		[
+			SNew(SOverlay)
+				+ SOverlay::Slot()
+				[
+					InArgs._Content.Widget
+				]
+				+ SOverlay::Slot().HAlign(HAlign_Right)
+				[
+					SNew(SSeparator).Orientation(Orient_Vertical)
+					.SeparatorImage(FAppStyle::GetBrush("PropertyEditor.VerticalDottedLine"))
+					.ColorAndOpacity_Lambda([this]() { return FSlateColor(bDropSupported ? EStyleColor::Select : EStyleColor::Error); })
+					.Visibility_Lambda([this]() { return bDragActive && !bDropBefore ? EVisibility::Visible : EVisibility::Hidden; })
+				]
+				+ SOverlay::Slot().HAlign(HAlign_Left)
+				[
+					SNew(SSeparator).Orientation(Orient_Vertical)
+					.SeparatorImage(FAppStyle::GetBrush("PropertyEditor.VerticalDottedLine"))
+					.ColorAndOpacity_Lambda([this]() { return FSlateColor(bDropSupported ? EStyleColor::Select : EStyleColor::Error); }) 
+					.Visibility_Lambda([this]() { return bDragActive && bDropBefore ? EVisibility::Visible : EVisibility::Hidden; })
+				]
+		];
+	}
+	
+	FReply SChooserColumnHandle::OnMouseButtonDown(const FGeometry& MyGeometry, const FPointerEvent& MouseEvent)
+	{
+		ChooserEditor->SelectColumn(ChooserEditor->GetChooser(), ColumnIndex);
+		
+		// act as a move handle if the row is already selected, and if there are multiselect modifiers pressed
+		if (MouseEvent.GetEffectingButton() != EKeys::RightMouseButton && !MouseEvent.IsControlDown() && !MouseEvent.IsShiftDown()
+			&& ChooserEditor->IsColumnSelected(ColumnIndex))
+		{
+			UChooserTable* Chooser = ChooserEditor->GetChooser();
+			if (!Chooser->ColumnsStructs[ColumnIndex].GetPtr<FRandomizeColumn>()) // don't allow dragging the Randomize Column
+			{
+				return FReply::Handled().DetectDrag(SharedThis(this), EKeys::LeftMouseButton);
+			}
+		}
+		
+		return SCompoundWidget::OnMouseButtonDown(MyGeometry, MouseEvent);
+	};
+
+	FReply SChooserColumnHandle::OnDragDetected(const FGeometry& MyGeometry, const FPointerEvent& MouseEvent)
+	{
+		// clear row selection so that delete key can't cause the selected row to be deleted
+		ChooserEditor->ClearSelectedColumn();
+		
+		TSharedRef<FChooserColumnDragDropOp> DragDropOp = FChooserColumnDragDropOp::New(ChooserEditor, ColumnIndex);
+		return FReply::Handled().BeginDragDrop(DragDropOp);
+	}
+
+
+	void SChooserColumnHandle::OnDragEnter(const FGeometry& MyGeometry, const FDragDropEvent& DragDropEvent)
+	{
+		bDropSupported = false;
+
+		UChooserTable* Chooser = ChooserEditor->GetChooser();
+		
+		if (!Chooser->ColumnsStructs.IsValidIndex(ColumnIndex))
+		{
+			return;
+		}
+
+		if (TSharedPtr<FChooserColumnDragDropOp> Operation = DragDropEvent.GetOperationAs<FChooserColumnDragDropOp>())
+		{
+			bDropSupported = true;
+		}
+		
+		float Center = MyGeometry.Position.X + MyGeometry.Size.X;
+		bDropBefore = bNoDropAfter || DragDropEvent.GetScreenSpacePosition().X < Center;
+		bDragActive = true;
+	}
+	void SChooserColumnHandle::OnDragLeave(const FDragDropEvent& DragDropEvent)
+	{
+		bDragActive = false;
+	}
+
+	FReply SChooserColumnHandle::OnDragOver(const FGeometry& MyGeometry, const FDragDropEvent& DragDropEvent)
+	{
+		DragActiveCounter = 2;
+		bDragActive = true;
+		float Center = MyGeometry.AbsolutePosition.X + MyGeometry.Size.X/2;
+		bDropBefore = bNoDropAfter || DragDropEvent.GetScreenSpacePosition().X < Center;
+		return FReply::Handled();
+	}
+
+	FReply SChooserColumnHandle::OnDrop(const FGeometry& MyGeometry, const FDragDropEvent& DragDropEvent)
+	{
+		bDragActive = false;
+		
+		if (!bDropSupported)
+		{
+			return FReply::Unhandled();
+		}
+		
+		if (TSharedPtr<FChooserColumnDragDropOp> Operation = DragDropEvent.GetOperationAs<FChooserColumnDragDropOp>())
+		{
+			if (UChooserTable* Chooser = Operation->ChooserEditor->GetChooser())
+			{
+				int NewColumnIndex = 0; 
+				if (bDropBefore)
+				{
+					NewColumnIndex = ChooserEditor->MoveColumn(Operation->ColumnIndex, ColumnIndex);
+				}
+				else
+				{
+					NewColumnIndex = ChooserEditor->MoveColumn(Operation->ColumnIndex, ColumnIndex+1);
+				}
+
+				if (NewColumnIndex >= 0)
+				{
+					ChooserEditor->SelectColumn(Chooser, NewColumnIndex);
+				}
+			}
+		}
+				
+		return FReply::Handled();		
+	}
+
+}
+
+#undef LOCTEXT_NAMESPACE
diff --git a/Engine/Plugins/Experimental/Chooser/Source/ChooserEditor/Private/SChooserColumnHandle.h b/Engine/Plugins/Experimental/Chooser/Source/ChooserEditor/Private/SChooserColumnHandle.h
new file mode 100644
--- /dev/null
+++ b/Engine/Plugins/Experimental/Chooser/Source/ChooserEditor/Private/SChooserColumnHandle.h
@@ -0,0 +1,53 @@
+// Copyright Epic Games, Inc. All Rights Reserved.
+#pragma once
+
+#include "DragAndDrop/DecoratedDragDropOp.h"
+
+namespace UE::ChooserEditor
+{
+
+class FChooserTableEditor;
+
+class FChooserColumnDragDropOp : public FDecoratedDragDropOp
+{
+public:
+	DRAG_DROP_OPERATOR_TYPE(FChooserColumnDragDropOp, FDecoratedDragDropOp)
+
+	FChooserTableEditor* ChooserEditor;
+	uint32 ColumnIndex;
+
+	/** Constructs the drag drop operation */
+	static TSharedRef<FChooserColumnDragDropOp> New(FChooserTableEditor* InEditor, uint32 InColumnIndex);
+};
+
+class SChooserColumnHandle : public SCompoundWidget
+{
+public:
+	SLATE_BEGIN_ARGS(SChooserColumnHandle)
+	{}
+	SLATE_DEFAULT_SLOT(FArguments, Content)
+	SLATE_ARGUMENT(FChooserTableEditor*, ChooserEditor)
+	SLATE_ARGUMENT(uint32, ColumnIndex)
+	SLATE_ARGUMENT(bool, NoDropAfter)
+	SLATE_END_ARGS()
+
+	void Construct(const FArguments& InArgs);
+
+	FReply OnMouseButtonDown(const FGeometry& MyGeometry, const FPointerEvent& MouseEvent) override;
+	FReply OnDragDetected(const FGeometry& MyGeometry, const FPointerEvent& MouseEvent) override;
+	void OnDragEnter(const FGeometry& MyGeometry, const FDragDropEvent& DragDropEvent);
+	void OnDragLeave(const FDragDropEvent& DragDropEvent);
+	FReply OnDragOver(const FGeometry& MyGeometry, const FDragDropEvent& DragDropEvent);
+	FReply OnDrop(const FGeometry& MyGeometry, const FDragDropEvent& DragDropEvent);
+
+private:
+	FChooserTableEditor* ChooserEditor = nullptr;
+	int32 ColumnIndex = 0;
+	int32 DragActiveCounter = 0;
+	bool bDragActive = false;
+	bool bDropBefore = false;
+	bool bDropSupported = false;
+	bool bNoDropAfter = false;
+};
+
+}
diff --git a/Engine/Plugins/Experimental/Chooser/Source/ChooserEditor/Private/SChooserRowHandle.h b/Engine/Plugins/Experimental/Chooser/Source/ChooserEditor/Private/SChooserRowHandle.h
--- a/Engine/Plugins/Experimental/Chooser/Source/ChooserEditor/Private/SChooserRowHandle.h
+++ b/Engine/Plugins/Experimental/Chooser/Source/ChooserEditor/Private/SChooserRowHandle.h
@@ -11,7 +11,7 @@ class FChooserTableEditor;
 class FChooserRowDragDropOp : public FDecoratedDragDropOp
 {
 public:
-	DRAG_DROP_OPERATOR_TYPE(FWidgetTemplateDragDropOp, FDecoratedDragDropOp)
+	DRAG_DROP_OPERATOR_TYPE(FChooserRowDragDropOp, FDecoratedDragDropOp)
 
 	FChooserTableEditor* ChooserEditor;
 	uint32 RowIndex;
diff --git a/Engine/Plugins/Experimental/Chooser/Source/ChooserEditor/Private/SChooserTableRow.cpp b/Engine/Plugins/Experimental/Chooser/Source/ChooserEditor/Private/SChooserTableRow.cpp
--- a/Engine/Plugins/Experimental/Chooser/Source/ChooserEditor/Private/SChooserTableRow.cpp
+++ b/Engine/Plugins/Experimental/Chooser/Source/ChooserEditor/Private/SChooserTableRow.cpp
@@ -36,27 +36,28 @@ namespace UE::ChooserEditor
 		);
 
 		if (RowIndex->RowIndex >=0)
-		SetContent(
-				SNew(SOverlay)
-						+ SOverlay::Slot()
-						[
-							Content.Pin().ToSharedRef()
-						]
-						+ SOverlay::Slot().VAlign(VAlign_Bottom)
-						[
-							SNew(SSeparator)
-							.SeparatorImage(FAppStyle::GetBrush("PropertyEditor.HorizontalDottedLine"))
-							.ColorAndOpacity_Lambda([this]() { return FSlateColor(bDropSupported ? EStyleColor::Select : EStyleColor::Error); })
-							.Visibility_Lambda([this]() { return bDragActive && !bDropAbove ? EVisibility::Visible : EVisibility::Hidden; })
-						]
-						+ SOverlay::Slot().VAlign(VAlign_Top)
-						[
-							SNew(SSeparator)
-							.SeparatorImage(FAppStyle::GetBrush("PropertyEditor.HorizontalDottedLine"))
-							.ColorAndOpacity_Lambda([this]() { return FSlateColor(bDropSupported ? EStyleColor::Select : EStyleColor::Error); }) 
-							.Visibility_Lambda([this]() { return bDragActive && bDropAbove ? EVisibility::Visible : EVisibility::Hidden; })
-						]
-			);
+		{
+			SetContent(SNew(SOverlay)
+					+ SOverlay::Slot()
+					[
+						Content.Pin().ToSharedRef()
+					]
+					+ SOverlay::Slot().VAlign(VAlign_Bottom)
+					[
+						SNew(SSeparator)
+						.SeparatorImage(FAppStyle::GetBrush("PropertyEditor.HorizontalDottedLine"))
+						.ColorAndOpacity_Lambda([this]() { return FSlateColor(bDropSupported ? EStyleColor::Select : EStyleColor::Error); })
+						.Visibility_Lambda([this]() { return bDragActive && !bDropAbove ? EVisibility::Visible : EVisibility::Hidden; })
+					]
+					+ SOverlay::Slot().VAlign(VAlign_Top)
+					[
+						SNew(SSeparator)
+						.SeparatorImage(FAppStyle::GetBrush("PropertyEditor.HorizontalDottedLine"))
+						.ColorAndOpacity_Lambda([this]() { return FSlateColor(bDropSupported ? EStyleColor::Select : EStyleColor::Error); }) 
+						.Visibility_Lambda([this]() { return bDragActive && bDropAbove ? EVisibility::Visible : EVisibility::Hidden; })
+					]
+				);
+		}
 		else if (RowIndex->RowIndex == SpecialIndex_Fallback)
 		{
 			SetContent(
@@ -337,18 +338,11 @@ namespace UE::ChooserEditor
 
 	FReply SChooserTableRow::OnDragOver(const FGeometry& MyGeometry, const FDragDropEvent& DragDropEvent)
 	{
-		if (bDropSupported)
-		{
-			DragActiveCounter = 2;
-			bDragActive = true;
-			float Center = MyGeometry.AbsolutePosition.Y + MyGeometry.Size.Y/2;
-			bDropAbove = DragDropEvent.GetScreenSpacePosition().Y < Center;
-			return FReply::Handled();
-		}
-		else
-		{
-			return FReply::Unhandled();
-		}
+		DragActiveCounter = 2;
+		bDragActive = true;
+		float Center = MyGeometry.AbsolutePosition.Y + MyGeometry.Size.Y/2;
+		bDropAbove = DragDropEvent.GetScreenSpacePosition().Y < Center;
+		return FReply::Handled();
 	}
 
 	FReply SChooserTableRow::OnDrop(const FGeometry& MyGeometry, const FDragDropEvent& DragDropEvent)
