Added descriptions to chooser assets. Update plugin to be not experimental, and proxy to experimental.

diff --git a/Engine/Plugins/Experimental/Chooser/Chooser.uplugin b/Engine/Plugins/Experimental/Chooser/Chooser.uplugin
--- a/Engine/Plugins/Experimental/Chooser/Chooser.uplugin
+++ b/Engine/Plugins/Experimental/Chooser/Chooser.uplugin
@@ -3,7 +3,7 @@
 	"Version" : 1,
 	"VersionName" : "1.0",
 	"FriendlyName" : "Chooser",
-	"Description" : "Chooser",
+	"Description" : "Use Chooser and Proxy Tables to build dynamic asset selection logic.",
 	"Category" : "Animation",
 	"CreatedBy" : "Epic Games, Inc.",
 	"CreatedByURL" : "https://epicgames.com",
@@ -12,8 +12,6 @@
 	"SupportURL" : "",
 	"EnabledByDefault" : false,
 	"CanContainContent" : true,
-	"IsBetaVersion" : true,
-	"IsExperimentalVersion" : false,
 	"Installed" : false,
 	"Plugins": [
 		{
diff --git a/Engine/Plugins/Experimental/Chooser/Source/Chooser/Internal/AnimNode_ChooserPlayer.h b/Engine/Plugins/Experimental/Chooser/Source/Chooser/Internal/AnimNode_ChooserPlayer.h
--- a/Engine/Plugins/Experimental/Chooser/Source/Chooser/Internal/AnimNode_ChooserPlayer.h
+++ b/Engine/Plugins/Experimental/Chooser/Source/Chooser/Internal/AnimNode_ChooserPlayer.h
@@ -134,7 +134,7 @@ public:
 	UPROPERTY()
 	TArray<FInstancedStruct> ChooserContextDefinition;
 
-	// Use pose matching to choose the start position. Requires experimental PoseSearch plugin.
+	// Use pose matching to choose the start position. Requires PoseSearch plugin.
 	UPROPERTY(EditAnywhere, Category = PoseMatching, meta = (PinHiddenByDefault))
 	bool bStartFromMatchingPose = false;
 
diff --git a/Engine/Plugins/Experimental/Chooser/Source/Chooser/Internal/Chooser.h b/Engine/Plugins/Experimental/Chooser/Source/Chooser/Internal/Chooser.h
--- a/Engine/Plugins/Experimental/Chooser/Source/Chooser/Internal/Chooser.h
+++ b/Engine/Plugins/Experimental/Chooser/Source/Chooser/Internal/Chooser.h
@@ -13,7 +13,9 @@
 
 #include "Chooser.generated.h"
 
-
+/**
+* Data table used to choose an asset based on input parameters
+*/
 UCLASS(BlueprintType)
 class CHOOSER_API UChooserTable : public UObject, public IHasContextClass
 {
diff --git a/Engine/Plugins/Experimental/Chooser/Source/ProxyTable/Public/ProxyAsset.h b/Engine/Plugins/Experimental/Chooser/Source/ProxyTable/Public/ProxyAsset.h
--- a/Engine/Plugins/Experimental/Chooser/Source/ProxyTable/Public/ProxyAsset.h
+++ b/Engine/Plugins/Experimental/Chooser/Source/ProxyTable/Public/ProxyAsset.h
@@ -11,8 +11,10 @@
 
 DECLARE_MULTICAST_DELEGATE_OneParam(FProxyTypeChanged, const UClass* OutputObjectType);
 
-
-UCLASS(MinimalAPI,BlueprintType)
+/**
+* EXPERIMENTAL: An indirect mapping to an asset. The direct mapping is defined in a Proxy Table.
+*/
+UCLASS(MinimalAPI,BlueprintType,Experimental)
 class UProxyAsset : public UObject, public IHasContextClass
 {
 	GENERATED_UCLASS_BODY()
diff --git a/Engine/Plugins/Experimental/Chooser/Source/ProxyTable/Public/ProxyTable.h b/Engine/Plugins/Experimental/Chooser/Source/ProxyTable/Public/ProxyTable.h
--- a/Engine/Plugins/Experimental/Chooser/Source/ProxyTable/Public/ProxyTable.h
+++ b/Engine/Plugins/Experimental/Chooser/Source/ProxyTable/Public/ProxyTable.h
@@ -66,7 +66,10 @@ struct PROXYTABLE_API FRuntimeProxyValue
 	TArray<FProxyStructOutput> OutputStructData;
 };
 
-UCLASS(BlueprintType)
+/**
+* EXPERIMENTAL: Table mapping of proxy assets to a specific asset.
+*/
+UCLASS(BlueprintType, Experimental)
 class PROXYTABLE_API UProxyTable : public UObject
 {
 	GENERATED_UCLASS_BODY()
