Adding support in Chooser Table Editor for  Dragging multiple rows, and Alt+ arrow keys to move rows or columns.

diff --git a/Engine/Plugins/Experimental/Chooser/Source/ChooserEditor/Private/ChooserTableEditor.cpp b/Engine/Plugins/Experimental/Chooser/Source/ChooserEditor/Private/ChooserTableEditor.cpp
--- a/Engine/Plugins/Experimental/Chooser/Source/ChooserEditor/Private/ChooserTableEditor.cpp
+++ b/Engine/Plugins/Experimental/Chooser/Source/ChooserEditor/Private/ChooserTableEditor.cpp
@@ -138,7 +138,7 @@ void FChooserTableEditor::SetChooserTableToEdit(UChooserTable* Chooser, bool bAp
 	
 	RefreshAll();
 }
-	
+
 void FChooserTableEditor::PushChooserTableToEdit(UChooserTable* Chooser)
 {
 	BreadcrumbTrail->PushCrumb(FText::FromString(Chooser->GetName()), Chooser);
@@ -393,6 +393,10 @@ void FChooserTableEditor::RegisterMenus()
 		Section.AddEntry(FToolMenuEntry::InitMenuEntry(FGenericCommands::Get().Paste));
 		Section.AddEntry(FToolMenuEntry::InitMenuEntry(FGenericCommands::Get().Delete));
 		Section.AddEntry(FToolMenuEntry::InitMenuEntry(Commands.Disable));
+		Section.AddEntry(FToolMenuEntry::InitMenuEntry(Commands.MoveUp));
+		Section.AddEntry(FToolMenuEntry::InitMenuEntry(Commands.MoveDown));
+		Section.AddEntry(FToolMenuEntry::InitMenuEntry(Commands.MoveLeft));
+		Section.AddEntry(FToolMenuEntry::InitMenuEntry(Commands.MoveRight));
 		Section.AddEntry(FToolMenuEntry::InitMenuEntry(Commands.AutoPopulateSelection));
 		
 
@@ -529,6 +533,37 @@ void FChooserTableEditor::BindCommands()
 		FGenericCommands::Get().Paste,
 		FExecuteAction::CreateSP(this, &FChooserTableEditor::Paste),
 		FCanExecuteAction::CreateSP(this, &FChooserTableEditor::CanPaste));
+	
+	ToolkitCommands->MapAction(
+		Commands.MoveUp,
+		FExecuteAction::CreateSP(this, &FChooserTableEditor::MoveRowsUp),
+		FCanExecuteAction::CreateSP(this, &FChooserTableEditor::CanMoveRowsUp),
+		FIsActionChecked(),
+		FIsActionButtonVisible::CreateSP(this, &FChooserTableEditor::HasRowsSelected)
+		);
+	ToolkitCommands->MapAction(
+    		Commands.MoveDown,
+    		FExecuteAction::CreateSP(this, &FChooserTableEditor::MoveRowsDown),
+    		FCanExecuteAction::CreateSP(this, &FChooserTableEditor::CanMoveRowsDown),
+    		FIsActionChecked(),
+    		FIsActionButtonVisible::CreateSP(this, &FChooserTableEditor::HasRowsSelected)
+    		);
+	
+	ToolkitCommands->MapAction(
+			Commands.MoveLeft,
+			FExecuteAction::CreateSP(this, &FChooserTableEditor::MoveColumnLeft),
+			FCanExecuteAction::CreateSP(this, &FChooserTableEditor::CanMoveColumnLeft),
+			FIsActionChecked(),
+			FIsActionButtonVisible::CreateSP(this, &FChooserTableEditor::HasColumnSelected)
+			);
+
+	ToolkitCommands->MapAction(
+			Commands.MoveRight,
+			FExecuteAction::CreateSP(this, &FChooserTableEditor::MoveColumnRight),
+			FCanExecuteAction::CreateSP(this, &FChooserTableEditor::CanMoveColumnRight),
+			FIsActionChecked(),
+			FIsActionButtonVisible::CreateSP(this, &FChooserTableEditor::HasColumnSelected)
+			);
 }
 
 void FChooserTableEditor::OnObjectsTransacted(UObject* Object, const FTransactionObjectEvent& Event)
@@ -1481,9 +1516,14 @@ void FChooserTableEditor::DeleteColumn(int Index)
 	}
 }
 	
-void FChooserTableEditor::DeleteSelectedRows()
+int32 FChooserTableEditor::DeleteSelectedRows(int32 RowIndexToRemember)
 {
 	const FScopedTransaction Transaction(LOCTEXT("Delete Row Transaction", "Delete Row"));
+	return DeleteSelectedRowsInternal(RowIndexToRemember);
+}
+	
+int32 FChooserTableEditor::DeleteSelectedRowsInternal(int32 RowIndexToRemember)
+{
 	UChooserTable* Chooser = GetChooser();
 	Chooser->Modify(true);
 	// delete selected rows.
@@ -1498,8 +1538,12 @@ void FChooserTableEditor::DeleteSelectedRows()
 
 	// sort indices in reverse
 	RowsToDelete.Sort([](int32 A, int32 B){ return A>B; });
-	for(uint32 RowIndex : RowsToDelete)
+	for(int32 RowIndex : RowsToDelete)
 	{
+		if (RowIndex <= RowIndexToRemember)
+		{
+			RowIndexToRemember--;
+		}
 		Chooser->ResultsStructs.RemoveAt(RowIndex);
 		Chooser->DisabledRows.RemoveAt(RowIndex);
 	}
@@ -1510,8 +1554,18 @@ void FChooserTableEditor::DeleteSelectedRows()
 		Column.DeleteRows(RowsToDelete);
 	}
 	UpdateTableRows();
+
+	return RowIndexToRemember;
 }
-	
+
+void FChooserTableEditor::MoveRows(int TargetIndex)
+{
+	const FScopedTransaction Transaction(LOCTEXT("Move Row(s)", "Move Row(s)"));
+	UChooserTable* RowCopy = CopySelectionInternal();
+	TargetIndex = DeleteSelectedRowsInternal(TargetIndex);
+	PasteInternal(RowCopy, TargetIndex);
+}
+
 void FChooserTableEditor::AutoPopulateColumn(FChooserColumnBase& Column)
 {
 	if (UChooserTable* Chooser = GetChooser())
@@ -1622,39 +1676,36 @@ void FChooserTableEditor::AutoPopulateAll()
 	}
 }
 
-	
-inline bool FChooserTableEditor::HasSelection()
+bool FChooserTableEditor::HasSelection() const
 {
-	if (CurrentSelectionType == ESelectionType::Column)
-	{
-		return true;
-	}
-	else if (CurrentSelectionType == ESelectionType::Rows)
-	{
-		return !SelectedRows.IsEmpty();
-	}
-	return false;
+	return HasRowsSelected() || HasColumnSelected();
+}
+
+bool FChooserTableEditor::HasRowsSelected() const
+{
+	return CurrentSelectionType == ESelectionType::Rows && !SelectedRows.IsEmpty();
 }
 	
-inline bool FChooserTableEditor::IsSelectionDisabled()
+bool FChooserTableEditor::HasColumnSelected() const
 {
-	if (UChooserTable* Chooser = GetChooser())
+	const UChooserTable* Chooser = GetChooser();
+	return CurrentSelectionType == ESelectionType::Column && SelectedColumn && Chooser->ColumnsStructs.IsValidIndex(SelectedColumn->Column);
+}
+
+bool FChooserTableEditor::IsSelectionDisabled() const
+{
+	if (const UChooserTable* Chooser = GetChooser())
 	{
-		if (CurrentSelectionType == ESelectionType::Column)
+		if (HasColumnSelected())
 		{
-			if (SelectedColumn && Chooser->ColumnsStructs.IsValidIndex(SelectedColumn->Column))
+			if ( Chooser->ColumnsStructs.IsValidIndex(SelectedColumn->Column))
 			{
-				FChooserColumnBase& Column = Chooser->ColumnsStructs[SelectedColumn->Column].GetMutable<FChooserColumnBase>();
+				const FChooserColumnBase& Column = Chooser->ColumnsStructs[SelectedColumn->Column].Get<FChooserColumnBase>();
 				return Column.bDisabled;
 			}
 		}
-		else if (CurrentSelectionType == ESelectionType::Rows)
+		else if (HasRowsSelected())
 		{
-			if (SelectedRows.IsEmpty())
-			{
-				return false;
-			}
-
 			bool bSomethingEnabled = false;
 			for(auto& Row : SelectedRows)
 			{
@@ -1675,15 +1726,15 @@ void FChooserTableEditor::ToggleDisableSelection()
 	bool bDisabled = IsSelectionDisabled();
 	if (UChooserTable* Chooser = GetChooser())
 	{
-		if (CurrentSelectionType == ESelectionType::Column)
+		if (HasColumnSelected())
 		{
-			if (SelectedColumn && Chooser->ColumnsStructs.IsValidIndex(SelectedColumn->Column))
+			if (Chooser->ColumnsStructs.IsValidIndex(SelectedColumn->Column))
 			{
 				FChooserColumnBase& Column = Chooser->ColumnsStructs[SelectedColumn->Column].GetMutable<FChooserColumnBase>();
 				Column.bDisabled = !Column.bDisabled;
 			}
 		}
-		else if (CurrentSelectionType == ESelectionType::Rows)
+		else if (HasRowsSelected())
 		{
 			for (auto& Row : SelectedRows)
 			{
@@ -1696,18 +1747,131 @@ void FChooserTableEditor::ToggleDisableSelection()
 
 void FChooserTableEditor::DeleteSelection()
 {
-	if (CurrentSelectionType == ESelectionType::Column)
+	if (HasColumnSelected())
 	{
 		DeleteColumn(SelectedColumn->Column);
 	}
-	else if (CurrentSelectionType == ESelectionType::Rows)
+	else if (HasRowsSelected())
 	{
 		DeleteSelectedRows();
 	}
 }
+
+bool FChooserTableEditor::CanMoveRowsUp()
+{
+	if (HasRowsSelected())
+	{
+		UChooserTable* Chooser = GetChooser();
+		
+		int MinSelectedRow = Chooser->ResultsStructs.Num();
+		for(UChooserRowDetails* SelectedRow : SelectedRows)
+		{
+			MinSelectedRow = FMath::Min(SelectedRow->Row, MinSelectedRow);
+		}
+
+		return MinSelectedRow > 0;
+	}
+	return false;
+}
+
+void FChooserTableEditor::MoveRowsUp()
+{
+	if (HasRowsSelected())
+	{
+		UChooserTable* Chooser = GetChooser();
+		int MinSelectedRow = Chooser->ResultsStructs.Num();
+		for(UChooserRowDetails* SelectedRow : SelectedRows)
+		{
+			MinSelectedRow = FMath::Min(SelectedRow->Row, MinSelectedRow);
+		}
+		MoveRows(MinSelectedRow - 1);
+	}
+}
+
+bool FChooserTableEditor::CanMoveRowsDown()
+{
+	if (HasRowsSelected())
+	{
+		UChooserTable* Chooser = GetChooser();
+		
+		int MaxSelectedRow = -1;
+		for(UChooserRowDetails* SelectedRow : SelectedRows)
+		{
+			MaxSelectedRow = FMath::Max(SelectedRow->Row, MaxSelectedRow);
+		}
+
+		return MaxSelectedRow < Chooser->ResultsStructs.Num() - 1;
+	}
+	return false;
+}
+
+void FChooserTableEditor::MoveRowsDown()
+{
+	if (HasRowsSelected())
+   	{
+   		UChooserTable* Chooser = GetChooser();
+   		int MaxSelectedRow = -1;
+   		for(UChooserRowDetails* SelectedRow : SelectedRows)
+   		{
+   			MaxSelectedRow = FMath::Max(SelectedRow->Row, MaxSelectedRow);
+   		}
+   		MoveRows(MaxSelectedRow + 2);
+   	}
+}
+
+bool FChooserTableEditor::CanMoveColumnLeft()
+{
+	if (HasColumnSelected())
+	{
+		UChooserTable* Chooser = GetChooser();
+		
+		if (Chooser->ColumnsStructs[SelectedColumn->Column].GetPtr<FRandomizeColumn>())
+		{
+			return false;
+		}
+		
+		return SelectedColumn->Column > 0;
+	}
+	return false;
+}
+
+void FChooserTableEditor::MoveColumnLeft()
+{
+	if (CanMoveColumnLeft())
+	{
+		SelectColumn(GetChooser(), MoveColumn(SelectedColumn->Column, SelectedColumn->Column - 1));
+	}
+}
+
+bool FChooserTableEditor::CanMoveColumnRight()
+{
+	if (HasColumnSelected())
+   	{
+		UChooserTable* Chooser = GetChooser();
+		
+		if (Chooser->ColumnsStructs[SelectedColumn->Column].GetPtr<FRandomizeColumn>())
+		{
+			return false;
+		}
+		int NumColumns = Chooser->ColumnsStructs.Num();
+		if (NumColumns > 0 && Chooser->ColumnsStructs.Last().GetPtr<FRandomizeColumn>())
+		{
+			NumColumns--;
+		}
+   		return (SelectedColumn->Column < NumColumns-1);
+   	}
+   	return false;
+}
+
+void FChooserTableEditor::MoveColumnRight()
+{
+	if (CanMoveColumnRight())
+	{
+		SelectColumn(GetChooser(), MoveColumn(SelectedColumn->Column, SelectedColumn->Column + 2));
+	}
+}
 	
-	
-void FChooserTableEditor::CopySelection()
+UChooserTable* FChooserTableEditor::CopySelectionInternal()
 {
 	UChooserTable* CopyData = NewObject<UChooserTable>(GetTransientPackage());
 
@@ -1725,6 +1889,10 @@ void FChooserTableEditor::CopySelection()
 	}
 	else if (CurrentSelectionType == ESelectionType::Rows)
 	{
+
+		TArray<UChooserRowDetails*> SelectedRowsCopy = SelectedRows;
+		Algo::Sort(SelectedRowsCopy, [](UChooserRowDetails *A, UChooserRowDetails *B) { return A->Row<B->Row;});
+		
 		// copy all columns
 		CopyData->ColumnsStructs = Chooser->ColumnsStructs;
 		
@@ -1733,26 +1901,32 @@ void FChooserTableEditor::CopySelection()
 		{
 			FChooserColumnBase& Column = ColumnData.GetMutable<FChooserColumnBase>();
 			Column.SetNumRows(0);
-			Column.SetNumRows(SelectedRows.Num());
+			Column.SetNumRows(SelectedRowsCopy.Num());
 		}
 
-		CopyData->ResultsStructs.SetNum(SelectedRows.Num());
+		CopyData->ResultsStructs.SetNum(SelectedRowsCopy.Num());
 
 		// add the selected results and column data
 		
-		for (int RowIndex = 0; RowIndex < SelectedRows.Num(); RowIndex++)
+		for (int RowIndex = 0; RowIndex < SelectedRowsCopy.Num(); RowIndex++)
 		{
-			CopyData->ResultsStructs[RowIndex] = Chooser->ResultsStructs[SelectedRows[RowIndex]->Row];
+			CopyData->ResultsStructs[RowIndex] = Chooser->ResultsStructs[SelectedRowsCopy[RowIndex]->Row];
 
 			for (int ColumnIndex = 0; ColumnIndex < CopyData->ColumnsStructs.Num(); ColumnIndex++)
 			{
 				FChooserColumnBase& SourceColumn = Chooser->ColumnsStructs[ColumnIndex].GetMutable<FChooserColumnBase>();
 				FChooserColumnBase& TargetColumn = CopyData->ColumnsStructs[ColumnIndex].GetMutable<FChooserColumnBase>();
-				TargetColumn.CopyRow(SourceColumn, SelectedRows[RowIndex]->Row, RowIndex);
+				TargetColumn.CopyRow(SourceColumn, SelectedRowsCopy[RowIndex]->Row, RowIndex);
 			}
-
 		}
 	}
+	
+	return CopyData;
+}
+
+void FChooserTableEditor::CopySelection()
+{
+	UChooserTable* CopyData = CopySelectionInternal();
 
 	// Clear the mark state for saving.
 	UnMarkAllObjects(EObjectMark(OBJECTMARK_TagExp | OBJECTMARK_TagImp));
@@ -1811,157 +1985,172 @@ static FString GetColumnName(FChooserColumnBase& Column)
 	return FString();
 }
 	
-void FChooserTableEditor::Paste()
+void FChooserTableEditor::PasteInternal(UChooserTable* PastedContent, int PasteRowIndex)
 {
-	// Get the text from the clipboard.
-	FString ClipboardText;
-	FPlatformApplicationMisc::ClipboardPaste(ClipboardText);
-	FChooserClipboardFactory Factory;
-
 	UChooserTable* Chooser = GetChooser();
-
-	if (Factory.CanCreateObjectsFromText(ClipboardText))
+	Chooser->Modify();
+	
+	if (PastedContent->ResultsStructs.IsEmpty())
 	{
-		Factory.ProcessBuffer((UObject*)GetTransientPackage(), RF_Transactional, ClipboardText);
-		if (UChooserTable* PastedContent = Factory.ClipboardContent)
+		// pasting a column
+		int InsertColumnIndex = Chooser->ColumnsStructs.Num();
+		if (CurrentSelectionType == ESelectionType::Column && SelectedColumn)
 		{
-			FScopedTransaction Transaction(LOCTEXT("Paste Chooser Data", "Paste Chooser Data"));
-			Chooser->Modify();
-			
-			if (PastedContent->ResultsStructs.IsEmpty())
+			InsertColumnIndex = FMath::Min(InsertColumnIndex, SelectedColumn->Column + 1);
+		}
+		
+		if (Chooser->ColumnsStructs.Num() > 0 && Chooser->ColumnsStructs.Num() == InsertColumnIndex)
+		{
+			// if were inserting at the end, there is a randomize column, insert new columns before it
+			if (Chooser->ColumnsStructs.Last().GetPtr<FRandomizeColumn>())
 			{
-				// pasting a column
-				int InsertColumnIndex = Chooser->ColumnsStructs.Num();
-				if (CurrentSelectionType == ESelectionType::Column && SelectedColumn)
-				{
-					InsertColumnIndex = FMath::Min(InsertColumnIndex, SelectedColumn->Column + 1);
-				}
-				
-				if (Chooser->ColumnsStructs.Num() > 0 && Chooser->ColumnsStructs.Num() == InsertColumnIndex)
-				{
-					// if were inserting at the end, there is a randomize column, insert new columns before it
-					if (Chooser->ColumnsStructs.Last().GetPtr<FRandomizeColumn>())
-					{
-						InsertColumnIndex--;
-					}
-				}
-				Chooser->ColumnsStructs.Insert(PastedContent->ColumnsStructs, InsertColumnIndex);
-				SelectColumn(Chooser, InsertColumnIndex);
+				InsertColumnIndex--;
 			}
-			else
-			{
-				// pasting rows
-				int RowsToPaste = PastedContent->ResultsStructs.Num();
+		}
+		Chooser->ColumnsStructs.Insert(PastedContent->ColumnsStructs, InsertColumnIndex);
+		SelectColumn(Chooser, InsertColumnIndex);
+	}
+	else
+	{
+		// pasting rows
+		int RowsToPaste = PastedContent->ResultsStructs.Num();
 
-				// figure out where to start inserting
-				int InsertIndex = Chooser->ResultsStructs.Num();
-				if (SelectedRows.Num() > 0)
-				{
-					InsertIndex = SelectedRows[0]->Row;
-					for(int SelectedRowIndex = 1; SelectedRowIndex < SelectedRows.Num(); SelectedRowIndex ++)
-					{
-						InsertIndex = FMath::Max(InsertIndex, SelectedRows[SelectedRowIndex]->Row);
-					}
-					InsertIndex++;
-				}
+		// figure out where to start inserting
+		int InsertIndex = Chooser->ResultsStructs.Num();
 
-				Chooser->ResultsStructs.Insert(PastedContent->ResultsStructs, InsertIndex);
-				
-				// Make sure each column has the same number of row datas as there are results
-				for(FInstancedStruct& ColumnData : Chooser->ColumnsStructs)
+		if (PasteRowIndex >=0)
+		{
+			InsertIndex = PasteRowIndex;
+		}
+		else
+		{
+			if (SelectedRows.Num() > 0)
+			{
+				InsertIndex = SelectedRows[0]->Row;
+				for(int SelectedRowIndex = 1; SelectedRowIndex < SelectedRows.Num(); SelectedRowIndex ++)
 				{
-					FChooserColumnBase& Column = ColumnData.GetMutable<FChooserColumnBase>();
-					Column.InsertRows(InsertIndex, RowsToPaste);
+					InsertIndex = FMath::Max(InsertIndex, SelectedRows[SelectedRowIndex]->Row);
 				}
+				InsertIndex++;
+			}
+		}
+
+		Chooser->ResultsStructs.Insert(PastedContent->ResultsStructs, InsertIndex);
+		
+		// Make sure each column has the same number of row datas as there are results
+		for(FInstancedStruct& ColumnData : Chooser->ColumnsStructs)
+		{
+			FChooserColumnBase& Column = ColumnData.GetMutable<FChooserColumnBase>();
+			Column.InsertRows(InsertIndex, RowsToPaste);
+		}
 
 
-				// try to also paste column data from columns in the paste buffer which match the columns in the current chooser
-				// -- matching by column type and input value name
+		// try to also paste column data from columns in the paste buffer which match the columns in the current chooser
+		// -- matching by column type and input value name
 
-				// keep track of target columns that have already been matched, to avoid matching multiple source columns with the same target column
-				TArray<bool> MatchedTargetColumns;
-				MatchedTargetColumns.SetNum(Chooser->ColumnsStructs.Num());
+		// keep track of target columns that have already been matched, to avoid matching multiple source columns with the same target column
+		TArray<bool> MatchedTargetColumns;
+		MatchedTargetColumns.SetNum(Chooser->ColumnsStructs.Num());
 
-				// keep track of which source columns were matched, so we can add new columns for the unmatched ones after
-				TArray<bool> MatchedSourceColumns;
-				MatchedSourceColumns.SetNum(PastedContent->ColumnsStructs.Num());
-				
-				for(int SourceColumnIndex = 0; SourceColumnIndex < PastedContent->ColumnsStructs.Num(); SourceColumnIndex++)
+		// keep track of which source columns were matched, so we can add new columns for the unmatched ones after
+		TArray<bool> MatchedSourceColumns;
+		MatchedSourceColumns.SetNum(PastedContent->ColumnsStructs.Num());
+		
+		for(int SourceColumnIndex = 0; SourceColumnIndex < PastedContent->ColumnsStructs.Num(); SourceColumnIndex++)
+		{
+			FInstancedStruct& PastedColumnData = PastedContent->ColumnsStructs[SourceColumnIndex];
+			FChooserColumnBase& PastedColumn = PastedColumnData.GetMutable<FChooserColumnBase>();
+			FString PastedColumnName = GetColumnName(PastedColumn);
+			if (FChooserParameterBase* PastedColumnInput = PastedColumn.GetInputValue())
+			{
+				for(int TargetColumnIndex = 0; TargetColumnIndex < Chooser->ColumnsStructs.Num(); TargetColumnIndex++)
 				{
-					FInstancedStruct& PastedColumnData = PastedContent->ColumnsStructs[SourceColumnIndex];
-					FChooserColumnBase& PastedColumn = PastedColumnData.GetMutable<FChooserColumnBase>();
-					FString PastedColumnName = GetColumnName(PastedColumn);
-					if (FChooserParameterBase* PastedColumnInput = PastedColumn.GetInputValue())
+					if (!MatchedTargetColumns[TargetColumnIndex])
 					{
-						for(int TargetColumnIndex = 0; TargetColumnIndex < Chooser->ColumnsStructs.Num(); TargetColumnIndex++)
+						FInstancedStruct& ColumnData = Chooser->ColumnsStructs[TargetColumnIndex];
+						if (ColumnData.GetScriptStruct() == PastedColumnData.GetScriptStruct())
 						{
-							if (!MatchedTargetColumns[TargetColumnIndex])
+							FChooserColumnBase& Column = ColumnData.GetMutable<FChooserColumnBase>();
+							FString ColumnName = GetColumnName(Column);
+
+							if (ColumnName == PastedColumnName)
 							{
-								FInstancedStruct& ColumnData = Chooser->ColumnsStructs[TargetColumnIndex];
-								if (ColumnData.GetScriptStruct() == PastedColumnData.GetScriptStruct())
-								{
-									FChooserColumnBase& Column = ColumnData.GetMutable<FChooserColumnBase>();
-									FString ColumnName = GetColumnName(Column);
+								MatchedTargetColumns[TargetColumnIndex] = true;
+								MatchedSourceColumns[SourceColumnIndex] = true;
 
-									if (ColumnName == PastedColumnName)
-									{
-										MatchedTargetColumns[TargetColumnIndex] = true;
-										MatchedSourceColumns[SourceColumnIndex] = true;
-
-										// found a match, copy the data over
-										for (int i = 0; i < RowsToPaste; i++)
-										{
-											Column.CopyRow(PastedColumn, i, InsertIndex + i);
-										}
-										break;
-									}
+								// found a match, copy the data over
+								for (int i = 0; i < RowsToPaste; i++)
+								{
+									Column.CopyRow(PastedColumn, i, InsertIndex + i);
 								}
+								break;
 							}
 						}
 					}
 				}
-				
-				// add new columns for any source columns that were unmatched
-				
-				int InsertColumnIndex = Chooser->ColumnsStructs.Num();
-				if (Chooser->ColumnsStructs.Num() > 0)
-				{
-					// if there is a randomize column, insert new columns before it
-					if (Chooser->ColumnsStructs.Last().GetPtr<FRandomizeColumn>())
-					{
-						InsertColumnIndex--;
-					}
-				}
-
-				
-				for(int SourceColumnIndex = 0; SourceColumnIndex < PastedContent->ColumnsStructs.Num(); SourceColumnIndex++)
-				{
-					if (!MatchedSourceColumns[SourceColumnIndex])
-					{
-						FInstancedStruct& PastedColumnData = PastedContent->ColumnsStructs[SourceColumnIndex];
-						FChooserColumnBase& PastedColumn = PastedColumnData.GetMutable<FChooserColumnBase>();
-						// if we couldn't find a match, paste a new column
-						Chooser->ColumnsStructs.Insert(PastedColumnData, InsertColumnIndex);
-						FChooserColumnBase& Column = Chooser->ColumnsStructs[InsertColumnIndex].GetMutable<FChooserColumnBase>();
-						InsertColumnIndex++;
-						Column.SetNumRows(0);
-						Column.SetNumRows(Chooser->ResultsStructs.Num());
-						for (int i = 0; i < RowsToPaste; i++)
-						{
-							Column.CopyRow(PastedColumn, i, InsertIndex + i);
-						}
-					}
-				}
+			}
+		}
+		
+		// add new columns for any source columns that were unmatched
+		
+		int InsertColumnIndex = Chooser->ColumnsStructs.Num();
+		if (Chooser->ColumnsStructs.Num() > 0)
+		{
+			// if there is a randomize column, insert new columns before it
+			if (Chooser->ColumnsStructs.Last().GetPtr<FRandomizeColumn>())
+			{
+				InsertColumnIndex--;
+			}
+		}
 
-				// select the inserted rows
-				ClearSelectedRows();
+		
+		for(int SourceColumnIndex = 0; SourceColumnIndex < PastedContent->ColumnsStructs.Num(); SourceColumnIndex++)
+		{
+			if (!MatchedSourceColumns[SourceColumnIndex])
+			{
+				FInstancedStruct& PastedColumnData = PastedContent->ColumnsStructs[SourceColumnIndex];
+				FChooserColumnBase& PastedColumn = PastedColumnData.GetMutable<FChooserColumnBase>();
+				// if we couldn't find a match, paste a new column
+				Chooser->ColumnsStructs.Insert(PastedColumnData, InsertColumnIndex);
+				FChooserColumnBase& Column = Chooser->ColumnsStructs[InsertColumnIndex].GetMutable<FChooserColumnBase>();
+				InsertColumnIndex++;
+				Column.SetNumRows(0);
+				Column.SetNumRows(Chooser->ResultsStructs.Num());
 				for (int i = 0; i < RowsToPaste; i++)
 				{
-					SelectRow(InsertIndex + i, false);
+					Column.CopyRow(PastedColumn, i, InsertIndex + i);
 				}
 			}
-			
-			RefreshAll();
+		}
+
+		RefreshAll();
+		// select the inserted rows
+		ClearSelectedRows();
+		for (int i = 0; i < RowsToPaste; i++)
+		{
+			SelectRow(InsertIndex + i, false);
+		}
+	}
+	
+	RefreshAll();
+}
+	
+void FChooserTableEditor::Paste()
+{
+	// Get the text from the clipboard.
+	FString ClipboardText;
+	FPlatformApplicationMisc::ClipboardPaste(ClipboardText);
+	FChooserClipboardFactory Factory;
+
+	UChooserTable* Chooser = GetChooser();
+
+	if (Factory.CanCreateObjectsFromText(ClipboardText))
+	{
+		Factory.ProcessBuffer((UObject*)GetTransientPackage(), RF_Transactional, ClipboardText);
+		if (UChooserTable* PastedContent = Factory.ClipboardContent)
+		{
+			FScopedTransaction Transaction(LOCTEXT("Paste Chooser Data", "Paste Chooser Data"));
+			PasteInternal(PastedContent);
 		}
 	}
 }
diff --git a/Engine/Plugins/Experimental/Chooser/Source/ChooserEditor/Private/ChooserTableEditor.h b/Engine/Plugins/Experimental/Chooser/Source/ChooserEditor/Private/ChooserTableEditor.h
--- a/Engine/Plugins/Experimental/Chooser/Source/ChooserEditor/Private/ChooserTableEditor.h
+++ b/Engine/Plugins/Experimental/Chooser/Source/ChooserEditor/Private/ChooserTableEditor.h
@@ -99,10 +99,21 @@ namespace UE::ChooserEditor
 		/** Can be used to disable the details view making it read-only */
 		void SetPropertyEditingEnabledDelegate(FIsPropertyEditingEnabled InPropertyEditingDelegate);
 
-		bool HasSelection();
-		bool IsSelectionDisabled();
+		bool HasSelection() const;
+		bool HasRowsSelected() const;
+		bool HasColumnSelected() const;
+		bool IsSelectionDisabled() const;
 		void ToggleDisableSelection();
 		void DeleteSelection();
+
+		bool CanMoveRowsUp();
+		void MoveRowsUp();
+		bool CanMoveRowsDown();
+		void MoveRowsDown();
+		bool CanMoveColumnLeft();
+		void MoveColumnLeft();
+		bool CanMoveColumnRight();
+		void MoveColumnRight();
 		
 		void CopySelection();
         void CutSelection();
@@ -122,7 +133,9 @@ namespace UE::ChooserEditor
 		void RefreshRowSelectionDetails();
 		TSharedRef<SWidget> MakeChoosersMenu(UObject* RootObject);
 		void MakeChoosersMenuRecursive(UObject* Outer, FMenuBuilder& MenuBuilder, const FString& Indent);
-		void DeleteSelectedRows();
+		int32 DeleteSelectedRows(int32 RowIndexToRemember = 0);
+		int32 DeleteSelectedRowsInternal(int32 RowIndexToRemember = 0);
+		void MoveRows(int TargetIndex);
 		int MoveRow(int SourceRowIndex, int TargetIndex);
 		int MoveColumn(int SourceIndex, int TargetIndex);
 		void SelectRow(int32 RowIndex, bool bClear = true);
@@ -138,8 +151,11 @@ namespace UE::ChooserEditor
 		};
 		
 		ESelectionType GetCurrentSelectionType() const { return CurrentSelectionType; }
-	private:
 
+		UChooserTable* CopySelectionInternal();
+		void PasteInternal(UChooserTable* PasteObject, int PasteRowIndex=-1);
+	private:
+		
 		void SelectRootProperties();
 		void RemoveDisabledData();
 		void RegisterToolbar();
diff --git a/Engine/Plugins/Experimental/Chooser/Source/ChooserEditor/Private/ChooserTableEditorCommands.h b/Engine/Plugins/Experimental/Chooser/Source/ChooserEditor/Private/ChooserTableEditorCommands.h
--- a/Engine/Plugins/Experimental/Chooser/Source/ChooserEditor/Private/ChooserTableEditorCommands.h
+++ b/Engine/Plugins/Experimental/Chooser/Source/ChooserEditor/Private/ChooserTableEditorCommands.h
@@ -22,6 +22,10 @@ public:
 	TSharedPtr<FUICommandInfo> AutoPopulateSelection;
 	TSharedPtr<FUICommandInfo> RemoveDisabledData;
 	TSharedPtr<FUICommandInfo> Disable;
+	TSharedPtr<FUICommandInfo> MoveLeft;
+	TSharedPtr<FUICommandInfo> MoveRight;
+	TSharedPtr<FUICommandInfo> MoveUp;
+	TSharedPtr<FUICommandInfo> MoveDown;
 
 	/** Initialize commands */
 	virtual void RegisterCommands() override
@@ -31,6 +35,10 @@ public:
 		UI_COMMAND(AutoPopulateSelection, "AutoPopulate", "Auto Populate cell data for selection (requires Columns that support Auto Populate)", EUserInterfaceActionType::Button, FInputChord())
 		UI_COMMAND(RemoveDisabledData, "Remove Disabled Data", "Delete all data that's marked as disabled.", EUserInterfaceActionType::Button, FInputChord())
 		UI_COMMAND(Disable, "Disable", "Disable the selected Rows or Column.", EUserInterfaceActionType::Check, FInputChord())
+		UI_COMMAND(MoveLeft, "Move Left", "Move the selected column to the left", EUserInterfaceActionType::Button, FInputChord(EModifierKey::Alt, EKeys::Left))
+		UI_COMMAND(MoveRight, "Move Right", "Move the selected column to the right", EUserInterfaceActionType::Button, FInputChord(EModifierKey::Alt, EKeys::Right))
+		UI_COMMAND(MoveUp, "Move Up", "Move the selected row(s) up", EUserInterfaceActionType::Button, FInputChord(EModifierKey::Alt, EKeys::Up))
+		UI_COMMAND(MoveDown, "Move Down", "Move the selecte row(s) down", EUserInterfaceActionType::Button, FInputChord(EModifierKey::Alt, EKeys::Down))
 	}
 };
 	
diff --git a/Engine/Plugins/Experimental/Chooser/Source/ChooserEditor/Private/SChooserRowHandle.cpp b/Engine/Plugins/Experimental/Chooser/Source/ChooserEditor/Private/SChooserRowHandle.cpp
--- a/Engine/Plugins/Experimental/Chooser/Source/ChooserEditor/Private/SChooserRowHandle.cpp
+++ b/Engine/Plugins/Experimental/Chooser/Source/ChooserEditor/Private/SChooserRowHandle.cpp
@@ -11,16 +11,29 @@ namespace UE::ChooserEditor
 	TSharedRef<FChooserRowDragDropOp> FChooserRowDragDropOp::New(FChooserTableEditor* InEditor, uint32 InRowIndex)
 	{
 		TSharedRef<FChooserRowDragDropOp> Operation = MakeShareable(new FChooserRowDragDropOp());
-		Operation->ChooserEditor = InEditor;
-		Operation->RowIndex = InRowIndex;
 		Operation->DefaultHoverText = LOCTEXT("Chooser Row", "Chooser Row");
 		Operation->CurrentHoverText = Operation->DefaultHoverText;
+
+		GEditor->BeginTransaction(LOCTEXT("Drag Chooser Table Rows", "Drag Chooser Table Rows"));
+		
+		Operation->RowData = InEditor->CopySelectionInternal();
+		InEditor->DeleteSelectedRowsInternal();
 		
 		Operation->Construct();
 
 		return Operation;
 	}
-	
+
+	void FChooserRowDragDropOp::OnDrop(bool bDropWasHandled, const FPointerEvent& MouseEvent)
+	{
+		if (!bDropWasHandled)
+		{
+			GEditor->EndTransaction();
+			GEditor->UndoTransaction(false);
+		}
+		FDecoratedDragDropOp::OnDrop(bDropWasHandled, MouseEvent);
+	}
+
 	void SChooserRowHandle::Construct(const FArguments& InArgs, bool bShowImage)
 	{
 		ChooserEditor = InArgs._ChooserEditor;
@@ -85,9 +98,6 @@ namespace UE::ChooserEditor
 
 	FReply SChooserRowHandle::OnDragDetected(const FGeometry& MyGeometry, const FPointerEvent& MouseEvent)
 	{
-		// clear row selection so that delete key can't cause the selected row to be deleted
-		ChooserEditor->ClearSelectedRows();
-		
 		TSharedRef<FChooserRowDragDropOp> DragDropOp = FChooserRowDragDropOp::New(ChooserEditor, RowIndex);
 		return FReply::Handled().BeginDragDrop(DragDropOp);
 	}
diff --git a/Engine/Plugins/Experimental/Chooser/Source/ChooserEditor/Private/SChooserRowHandle.h b/Engine/Plugins/Experimental/Chooser/Source/ChooserEditor/Private/SChooserRowHandle.h
--- a/Engine/Plugins/Experimental/Chooser/Source/ChooserEditor/Private/SChooserRowHandle.h
+++ b/Engine/Plugins/Experimental/Chooser/Source/ChooserEditor/Private/SChooserRowHandle.h
@@ -3,6 +3,8 @@
 
 #include "DragAndDrop/DecoratedDragDropOp.h"
 
+class UChooserTable;
+
 namespace UE::ChooserEditor
 {
 
@@ -13,11 +15,12 @@ class FChooserRowDragDropOp : public FDecoratedDragDropOp
 public:
 	DRAG_DROP_OPERATOR_TYPE(FChooserRowDragDropOp, FDecoratedDragDropOp)
 
-	FChooserTableEditor* ChooserEditor;
-	uint32 RowIndex;
+	UChooserTable* RowData = nullptr;
 
 	/** Constructs the drag drop operation */
 	static TSharedRef<FChooserRowDragDropOp> New(FChooserTableEditor* InEditor, uint32 InRowIndex);
+
+	virtual void OnDrop(bool bDropWasHandled, const FPointerEvent& MouseEvent) override;
 };
 
 class SChooserRowHandle : public SCompoundWidget
diff --git a/Engine/Plugins/Experimental/Chooser/Source/ChooserEditor/Private/SChooserTableRow.cpp b/Engine/Plugins/Experimental/Chooser/Source/ChooserEditor/Private/SChooserTableRow.cpp
--- a/Engine/Plugins/Experimental/Chooser/Source/ChooserEditor/Private/SChooserTableRow.cpp
+++ b/Engine/Plugins/Experimental/Chooser/Source/ChooserEditor/Private/SChooserTableRow.cpp
@@ -356,24 +356,22 @@ namespace UE::ChooserEditor
 		
 		if (TSharedPtr<FChooserRowDragDropOp> Operation = DragDropEvent.GetOperationAs<FChooserRowDragDropOp>())
 		{
-			if (Chooser == Operation->ChooserEditor->GetChooser())
+			int NewRowIndex;
+			if (!Chooser->ResultsStructs.IsValidIndex(RowIndex->RowIndex))
 			{
-				int NewRowIndex;
-				if (!Chooser->ResultsStructs.IsValidIndex(RowIndex->RowIndex))
-				{
-					// for special (negative) indices, move to the end
-					NewRowIndex = Editor->MoveRow(Operation->RowIndex, Chooser->ResultsStructs.Num());
-				}
-				else if (bDropAbove)
-				{
-					NewRowIndex = Editor->MoveRow(Operation->RowIndex, RowIndex->RowIndex);
-				}
-				else
-				{
-					NewRowIndex = Editor->MoveRow(Operation->RowIndex, RowIndex->RowIndex+1);
-				}
-				Editor->SelectRow(NewRowIndex);
+				// for special (negative) indices, move to the end
+				NewRowIndex = Chooser->ResultsStructs.Num();
+			}
+			else if (bDropAbove)
+			{
+				NewRowIndex = RowIndex->RowIndex;
+			}
+			else
+			{
+				NewRowIndex = RowIndex->RowIndex+1;
 			}
+			Editor->PasteInternal(Operation->RowData, NewRowIndex);
+			GEditor->EndTransaction();
 		}
 		else if (TSharedPtr<FAssetDragDropOp> ContentDragDropOp = DragDropEvent.GetOperationAs<FAssetDragDropOp>())
 		{
