Improve row drag/drop cancel so it doesn't clear the redo buffer

diff --git a/Engine/Plugins/Experimental/Chooser/Source/ChooserEditor/Private/SChooserRowHandle.cpp b/Engine/Plugins/Experimental/Chooser/Source/ChooserEditor/Private/SChooserRowHandle.cpp
--- a/Engine/Plugins/Experimental/Chooser/Source/ChooserEditor/Private/SChooserRowHandle.cpp
+++ b/Engine/Plugins/Experimental/Chooser/Source/ChooserEditor/Private/SChooserRowHandle.cpp
@@ -2,6 +2,7 @@
 
 #include "SChooserRowHandle.h"
 #include "ChooserTableEditor.h"
+#include "Misc/ITransaction.h"
 
 #define LOCTEXT_NAMESPACE "ChooserRowHandle"
 
@@ -14,10 +15,11 @@ namespace UE::ChooserEditor
 		Operation->DefaultHoverText = LOCTEXT("Chooser Row", "Chooser Row");
 		Operation->CurrentHoverText = Operation->DefaultHoverText;
 
-		GEditor->BeginTransaction(LOCTEXT("Drag Chooser Table Rows", "Drag Chooser Table Rows"));
+		Operation->TransactionIndex = GEditor->BeginTransaction(LOCTEXT("Drag Chooser Table Rows", "Drag Chooser Table Rows"));
 		
 		Operation->RowData = InEditor->CopySelectionInternal();
 		InEditor->DeleteSelectedRowsInternal();
+		Operation->Editor = InEditor;
 		
 		Operation->Construct();
 
@@ -28,8 +30,12 @@ namespace UE::ChooserEditor
 	{
 		if (!bDropWasHandled)
 		{
-			GEditor->EndTransaction();
-			GEditor->UndoTransaction(false);
+			// we need to call Apply on the global undo, or cancelling the transaction doesn't actually roll back
+			GUndo->Apply();
+			GEditor->CancelTransaction(TransactionIndex);
+
+			// PostUndo doesn't get called for cancelled transactions, so refresh the editor manually
+			Editor->RefreshAll();
 		}
 		FDecoratedDragDropOp::OnDrop(bDropWasHandled, MouseEvent);
 	}
diff --git a/Engine/Plugins/Experimental/Chooser/Source/ChooserEditor/Private/SChooserRowHandle.h b/Engine/Plugins/Experimental/Chooser/Source/ChooserEditor/Private/SChooserRowHandle.h
--- a/Engine/Plugins/Experimental/Chooser/Source/ChooserEditor/Private/SChooserRowHandle.h
+++ b/Engine/Plugins/Experimental/Chooser/Source/ChooserEditor/Private/SChooserRowHandle.h
@@ -16,6 +16,8 @@ public:
 	DRAG_DROP_OPERATOR_TYPE(FChooserRowDragDropOp, FDecoratedDragDropOp)
 
 	UChooserTable* RowData = nullptr;
+	int32 TransactionIndex = -1;
+	FChooserTableEditor* Editor = nullptr;
 
 	/** Constructs the drag drop operation */
 	static TSharedRef<FChooserRowDragDropOp> New(FChooserTableEditor* InEditor, uint32 InRowIndex);
