fix Copy/Paste messing up following rows data in Chooser Table Editor

diff --git a/Engine/Plugins/Experimental/Chooser/Source/Chooser/Public/IChooserColumn.h b/Engine/Plugins/Experimental/Chooser/Source/Chooser/Public/IChooserColumn.h
--- a/Engine/Plugins/Experimental/Chooser/Source/Chooser/Public/IChooserColumn.h
+++ b/Engine/Plugins/Experimental/Chooser/Source/Chooser/Public/IChooserColumn.h
@@ -127,7 +127,7 @@ public:
 	}\
 	virtual void CopyRow(FChooserColumnBase& SourceColumn, int SourceRowIndex, int TargetRowIndex) override\
 	{\
-		RowValuesProperty.Insert(static_cast<decltype(this)>(&SourceColumn)->RowValuesProperty[SourceRowIndex], TargetRowIndex);\
+		RowValuesProperty[TargetRowIndex]=static_cast<decltype(this)>(&SourceColumn)->RowValuesProperty[SourceRowIndex];\
 	}\
 	virtual UScriptStruct* GetInputBaseType() const override { return ParameterType::StaticStruct(); };\
 	virtual const UScriptStruct* GetInputType() const override { return InputValue.IsValid() ? InputValue.GetScriptStruct() : nullptr; };\
