Chooser Table Editor column header refactor

diff --git a/Engine/Plugins/Experimental/Chooser/Source/ChooserEditor/Private/BoolColumnEditor.cpp b/Engine/Plugins/Experimental/Chooser/Source/ChooserEditor/Private/BoolColumnEditor.cpp
--- a/Engine/Plugins/Experimental/Chooser/Source/ChooserEditor/Private/BoolColumnEditor.cpp
+++ b/Engine/Plugins/Experimental/Chooser/Source/ChooserEditor/Private/BoolColumnEditor.cpp
@@ -1,6 +1,9 @@
 // Copyright Epic Games, Inc. All Rights Reserved.
 
 #include "BoolColumnEditor.h"
+
+#include <ChooserColumnHeader.h>
+
 #include "BoolColumn.h"
 #include "OutputBoolColumn.h"
 #include "SPropertyAccessChainWidget.h"
@@ -27,52 +30,25 @@ TSharedRef<SWidget> CreateBoolColumnWidget(UChooserTable* Chooser, FChooserColum
 	}
 	if (Row == ColumnWidget_SpecialIndex_Header)
 	{
-		// create column header widget
-		TSharedPtr<SWidget> InputValueWidget = nullptr;
-		if (FChooserParameterBase* InputValue = Column->GetInputValue())
-		{
-			InputValueWidget = FObjectChooserWidgetFactories::CreateWidget(false, Chooser, InputValue, Column->GetInputType(), Chooser->OutputObjectType);
-		}
-	
 		const FSlateBrush* ColumnIcon = FCoreStyle::Get().GetBrush("Icons.Filter");
+		const FText ColumnTooltip = LOCTEXT("Bool Tooltip", "Bool: cells pass if the cell value is equal to the column input value, or the cell is set to \"Any\"");
+		const FText ColumnName = LOCTEXT("Bool","Bool");
 		
-		TSharedRef<SWidget> ColumnHeaderWidget = SNew(SHorizontalBox)
-			+ SHorizontalBox::Slot().AutoWidth()
-			[
-				SNew(SBorder)
-				.BorderBackgroundColor(FLinearColor(0,0,0,0))
-				.Content()
-				[
-					SNew(SImage).Image(ColumnIcon)
-				]
-			]
-			+ SHorizontalBox::Slot()
-			[
-				InputValueWidget ? InputValueWidget.ToSharedRef() : SNullWidget::NullWidget
-			];
-	
+		TSharedPtr<SWidget> DebugWidget = nullptr;
 		if (Chooser->GetEnableDebugTesting())
 		{
-			ColumnHeaderWidget = SNew(SVerticalBox)
-			+ SVerticalBox::Slot()
-			[
-				ColumnHeaderWidget
-			]
-			+ SVerticalBox::Slot()
-			[
-				SNew(SHorizontalBox)
-				 + SHorizontalBox::Slot().FillWidth(1)
-				 + SHorizontalBox::Slot().AutoWidth().VAlign(VAlign_Center)
-				[
-					SNew(SCheckBox).IsEnabled_Lambda([Chooser](){ return !Chooser->HasDebugTarget(); })
-					.IsChecked_Lambda([BoolColumn]() { return BoolColumn->TestValue ? ECheckBoxState::Checked : ECheckBoxState::Unchecked; })
-					.OnCheckStateChanged_Lambda([Chooser, BoolColumn](ECheckBoxState NewValue) { BoolColumn->TestValue = NewValue == ECheckBoxState::Checked; })
-				]
-				 + SHorizontalBox::Slot().FillWidth(1)
-			];
+			DebugWidget = SNew(SHorizontalBox)
+					 + SHorizontalBox::Slot().FillWidth(1)
+					 + SHorizontalBox::Slot().AutoWidth().VAlign(VAlign_Center)
+					[
+						SNew(SCheckBox).IsEnabled_Lambda([Chooser](){ return !Chooser->HasDebugTarget(); })
+						.IsChecked_Lambda([BoolColumn]() { return BoolColumn->TestValue ? ECheckBoxState::Checked : ECheckBoxState::Unchecked; })
+						.OnCheckStateChanged_Lambda([Chooser, BoolColumn](ECheckBoxState NewValue) { BoolColumn->TestValue = NewValue == ECheckBoxState::Checked; })
+					]
+					 + SHorizontalBox::Slot().FillWidth(1);
 		}
 
-		return ColumnHeaderWidget;
+		return MakeColumnHeaderWidget(Chooser, Column, ColumnName, ColumnTooltip, ColumnIcon, DebugWidget);
 	}
 
 	// create cell widget
@@ -117,51 +93,24 @@ TSharedRef<SWidget> CreateOutputBoolColumnWidget(UChooserTable* Chooser, FChoose
 
 	if (Row == ColumnWidget_SpecialIndex_Header)
 	{
-		// create column header widget
-		TSharedPtr<SWidget> InputValueWidget = nullptr;
-		if (FChooserParameterBase* InputValue = Column->GetInputValue())
-		{
-			InputValueWidget = FObjectChooserWidgetFactories::CreateWidget(false, Chooser, InputValue, Column->GetInputType(), Chooser->OutputObjectType);
-		}
-		
 		const FSlateBrush* ColumnIcon = FCoreStyle::Get().GetBrush("Icons.ArrowRight");
+		const FText ColumnTooltip = LOCTEXT("Output Bool Tooltip", "Output Bool: Writes the value from the cell from the selected row to the bound property");
+		const FText ColumnName = LOCTEXT("Output Bool","Output Bool");
 		
-		TSharedRef<SWidget> ColumnHeaderWidget = SNew(SHorizontalBox)
-			+ SHorizontalBox::Slot().AutoWidth()
-			[
-				SNew(SBorder)
-				.BorderBackgroundColor(FLinearColor(0,0,0,0))
-				.Content()
-				[
-					SNew(SImage).Image(ColumnIcon)
-				]
-			]
-			+ SHorizontalBox::Slot()
-			[
-				InputValueWidget ? InputValueWidget.ToSharedRef() : SNullWidget::NullWidget
-			];
-	
+		TSharedPtr<SWidget> DebugWidget = nullptr;
 		if (Chooser->GetEnableDebugTesting())
 		{
-			ColumnHeaderWidget = SNew(SVerticalBox)
-			+ SVerticalBox::Slot()
-			[
-				ColumnHeaderWidget
-			]
-			+ SVerticalBox::Slot()
-			[
-				SNew(SHorizontalBox)
-             	+ SHorizontalBox::Slot().FillWidth(1)
-             	+ SHorizontalBox::Slot().AutoWidth().VAlign(VAlign_Center)
-             	[
-				SNew(SCheckBox).IsEnabled(false)
-				.IsChecked_Lambda([BoolColumn]() { return BoolColumn->TestValue ? ECheckBoxState::Checked : ECheckBoxState::Unchecked; })
-				]
-             	+ SHorizontalBox::Slot().FillWidth(1)
-			];
+			DebugWidget = SNew(SHorizontalBox)
+					 + SHorizontalBox::Slot().FillWidth(1)
+					 + SHorizontalBox::Slot().AutoWidth().VAlign(VAlign_Center)
+					[
+						SNew(SCheckBox).IsEnabled_Lambda([Chooser](){ return !Chooser->HasDebugTarget(); })
+						.IsChecked_Lambda([BoolColumn]() { return BoolColumn->TestValue ? ECheckBoxState::Checked : ECheckBoxState::Unchecked; })
+					]
+					 + SHorizontalBox::Slot().FillWidth(1);
 		}
 
-		return ColumnHeaderWidget;
+		return MakeColumnHeaderWidget(Chooser, Column, ColumnName, ColumnTooltip, ColumnIcon, DebugWidget);
 	}
 	else if (Row == ColumnWidget_SpecialIndex_Fallback) 
 	{
diff --git a/Engine/Plugins/Experimental/Chooser/Source/ChooserEditor/Private/ChooserColumnHeader.cpp b/Engine/Plugins/Experimental/Chooser/Source/ChooserEditor/Private/ChooserColumnHeader.cpp
new file mode 100644
--- /dev/null
+++ b/Engine/Plugins/Experimental/Chooser/Source/ChooserEditor/Private/ChooserColumnHeader.cpp
@@ -0,0 +1,114 @@
+// Copyright Epic Games, Inc. All Rights Reserved.
+
+#include "ChooserColumnHeader.h"
+
+#include "ChooserEditorStyle.h"
+#include "ObjectChooserWidgetFactories.h"
+#include "ChooserTableEditor.h"
+#include "ChooserEditorStyle.h"
+#include "Widgets/Images/SImage.h"
+#include "GraphEditorSettings.h"
+#include "ScopedTransaction.h"
+
+#define LOCTEXT_NAMESPACE "ChooserColumnHeader"
+
+namespace UE::ChooserEditor
+{
+
+	TSharedRef<SWidget> MakeColumnHeaderWidget(UChooserTable* Chooser, FChooserColumnBase* Column, const FText& ColumnName,const FText& ColumnTooltip, const FSlateBrush* ColumnIcon, TSharedPtr<SWidget> DebugWidget)
+	{
+		TSharedPtr<SWidget> InputValueWidget = nullptr;
+		if (FChooserParameterBase* InputValue = Column->GetInputValue())
+		{
+			InputValueWidget = FObjectChooserWidgetFactories::CreateWidget(false, Chooser, InputValue, Column->GetInputType(), Chooser->OutputObjectType);
+		}
+	
+		TSharedRef<SVerticalBox> ColumnHeaderWidget =
+			SNew(SVerticalBox)
+			+ SVerticalBox::Slot().AutoHeight()
+			[
+				SNew(SHorizontalBox)
+				+ SHorizontalBox::Slot().AutoWidth()
+				[
+					SNew(SBorder)
+					.ToolTipText(ColumnTooltip)
+					.BorderBackgroundColor(FLinearColor(0,0,0,0))
+					.Content()
+					[
+						SNew(SImage).Image(ColumnIcon)
+					]
+				]
+				+ SHorizontalBox::Slot().VAlign(VAlign_Center).FillWidth(1)
+				[
+					SNew(STextBlock)
+						.Margin(FMargin(5, 0, 0, 0))
+						.Text(ColumnName)
+						.ToolTipText(ColumnTooltip)
+				]
+				+ SHorizontalBox::Slot().AutoWidth()
+				[
+					SNew(SButton)
+					.ToolTipText(LOCTEXT("ToggleDisableTooltip","Toggle disable this column.  Disabled columns will not be evaluated and will be stripped from cooked data."))
+					.ButtonStyle(FAppStyle::Get(), "NoBorder")
+					.OnClicked_Lambda([Column, Chooser]()
+					{
+						FScopedTransaction Transaction(LOCTEXT("Toggle Disable Column", "Toggle Disable Column"));
+						Chooser->Modify();
+						Column->bDisabled = !Column->bDisabled;
+						return FReply::Handled();
+					})
+					.Content()
+					[
+						SNew(SOverlay)
+							+ SOverlay::Slot()
+							[
+								SNew(SImage).Image(FChooserEditorStyle::Get().GetBrush("ChooserEditor.EyeOpened"))
+								.Visibility_Lambda([Column]() { return Column->bDisabled ? EVisibility::Hidden : EVisibility::Visible; })
+							]
+							+ SOverlay::Slot()
+							[
+								SNew(SImage).Image(FChooserEditorStyle::Get().GetBrush("ChooserEditor.EyeClosed"))
+								.Visibility_Lambda([Column]() { return Column->bDisabled ? EVisibility::Visible : EVisibility::Hidden; })
+							]
+					]
+				]
+			]
+			+ SVerticalBox::Slot().AutoHeight()
+			[
+				InputValueWidget ? InputValueWidget.ToSharedRef() : SNullWidget::NullWidget
+			];
+		
+		if (DebugWidget.IsValid())
+		{
+			ColumnHeaderWidget->AddSlot().AutoHeight()
+			[
+				SNew(SBorder)
+					.BorderBackgroundColor(FLinearColor(0,0,0,0))
+					.Padding(FMargin(0,5,0,0))
+					.Content()
+					[
+						SNew(SHorizontalBox)
+						+ SHorizontalBox::Slot().AutoWidth()
+						[
+							SNew(SBorder)
+							.ToolTipText(LOCTEXT("DebugTooltip","Debug Value: This value either comes from the attached debug target, or is manually entered.  It is used to colorize the column cells based on how they would evaulate for the given input."))
+							.BorderBackgroundColor(FLinearColor(0,0,0,0))
+							.Content()
+							[
+								SNew(SImage).Image(FAppStyle::Get().GetBrush("Debug"))
+							]
+						]
+						+ SHorizontalBox::Slot().FillWidth(1)
+						[
+							DebugWidget.ToSharedRef()
+						]
+					]
+			];
+		}
+
+		return ColumnHeaderWidget;
+	}
+	
+}
+
+#undef LOCTEXT_NAMESPACE
diff --git a/Engine/Plugins/Experimental/Chooser/Source/ChooserEditor/Private/ChooserEditorStyle.cpp b/Engine/Plugins/Experimental/Chooser/Source/ChooserEditor/Private/ChooserEditorStyle.cpp
--- a/Engine/Plugins/Experimental/Chooser/Source/ChooserEditor/Private/ChooserEditorStyle.cpp
+++ b/Engine/Plugins/Experimental/Chooser/Source/ChooserEditor/Private/ChooserEditorStyle.cpp
@@ -28,6 +28,11 @@ FChooserEditorStyle::FChooserEditorStyle() :
 	Set("ChooserEditor.ChooserTableIconSmall", new IMAGE_BRUSH_SVG("ChooserIcon_16", Icon16x16));
 	
 	Set("ChooserEditor.FallbackIcon", new IMAGE_BRUSH_SVG("Fallback", Icon16x16));
+
+	// visibility bottons
+	Set("ChooserEditor.EyeClosed", new IMAGE_BRUSH(TEXT("icon_EyeClosed_40x"), Icon16x16));
+	Set("ChooserEditor.EyeOpened", new IMAGE_BRUSH(TEXT("icon_EyeOpened_40x"), Icon16x16));
+	
 }
 
 void FChooserEditorStyle::Initialize()
diff --git a/Engine/Plugins/Experimental/Chooser/Source/ChooserEditor/Private/ChooserTableEditor.cpp b/Engine/Plugins/Experimental/Chooser/Source/ChooserEditor/Private/ChooserTableEditor.cpp
--- a/Engine/Plugins/Experimental/Chooser/Source/ChooserEditor/Private/ChooserTableEditor.cpp
+++ b/Engine/Plugins/Experimental/Chooser/Source/ChooserEditor/Private/ChooserTableEditor.cpp
@@ -920,8 +920,9 @@ void FChooserTableEditor::RemoveDisabledData()
 int FChooserTableEditor::MoveColumn(int SourceIndex, int TargetIndex)
 {
 	UChooserTable* Chooser = GetChooser();
-	TargetIndex = FMath::Clamp(TargetIndex, 0, Chooser->ResultsStructs.Num());
 	
+	TargetIndex = FMath::Clamp(TargetIndex, 0, Chooser->ColumnsStructs.Num());
+
 	if (SourceIndex < 0 || SourceIndex == TargetIndex)
 	{
 		return TargetIndex;
@@ -933,11 +934,12 @@ int FChooserTableEditor::MoveColumn(int SourceIndex, int TargetIndex)
 
 	FInstancedStruct ColumnData = Chooser->ColumnsStructs[SourceIndex];
 	Chooser->ColumnsStructs.RemoveAt(SourceIndex);
+		
 	if (SourceIndex < TargetIndex)
 	{
 		TargetIndex--;
 	}
-
+	
 	if (TargetIndex == Chooser->ColumnsStructs.Num())
 	{
 		if (Chooser->ColumnsStructs.Last().GetPtr<FRandomizeColumn>())
@@ -1020,7 +1022,7 @@ bool FChooserTableEditor::IsRowSelected(int32 RowIndex)
 	
 bool FChooserTableEditor::IsColumnSelected(int32 ColumnIndex)
 {
-	return  (CurrentSelectionType == ESelectionType::Column && SelectedColumn);
+	return  (CurrentSelectionType == ESelectionType::Column && SelectedColumn && SelectedColumn->Column == ColumnIndex);
 }
 
 void FChooserTableEditor::UpdateTableColumns()
@@ -1034,8 +1036,18 @@ void FChooserTableEditor::UpdateTableColumns()
 					.ManualWidth(30));
 
 	HeaderRow->AddColumn(SHeaderRow::Column("Result")
-					.DefaultLabel(LOCTEXT("ResultColumnName", "Result"))
-					.ManualWidth(300));
+					.ManualWidth(300)
+					.HeaderContent()
+					[
+						SNew(SVerticalBox)
+						+ SVerticalBox::Slot()
+						.VAlign(VAlign_Top)
+						[
+							SNew(STextBlock)
+								.Text(LOCTEXT("Result", "Result"))
+								.ToolTipText(LOCTEXT("ResultTooltip", "The Result is the asset which will be returned if a row is selected (or other Chooser to evaluate to get the asset to return"))
+						]
+					]);
 
 	FName ColumnId("ChooserColumn", 1);
 	int NumColumns = Chooser->ColumnsStructs.Num();	
@@ -1060,17 +1072,7 @@ void FChooserTableEditor::UpdateTableColumns()
 					.ColumnIndex(ColumnIndex)
 					.NoDropAfter(Chooser->ColumnsStructs[ColumnIndex].GetPtr<FRandomizeColumn>() != nullptr)
 				[
-					SNew(SBorder)
-					.VAlign(VAlign_Center)
-					.Padding(3)
-					.BorderBackgroundColor_Lambda([this, ColumnIndex] ()
-					{
-						// unclear why this color is coming out much darker
-						return (SelectedColumn && SelectedColumn->Column == ColumnIndex) ? FSlateColor(FColor(0x00, 0x70, 0xe0, 0xFF)) : FSlateColor(FLinearColor(0.05f,0.05f,0.05f));
-					})
-					[
-						HeaderWidget.ToSharedRef()
-					]
+					HeaderWidget.ToSharedRef()
 				]
 			
 			]);
diff --git a/Engine/Plugins/Experimental/Chooser/Source/ChooserEditor/Private/ChooserTableEditor.h b/Engine/Plugins/Experimental/Chooser/Source/ChooserEditor/Private/ChooserTableEditor.h
--- a/Engine/Plugins/Experimental/Chooser/Source/ChooserEditor/Private/ChooserTableEditor.h
+++ b/Engine/Plugins/Experimental/Chooser/Source/ChooserEditor/Private/ChooserTableEditor.h
@@ -154,6 +154,8 @@ namespace UE::ChooserEditor
 
 		UChooserTable* CopySelectionInternal();
 		void PasteInternal(UChooserTable* PasteObject, int PasteRowIndex=-1);
+
+		bool TableHasFocus() const { return TableView->HasKeyboardFocus(); }
 	private:
 		
 		void SelectRootProperties();
diff --git a/Engine/Plugins/Experimental/Chooser/Source/ChooserEditor/Private/EnumColumnEditor.cpp b/Engine/Plugins/Experimental/Chooser/Source/ChooserEditor/Private/EnumColumnEditor.cpp
--- a/Engine/Plugins/Experimental/Chooser/Source/ChooserEditor/Private/EnumColumnEditor.cpp
+++ b/Engine/Plugins/Experimental/Chooser/Source/ChooserEditor/Private/EnumColumnEditor.cpp
@@ -6,6 +6,7 @@
 #include "SPropertyAccessChainWidget.h"
 #include "ObjectChooserWidgetFactories.h"
 #include "ChooserTableEditor.h"
+#include "ChooserColumnHeader.h"
 #include "GraphEditorSettings.h"
 #include "SEnumCombo.h"
 #include "TransactionCommon.h"
@@ -17,117 +18,6 @@
 namespace UE::ChooserEditor
 {
 
-	
-// Wrapper widget for EnumComboBox which will reconstruct the combo box when the Enum has changed
-template <typename ColumnType>
-class SEnumCell : public SCompoundWidget
-{
-public:
-
-	DECLARE_DELEGATE_OneParam(FOnValueSet, int);
-	
-	SLATE_BEGIN_ARGS(SEnumCell)
-	{}
-
-	SLATE_ARGUMENT(UObject*, TransactionObject)
-	SLATE_ARGUMENT(ColumnType*, EnumColumn)
-	SLATE_ATTRIBUTE(int32, EnumValue);
-	SLATE_EVENT(FOnValueSet, OnValueSet)
-            
-	SLATE_END_ARGS()
-
-	TSharedRef<SWidget> CreateEnumComboBox()
-	{
-		if (const ColumnType* EnumColumnPointer = EnumColumn)
-		{
-			if (EnumColumnPointer->InputValue.IsValid())
-			{
-				if (const UEnum* Enum = EnumColumnPointer->InputValue.template Get<FChooserParameterEnumBase>().GetEnum())
-				{
-					return SNew(SEnumComboBox, Enum)
-						.IsEnabled_Lambda([this](){ return IsEnabled(); } )
-						.CurrentValue(EnumValue)
-						.OnEnumSelectionChanged_Lambda([this](int32 InEnumValue, ESelectInfo::Type)
-						{
-							const FScopedTransaction Transaction(LOCTEXT("Edit RHS", "Edit Enum Value"));
-							TransactionObject->Modify(true);
-							OnValueSet.ExecuteIfBound(InEnumValue);
-						});
-				}
-			}
-		}
-		
-		return SNullWidget::NullWidget;
-	}
-
-	void UpdateEnumComboBox()
-	{
-		ChildSlot[ CreateEnumComboBox()	];
-	}
-
-	virtual void Tick(const FGeometry& AllottedGeometry, const double InCurrentTime, const float InDeltaTime) override
-	{
-		const UEnum* CurrentEnumSource = nullptr;
-		if (EnumColumn->InputValue.IsValid())
-		{
-			CurrentEnumSource = EnumColumn->InputValue.template Get<FChooserParameterEnumBase>().GetEnum(); 
-		}
-		if (EnumSource != CurrentEnumSource)
-		{
-			EnumComboBorder->SetContent(CreateEnumComboBox());
-			EnumSource = CurrentEnumSource;
-		}
-	}
-    					
-
-	void Construct( const FArguments& InArgs)
-	{
-		SetEnabled(InArgs._IsEnabled);
-		
-		SetCanTick(true);
-		EnumColumn = InArgs._EnumColumn;
-		TransactionObject = InArgs._TransactionObject;
-		EnumValue = InArgs._EnumValue;
-		OnValueSet = InArgs._OnValueSet;
-
-		if (EnumColumn)
-		{
-			if (EnumColumn->InputValue.IsValid())
-			{
-				EnumSource = EnumColumn->InputValue.template Get<FChooserParameterEnumBase>().GetEnum();
-			}
-		}
-
-		UpdateEnumComboBox();
-
-		int Row = RowIndex.Get();
-
-		ChildSlot
-		[
-			SAssignNew(EnumComboBorder, SBorder).Padding(0).BorderBackgroundColor(FLinearColor(0,0,0,0))
-			[
-				CreateEnumComboBox()
-			]
-		];
-		
-	}
-
-	~SEnumCell()
-	{
-	}
-
-private:
-	UObject* TransactionObject = nullptr;
-	ColumnType* EnumColumn = nullptr;
-	const UEnum* EnumSource = nullptr;
-	TSharedPtr<SBorder> EnumComboBorder;
-	TAttribute<int> RowIndex;
-	FDelegateHandle EnumChangedHandle;
-	
-	FOnValueSet OnValueSet;
-	TAttribute<int32> EnumValue;
-};
-
 TSharedRef<SWidget> CreateEnumColumnWidget(UChooserTable* Chooser, FChooserColumnBase* Column, int Row)
 {
 	FEnumColumn* EnumColumn = static_cast<FEnumColumn*>(Column);
@@ -138,47 +28,20 @@ TSharedRef<SWidget> CreateEnumColumnWidget(UChooserTable* Chooser, FChooserColum
 	}
 	if (Row == ColumnWidget_SpecialIndex_Header)
 	{
-		// create column header widget
-		TSharedPtr<SWidget> InputValueWidget = nullptr;
-		if (FChooserParameterBase* InputValue = Column->GetInputValue())
-		{
-			InputValueWidget = FObjectChooserWidgetFactories::CreateWidget(false, Chooser, InputValue, Column->GetInputType(), Chooser->OutputObjectType);
-		}
-		
 		const FSlateBrush* ColumnIcon = FCoreStyle::Get().GetBrush("Icons.Filter");
+		const FText ColumnTooltip = LOCTEXT("Enum Value Tooltip", "Enum Value: cells pass if the cell value is equal to the column input value");
+		const FText ColumnName = LOCTEXT("Enum Value","Enum Value");
 		
-		TSharedRef<SWidget> ColumnHeaderWidget = SNew(SHorizontalBox)
-			+ SHorizontalBox::Slot().AutoWidth()
-			[
-				SNew(SBorder)
-				.BorderBackgroundColor(FLinearColor(0,0,0,0))
-				.Content()
-				[
-					SNew(SImage).Image(ColumnIcon)
-				]
-			]
-			+ SHorizontalBox::Slot()
-			[
-				InputValueWidget ? InputValueWidget.ToSharedRef() : SNullWidget::NullWidget
-			];
-	
+		TSharedPtr<SWidget> DebugWidget = nullptr;
 		if (Chooser->GetEnableDebugTesting())
 		{
-			ColumnHeaderWidget = SNew(SVerticalBox)
-			+ SVerticalBox::Slot()
-			[
-				ColumnHeaderWidget
-			]
-			+ SVerticalBox::Slot()
-			[
-				SNew(SEnumCell<FEnumColumn>).TransactionObject(Chooser).EnumColumn(EnumColumn)
-					.OnValueSet_Lambda([EnumColumn](int Value) { EnumColumn->TestValue = Value; })
-					.EnumValue_Lambda([EnumColumn]() { return EnumColumn->TestValue; })
-					.IsEnabled_Lambda([Chooser] { return !Chooser->HasDebugTarget(); })
-			];
+			DebugWidget = SNew(SEnumCell<FEnumColumn>).TransactionObject(Chooser).EnumColumn(EnumColumn)
+								.OnValueSet_Lambda([EnumColumn](int Value) { EnumColumn->TestValue = Value; })
+								.EnumValue_Lambda([EnumColumn]() { return EnumColumn->TestValue; })
+								.IsEnabled_Lambda([Chooser] { return !Chooser->HasDebugTarget(); });
 		}
 
-		return ColumnHeaderWidget;
+		return MakeColumnHeaderWidget(Chooser, Column, ColumnName, ColumnTooltip, ColumnIcon, DebugWidget);
 	}
 
 	// create cell widget
@@ -250,45 +113,18 @@ TSharedRef<SWidget> CreateOutputEnumColumnWidget(UChooserTable* Chooser, FChoose
 	
 	if (Row == ColumnWidget_SpecialIndex_Header)
 	{
-		// create column header widget
-		TSharedPtr<SWidget> InputValueWidget = nullptr;
-		if (FChooserParameterBase* InputValue = Column->GetInputValue())
-		{
-			InputValueWidget = FObjectChooserWidgetFactories::CreateWidget(false, Chooser, InputValue, Column->GetInputType(), Chooser->OutputObjectType);
-		}
-		
 		const FSlateBrush* ColumnIcon = FCoreStyle::Get().GetBrush("Icons.ArrowRight");
+		const FText ColumnTooltip = LOCTEXT("Output Enum", "Output Enum:  writes the value from cell in the result row to the bound variable");
+		const FText ColumnName = LOCTEXT("Output Enum","Output Enum");
 		
-		TSharedRef<SWidget> ColumnHeaderWidget = SNew(SHorizontalBox)
-			+ SHorizontalBox::Slot().AutoWidth()
-			[
-				SNew(SBorder)
-				.BorderBackgroundColor(FLinearColor(0,0,0,0))
-				.Content()
-				[
-					SNew(SImage).Image(ColumnIcon)
-				]
-			]
-			+ SHorizontalBox::Slot()
-			[
-				InputValueWidget ? InputValueWidget.ToSharedRef() : SNullWidget::NullWidget
-			];
-	
+		TSharedPtr<SWidget> DebugWidget = nullptr;
 		if (Chooser->GetEnableDebugTesting())
 		{
-			ColumnHeaderWidget = SNew(SVerticalBox)
-			+ SVerticalBox::Slot()
-			[
-				ColumnHeaderWidget
-			]
-			+ SVerticalBox::Slot()
-			[
-				SNew(SEnumCell<FOutputEnumColumn>).TransactionObject(Chooser).EnumColumn(EnumColumn).IsEnabled(false)
-				.EnumValue_Lambda([EnumColumn]() { return static_cast<int32>(EnumColumn->TestValue); })
-			];
+			DebugWidget = SNew(SEnumCell<FOutputEnumColumn>).TransactionObject(Chooser).EnumColumn(EnumColumn).IsEnabled(false)
+					.EnumValue_Lambda([EnumColumn]() { return static_cast<int32>(EnumColumn->TestValue); });
 		}
 
-		return ColumnHeaderWidget;
+		return MakeColumnHeaderWidget(Chooser, Column, ColumnName, ColumnTooltip, ColumnIcon, DebugWidget);
 	}
 	else if (Row == ColumnWidget_SpecialIndex_Fallback)
 	{
diff --git a/Engine/Plugins/Experimental/Chooser/Source/ChooserEditor/Private/EnumColumnEditor.h b/Engine/Plugins/Experimental/Chooser/Source/ChooserEditor/Private/EnumColumnEditor.h
--- a/Engine/Plugins/Experimental/Chooser/Source/ChooserEditor/Private/EnumColumnEditor.h
+++ b/Engine/Plugins/Experimental/Chooser/Source/ChooserEditor/Private/EnumColumnEditor.h
@@ -3,8 +3,128 @@
 #pragma once
 
 #include "CoreMinimal.h"
+#include "SEnumCombo.h"
+#include "IChooserParameterEnum.h"
+#include "Widgets/Layout/SBorder.h"
+#include "Widgets/SCompoundWidget.h"
+#include "ScopedTransaction.h"
+
+#define LOCTEXT_NAMESPACE "EnumColumnEditor"
 
 namespace UE::ChooserEditor
 {
+	
+	// Wrapper widget for EnumComboBox which will reconstruct the combo box when the Enum has changed
+	template <typename ColumnType>
+	class SEnumCell : public SCompoundWidget
+	{
+	public:
+
+		DECLARE_DELEGATE_OneParam(FOnValueSet, int);
+		
+		SLATE_BEGIN_ARGS(SEnumCell)
+		{}
+
+		SLATE_ARGUMENT(UObject*, TransactionObject)
+		SLATE_ARGUMENT(ColumnType*, EnumColumn)
+		SLATE_ATTRIBUTE(int32, EnumValue);
+		SLATE_EVENT(FOnValueSet, OnValueSet)
+				
+		SLATE_END_ARGS()
+
+		TSharedRef<SWidget> CreateEnumComboBox()
+		{
+			if (const ColumnType* EnumColumnPointer = EnumColumn)
+			{
+				if (EnumColumnPointer->InputValue.IsValid())
+				{
+					if (const UEnum* Enum = EnumColumnPointer->InputValue.template Get<FChooserParameterEnumBase>().GetEnum())
+					{
+						return SNew(SEnumComboBox, Enum)
+							.IsEnabled_Lambda([this](){ return IsEnabled(); } )
+							.CurrentValue(EnumValue)
+							.OnEnumSelectionChanged_Lambda([this](int32 InEnumValue, ESelectInfo::Type)
+							{
+								const FScopedTransaction Transaction(LOCTEXT("Edit RHS", "Edit Enum Value"));
+								TransactionObject->Modify(true);
+								OnValueSet.ExecuteIfBound(InEnumValue);
+							});
+					}
+				}
+			}
+			
+			return SNullWidget::NullWidget;
+		}
+
+		void UpdateEnumComboBox()
+		{
+			ChildSlot[ CreateEnumComboBox()	];
+		}
+
+		virtual void Tick(const FGeometry& AllottedGeometry, const double InCurrentTime, const float InDeltaTime) override
+		{
+			const UEnum* CurrentEnumSource = nullptr;
+			if (EnumColumn->InputValue.IsValid())
+			{
+				CurrentEnumSource = EnumColumn->InputValue.template Get<FChooserParameterEnumBase>().GetEnum(); 
+			}
+			if (EnumSource != CurrentEnumSource)
+			{
+				EnumComboBorder->SetContent(CreateEnumComboBox());
+				EnumSource = CurrentEnumSource;
+			}
+		}
+							
+		void Construct( const FArguments& InArgs)
+		{
+			SetEnabled(InArgs._IsEnabled);
+			
+			SetCanTick(true);
+			EnumColumn = InArgs._EnumColumn;
+			TransactionObject = InArgs._TransactionObject;
+			EnumValue = InArgs._EnumValue;
+			OnValueSet = InArgs._OnValueSet;
+
+			if (EnumColumn)
+			{
+				if (EnumColumn->InputValue.IsValid())
+				{
+					EnumSource = EnumColumn->InputValue.template Get<FChooserParameterEnumBase>().GetEnum();
+				}
+			}
+
+			UpdateEnumComboBox();
+
+			int Row = RowIndex.Get();
+
+			ChildSlot
+			[
+				SAssignNew(EnumComboBorder, SBorder).Padding(0).BorderBackgroundColor(FLinearColor(0,0,0,0))
+				[
+					CreateEnumComboBox()
+				]
+			];
+			
+		}
+
+		~SEnumCell()
+		{
+		}
+
+	private:
+		UObject* TransactionObject = nullptr;
+		ColumnType* EnumColumn = nullptr;
+		const UEnum* EnumSource = nullptr;
+		TSharedPtr<SBorder> EnumComboBorder;
+		TAttribute<int> RowIndex;
+		FDelegateHandle EnumChangedHandle;
+		
+		FOnValueSet OnValueSet;
+		TAttribute<int32> EnumValue;
+	};
+
+	
 	void RegisterEnumWidgets();
-}
\ No newline at end of file
+}
+
+#undef LOCTEXT_NAMESPACE
\ No newline at end of file
diff --git a/Engine/Plugins/Experimental/Chooser/Source/ChooserEditor/Private/FloatDistanceColumnEditor.cpp b/Engine/Plugins/Experimental/Chooser/Source/ChooserEditor/Private/FloatDistanceColumnEditor.cpp
--- a/Engine/Plugins/Experimental/Chooser/Source/ChooserEditor/Private/FloatDistanceColumnEditor.cpp
+++ b/Engine/Plugins/Experimental/Chooser/Source/ChooserEditor/Private/FloatDistanceColumnEditor.cpp
@@ -1,10 +1,13 @@
 // Copyright Epic Games, Inc. All Rights Reserved.
 
 #include "FloatDistanceColumnEditor.h"
+
+#include "ChooserEditorStyle.h"
+#include "ChooserColumnHeader.h"
 #include "FloatDistanceColumn.h"
 #include "ObjectChooserWidgetFactories.h"
 #include "ChooserTableEditor.h"
-#include "Widgets/Input/SEditableTextBox.h"
+#include "ChooserEditorStyle.h"
 #include "Widgets/Input/SNumericEntryBox.h"
 #include "Widgets/Images/SImage.h"
 #include "GraphEditorSettings.h"
@@ -26,45 +29,21 @@ TSharedRef<SWidget> CreateFloatDistanceColumnWidget(UChooserTable* Chooser, FCho
 	else if (Row == ColumnWidget_SpecialIndex_Header)
 	{
 		// create column header widget
-		TSharedPtr<SWidget> InputValueWidget = nullptr;
-		if (FChooserParameterBase* InputValue = Column->GetInputValue())
-		{
-			InputValueWidget = FObjectChooserWidgetFactories::CreateWidget(false, Chooser, InputValue, Column->GetInputType(), Chooser->OutputObjectType);
-		}
 		
-		const FSlateBrush* ColumnIcon = FCoreStyle::Get().GetBrush("Icons.Filter");
+		const FSlateBrush* ColumnIcon = FCoreStyle::Get().GetBrush("Icons.SortUp");
+		const FText ColumnTooltip = LOCTEXT("Float distance tooltip", "Float Distance Column: rows recieve a Cost based on how different the input float is from the row value");
+		const FText ColumnName = LOCTEXT("Float Distance","Float Distance");
 		
-		TSharedRef<SWidget> ColumnHeaderWidget = SNew(SHorizontalBox)
-			+ SHorizontalBox::Slot().AutoWidth()
-			[
-				SNew(SBorder)
-				.BorderBackgroundColor(FLinearColor(0,0,0,0))
-				.Content()
-				[
-					SNew(SImage).Image(ColumnIcon)
-				]
-			]
-			+ SHorizontalBox::Slot()
-			[
-				InputValueWidget ? InputValueWidget.ToSharedRef() : SNullWidget::NullWidget
-			];
-	
+		TSharedPtr<SWidget> DebugWidget = nullptr;
 		if (Chooser->GetEnableDebugTesting())
 		{
-			ColumnHeaderWidget = SNew(SVerticalBox)
-			+ SVerticalBox::Slot()
-			[
-				ColumnHeaderWidget
-			]
-			+ SVerticalBox::Slot()
-			[
-				SNew(SNumericEntryBox<float>).IsEnabled_Lambda([Chooser](){ return !Chooser->HasDebugTarget(); })
+			DebugWidget = SNew(SNumericEntryBox<float>)
+				.IsEnabled_Lambda([Chooser](){ return !Chooser->HasDebugTarget(); })
 				.Value_Lambda([FloatDistanceColumn]() { return FloatDistanceColumn->TestValue; })
-				.OnValueCommitted_Lambda([Chooser, FloatDistanceColumn](float NewValue, ETextCommit::Type CommitType) { FloatDistanceColumn->TestValue = NewValue; })
-			];
+				.OnValueCommitted_Lambda([Chooser, FloatDistanceColumn](float NewValue, ETextCommit::Type CommitType) { FloatDistanceColumn->TestValue = NewValue; });
 		}
 
-		return ColumnHeaderWidget;
+		return MakeColumnHeaderWidget(Chooser, Column, ColumnName, ColumnTooltip, ColumnIcon, DebugWidget);
 	}
 
 	// create cell widget
diff --git a/Engine/Plugins/Experimental/Chooser/Source/ChooserEditor/Private/FloatRangeColumnEditor.cpp b/Engine/Plugins/Experimental/Chooser/Source/ChooserEditor/Private/FloatRangeColumnEditor.cpp
--- a/Engine/Plugins/Experimental/Chooser/Source/ChooserEditor/Private/FloatRangeColumnEditor.cpp
+++ b/Engine/Plugins/Experimental/Chooser/Source/ChooserEditor/Private/FloatRangeColumnEditor.cpp
@@ -1,6 +1,9 @@
 // Copyright Epic Games, Inc. All Rights Reserved.
 
 #include "FloatRangeColumnEditor.h"
+
+#include <ChooserColumnHeader.h>
+
 #include "FloatRangeColumn.h"
 #include "ObjectChooserWidgetFactories.h"
 #include "ChooserTableEditor.h"
@@ -39,45 +42,22 @@ TSharedRef<SWidget> CreateFloatRangeColumnWidget(UChooserTable* Chooser, FChoose
 	else if (Row == ColumnWidget_SpecialIndex_Header)
 	{
 		// create column header widget
-		TSharedPtr<SWidget> InputValueWidget = nullptr;
-		if (FChooserParameterBase* InputValue = Column->GetInputValue())
-		{
-			InputValueWidget = FObjectChooserWidgetFactories::CreateWidget(false, Chooser, InputValue, Column->GetInputType(), Chooser->OutputObjectType);
-		}
-		
 		const FSlateBrush* ColumnIcon = FCoreStyle::Get().GetBrush("Icons.Filter");
+		const FText ColumnTooltip = LOCTEXT("Float Range Tooltip", "Float Range: cells pass if the input value is within the range specified in the cell");
+		const FText ColumnName = LOCTEXT("Float Range","Float Range");
 		
-		TSharedRef<SWidget> ColumnHeaderWidget = SNew(SHorizontalBox)
-			+ SHorizontalBox::Slot().AutoWidth()
-			[
-				SNew(SBorder)
-				.BorderBackgroundColor(FLinearColor(0,0,0,0))
-				.Content()
-				[
-					SNew(SImage).Image(ColumnIcon)
-				]
-			]
-			+ SHorizontalBox::Slot()
-			[
-				InputValueWidget ? InputValueWidget.ToSharedRef() : SNullWidget::NullWidget
-			];
-	
+		TSharedPtr<SWidget> DebugWidget = nullptr;
 		if (Chooser->GetEnableDebugTesting())
 		{
-			ColumnHeaderWidget = SNew(SVerticalBox)
-			+ SVerticalBox::Slot()
-			[
-				ColumnHeaderWidget
-			]
-			+ SVerticalBox::Slot()
-			[
-				SNew(SNumericEntryBox<float>).IsEnabled_Lambda([Chooser](){ return !Chooser->HasDebugTarget(); })
-				.Value_Lambda([FloatRangeColumn]() { return FloatRangeColumn->TestValue; })
-				.OnValueCommitted_Lambda([Chooser, FloatRangeColumn](float NewValue, ETextCommit::Type CommitType) { FloatRangeColumn->TestValue = NewValue; })
-			];
+            DebugWidget = SNew(SNumericEntryBox<float>).IsEnabled_Lambda([Chooser](){ return !Chooser->HasDebugTarget(); })
+            				.Value_Lambda([FloatRangeColumn]() { return FloatRangeColumn->TestValue; })
+            				.OnValueCommitted_Lambda([Chooser, FloatRangeColumn](float NewValue, ETextCommit::Type CommitType)
+            				{
+            					FloatRangeColumn->TestValue = NewValue;
+            				});
 		}
 
-		return ColumnHeaderWidget;
+		return MakeColumnHeaderWidget(Chooser, Column, ColumnName, ColumnTooltip, ColumnIcon, DebugWidget);
 	}
 
 	// create cell widget
diff --git a/Engine/Plugins/Experimental/Chooser/Source/ChooserEditor/Private/GameplayTagColumnEditor.cpp b/Engine/Plugins/Experimental/Chooser/Source/ChooserEditor/Private/GameplayTagColumnEditor.cpp
--- a/Engine/Plugins/Experimental/Chooser/Source/ChooserEditor/Private/GameplayTagColumnEditor.cpp
+++ b/Engine/Plugins/Experimental/Chooser/Source/ChooserEditor/Private/GameplayTagColumnEditor.cpp
@@ -1,6 +1,9 @@
 // Copyright Epic Games, Inc. All Rights Reserved.
 
 #include "GameplayTagColumnEditor.h"
+
+#include <ChooserColumnHeader.h>
+
 #include "SPropertyAccessChainWidget.h"
 #include "GameplayTagColumn.h"
 #include "ObjectChooserWidgetFactories.h"
@@ -25,62 +28,36 @@ TSharedRef<SWidget> CreateGameplayTagColumnWidget(UChooserTable* Chooser, FChoos
 	else if (Row == ColumnWidget_SpecialIndex_Header)
 	{
 		// create column header widget
-		TSharedPtr<SWidget> InputValueWidget = nullptr;
-		if (FChooserParameterBase* InputValue = Column->GetInputValue())
-		{
-			InputValueWidget = FObjectChooserWidgetFactories::CreateWidget(false, Chooser, InputValue, Column->GetInputType(), Chooser->OutputObjectType);
-		}
-		
 		const FSlateBrush* ColumnIcon = FCoreStyle::Get().GetBrush("Icons.Filter");
+		const FText ColumnTooltip = LOCTEXT("Gameplay Tag Tooltip", "Gameplay Tag: cells pass if the input gameplay tag collection matches the cell data (accoding to comparison settings in the column properties).");
+		const FText ColumnName = LOCTEXT("Gameplay Tag","Gameplay Tag");
 		
-		TSharedRef<SWidget> ColumnHeaderWidget = SNew(SHorizontalBox)
-			+ SHorizontalBox::Slot().AutoWidth()
-			[
-				SNew(SBorder)
-				.BorderBackgroundColor(FLinearColor(0,0,0,0))
-				.Content()
-				[
-					SNew(SImage).Image(ColumnIcon)
-				]
-			]
-			+ SHorizontalBox::Slot()
-			[
-				InputValueWidget ? InputValueWidget.ToSharedRef() : SNullWidget::NullWidget
-			];
-	
+		TSharedPtr<SWidget> DebugWidget = nullptr;
 		if (Chooser->GetEnableDebugTesting())
 		{
-			ColumnHeaderWidget = SNew(SVerticalBox)
-			+ SVerticalBox::Slot()
-			[
-				ColumnHeaderWidget
-			]
-			+ SVerticalBox::Slot()
-			[
-				SNew(SSimpleComboButton)
-					.IsEnabled_Lambda([Chooser]()
-					{
-						 return !Chooser->HasDebugTarget();
-					})
-					.Text_Lambda([GameplayTagColumn]()
-					{
-						FText Text = FText::FromString(GameplayTagColumn->TestValue.ToStringSimple(false));
-						if (Text.IsEmpty())
-						{
-							Text = LOCTEXT("None", "None");
-						}
-						return Text;
-					})	
-					.OnGetMenuContent_Lambda([Chooser, GameplayTagColumn, Row]()
+			DebugWidget = SNew(SSimpleComboButton)
+				.IsEnabled_Lambda([Chooser]()
+				{
+					 return !Chooser->HasDebugTarget();
+				})
+				.Text_Lambda([GameplayTagColumn]()
+				{
+					FText Text = FText::FromString(GameplayTagColumn->TestValue.ToStringSimple(false));
+					if (Text.IsEmpty())
 					{
-						TArray<SGameplayTagWidget::FEditableGameplayTagContainerDatum> EditableContainers;
-						EditableContainers.Emplace(Chooser, &(GameplayTagColumn->TestValue));
-						return TSharedRef<SWidget>(SNew(SGameplayTagWidget, EditableContainers));
-					})
-			];
+						Text = LOCTEXT("None", "None");
+					}
+					return Text;
+				})	
+				.OnGetMenuContent_Lambda([Chooser, GameplayTagColumn, Row]()
+				{
+					TArray<SGameplayTagWidget::FEditableGameplayTagContainerDatum> EditableContainers;
+					EditableContainers.Emplace(Chooser, &(GameplayTagColumn->TestValue));
+					return TSharedRef<SWidget>(SNew(SGameplayTagWidget, EditableContainers));
+				});
 		}
 
-		return ColumnHeaderWidget;
+		return MakeColumnHeaderWidget(Chooser, Column, ColumnName, ColumnTooltip, ColumnIcon, DebugWidget);
 	}
 
 	// create cell widget
diff --git a/Engine/Plugins/Experimental/Chooser/Source/ChooserEditor/Private/MultiEnumColumnEditor.cpp b/Engine/Plugins/Experimental/Chooser/Source/ChooserEditor/Private/MultiEnumColumnEditor.cpp
--- a/Engine/Plugins/Experimental/Chooser/Source/ChooserEditor/Private/MultiEnumColumnEditor.cpp
+++ b/Engine/Plugins/Experimental/Chooser/Source/ChooserEditor/Private/MultiEnumColumnEditor.cpp
@@ -1,9 +1,11 @@
 // Copyright Epic Games, Inc. All Rights Reserved.
 
 #include "MultiEnumColumnEditor.h"
+#include "EnumColumnEditor.h"
 #include "MultiEnumColumn.h"
 #include "SPropertyAccessChainWidget.h"
 #include "ObjectChooserWidgetFactories.h"
+#include "ChooserColumnHeader.h"
 #include "ChooserTableEditor.h"
 #include "GraphEditorSettings.h"
 #include "SEnumCombo.h"
@@ -139,47 +141,27 @@ TSharedRef<SWidget> CreateMultiEnumColumnWidget(UChooserTable* Chooser, FChooser
 	}
 	if (Row == ColumnWidget_SpecialIndex_Header)
 	{
-		// create column header widget
-		TSharedPtr<SWidget> InputValueWidget = nullptr;
-		if (FChooserParameterBase* InputValue = Column->GetInputValue())
-		{
-			InputValueWidget = FObjectChooserWidgetFactories::CreateWidget(false, Chooser, InputValue, Column->GetInputType(), Chooser->OutputObjectType);
-		}
-		
 		const FSlateBrush* ColumnIcon = FCoreStyle::Get().GetBrush("Icons.Filter");
+		const FText ColumnTooltip = LOCTEXT("Multi Enum Tooltip", "Enum Any: cells will pass if the input value is any of the enum values checked in the cell");
+		const FText ColumnName = LOCTEXT("Enum Or","Enum (Or)");
 		
-		TSharedRef<SWidget> ColumnHeaderWidget = SNew(SHorizontalBox)
-			+ SHorizontalBox::Slot().AutoWidth()
-			[
-				SNew(SBorder)
-				.BorderBackgroundColor(FLinearColor(0,0,0,0))
-				.Content()
-				[
-					SNew(SImage).Image(ColumnIcon)
-				]
-			]
-			+ SHorizontalBox::Slot()
-			[
-				InputValueWidget ? InputValueWidget.ToSharedRef() : SNullWidget::NullWidget
-			];
 	
+		TSharedPtr<SWidget> DebugWidget = nullptr;
 		if (Chooser->GetEnableDebugTesting())
 		{
-			ColumnHeaderWidget = SNew(SVerticalBox)
-			+ SVerticalBox::Slot()
-			[
-				ColumnHeaderWidget
-			]
-			+ SVerticalBox::Slot()
-			[
-				SNew(SMultiEnumCell<FMultiEnumColumn>).TransactionObject(Chooser).MultiEnumColumn(MultiEnumColumn)
-					.OnValueSet_Lambda([MultiEnumColumn](int Value) { MultiEnumColumn->TestValue = Value; })
-					.EnumValue_Lambda([MultiEnumColumn]() { return MultiEnumColumn->TestValue; })
-					.IsEnabled_Lambda([Chooser] { return !Chooser->HasDebugTarget(); })
-			];
+			DebugWidget = SNew(SEnumCell<FMultiEnumColumn>).TransactionObject(Chooser).EnumColumn(MultiEnumColumn)
+								.OnValueSet_Lambda([MultiEnumColumn](int Value) { MultiEnumColumn->TestValue = Value; })
+								.EnumValue_Lambda([MultiEnumColumn]() { return MultiEnumColumn->TestValue; })
+								.IsEnabled_Lambda([Chooser] { return !Chooser->HasDebugTarget(); });
+			
+			// need to fix support for bitfield enums:
+			// DebugWidget = SNew(SMultiEnumCell<FMultiEnumColumn>).TransactionObject(Chooser).MultiEnumColumn(MultiEnumColumn)
+   //                  					.OnValueSet_Lambda([MultiEnumColumn](int Value) { MultiEnumColumn->TestValue = Value; })
+   //                  					.EnumValue_Lambda([MultiEnumColumn]() { return MultiEnumColumn->TestValue; })
+   //                  					.IsEnabled_Lambda([Chooser] { return !Chooser->HasDebugTarget(); });
 		}
 
-		return ColumnHeaderWidget;
+		return MakeColumnHeaderWidget(Chooser, Column, ColumnName, ColumnTooltip, ColumnIcon, DebugWidget);
 	}
 
 	// create cell widget
diff --git a/Engine/Plugins/Experimental/Chooser/Source/ChooserEditor/Private/ObjectClassColumnEditor.cpp b/Engine/Plugins/Experimental/Chooser/Source/ChooserEditor/Private/ObjectClassColumnEditor.cpp
--- a/Engine/Plugins/Experimental/Chooser/Source/ChooserEditor/Private/ObjectClassColumnEditor.cpp
+++ b/Engine/Plugins/Experimental/Chooser/Source/ChooserEditor/Private/ObjectClassColumnEditor.cpp
@@ -2,6 +2,8 @@
 
 #include "ObjectClassColumnEditor.h"
 
+#include <ChooserColumnHeader.h>
+
 #include "ClassViewerFilter.h"
 #include "SPropertyAccessChainWidget.h"
 #include "GraphEditorSettings.h"
@@ -47,30 +49,11 @@ namespace UE::ChooserEditor
 		}
 		else if (Row == ColumnWidget_SpecialIndex_Header)
 		{
-			// create column header widget
-			TSharedPtr<SWidget> InputValueWidget = nullptr;
-			if (FChooserParameterBase* InputValue = Column->GetInputValue())
-			{
-				InputValueWidget = FObjectChooserWidgetFactories::CreateWidget(false, Chooser, InputValue, Column->GetInputType(), Chooser->OutputObjectType);
-			}
-			
 			const FSlateBrush* ColumnIcon = FCoreStyle::Get().GetBrush("Icons.Filter");
+			const FText ColumnTooltip = LOCTEXT("Object Class Tooltip", "Object Class: cells pass if the Object input has a type that matches the cell's Class setting");
+			const FText ColumnName = LOCTEXT("Object Class","Object Class");
 			
-			TSharedRef<SWidget> ColumnHeaderWidget = SNew(SHorizontalBox)
-				+ SHorizontalBox::Slot().AutoWidth()
-				[
-					SNew(SBorder)
-					.BorderBackgroundColor(FLinearColor(0,0,0,0))
-					.Content()
-					[
-						SNew(SImage).Image(ColumnIcon)
-					]
-				]
-				+ SHorizontalBox::Slot()
-				[
-					InputValueWidget ? InputValueWidget.ToSharedRef() : SNullWidget::NullWidget
-				];
-
+			TSharedPtr<SWidget> DebugWidget = nullptr;
 			if (Chooser->GetEnableDebugTesting())
 			{
 				UClass* AllowedClass = UObject::StaticClass();
@@ -82,31 +65,21 @@ namespace UE::ChooserEditor
 				
 				// create widget for test value object picker
 				TSharedRef<SWidget> ObjectPicker = SNew(SObjectPropertyEntryBox)
-                                                   			.ObjectPath_Lambda([ObjectClassColumn]() {
-                                                   				return ObjectClassColumn->TestValue.ToString();
-                                                   			})
-                                                   			.OnObjectChanged_Lambda([ObjectClassColumn](const FAssetData& AssetData) {
-                                                   				ObjectClassColumn->TestValue = AssetData.ToSoftObjectPath();
-                                                   			})
+															.ObjectPath_Lambda([ObjectClassColumn]() {
+																return ObjectClassColumn->TestValue.ToString();
+															})
+															.OnObjectChanged_Lambda([ObjectClassColumn](const FAssetData& AssetData) {
+																ObjectClassColumn->TestValue = AssetData.ToSoftObjectPath();
+															})
 															.AllowedClass(AllowedClass)
-                                                   			.DisplayUseSelected(false)
-                                                   			.DisplayBrowse(false)
-                                                   			.DisplayThumbnail(false);
+															.DisplayUseSelected(false)
+															.DisplayBrowse(false)
+															.DisplayThumbnail(false);
 				
 				ObjectPicker->SetEnabled(TAttribute<bool>::CreateLambda([Chooser]() { return !Chooser->HasDebugTarget(); }));
-				
-				ColumnHeaderWidget = SNew(SVerticalBox)
-				+ SVerticalBox::Slot()
-				[
-					ColumnHeaderWidget
-				]
-				+ SVerticalBox::Slot()
-				[
-					ObjectPicker
-				];
+				DebugWidget = ObjectPicker;
 			}
-			
-			return ColumnHeaderWidget;
+			return MakeColumnHeaderWidget(Chooser, Column, ColumnName, ColumnTooltip, ColumnIcon, DebugWidget);
 		}
 
 		// create widget for cell
diff --git a/Engine/Plugins/Experimental/Chooser/Source/ChooserEditor/Private/ObjectColumnEditor.cpp b/Engine/Plugins/Experimental/Chooser/Source/ChooserEditor/Private/ObjectColumnEditor.cpp
--- a/Engine/Plugins/Experimental/Chooser/Source/ChooserEditor/Private/ObjectColumnEditor.cpp
+++ b/Engine/Plugins/Experimental/Chooser/Source/ChooserEditor/Private/ObjectColumnEditor.cpp
@@ -1,6 +1,9 @@
 // Copyright Epic Games, Inc. All Rights Reserved.
 
 #include "ObjectColumnEditor.h"
+
+#include <ChooserColumnHeader.h>
+
 #include "SPropertyAccessChainWidget.h"
 #include "GraphEditorSettings.h"
 #include "ObjectChooserWidgetFactories.h"
@@ -99,32 +102,11 @@ namespace UE::ChooserEditor
 		}
 		else if (Row == ColumnWidget_SpecialIndex_Header)
 		{
-			// create column header widget
-			TSharedPtr<SWidget> InputValueWidget = nullptr;
-			if (FChooserParameterBase* InputValue = Column->GetInputValue())
-			{
-				InputValueWidget = FObjectChooserWidgetFactories::CreateWidget(false, Chooser, InputValue, Column->GetInputType(), Chooser->OutputObjectType);
-			}
-			
 			const FSlateBrush* ColumnIcon = FCoreStyle::Get().GetBrush("Icons.Filter");
-			// ColumnIcon = FCoreStyle::Get().GetBrush("Icons.ArrowRight");
-			// ColumnIcon = FAppStyle::Get().GetBrush("Icons.Help");
+			const FText ColumnTooltip = LOCTEXT("Object Tooltip", "Object: cells pass if the Object input is equal to the Object set in the cell");
+			const FText ColumnName = LOCTEXT("Object","Object");
 			
-			TSharedRef<SWidget> ColumnHeaderWidget = SNew(SHorizontalBox)
-				+ SHorizontalBox::Slot().AutoWidth()
-				[
-					SNew(SBorder)
-					.BorderBackgroundColor(FLinearColor(0,0,0,0))
-					.Content()
-					[
-						SNew(SImage).Image(ColumnIcon)
-					]
-				]
-				+ SHorizontalBox::Slot()
-				[
-					InputValueWidget ? InputValueWidget.ToSharedRef() : SNullWidget::NullWidget
-				];
-		
+			TSharedPtr<SWidget> DebugWidget = nullptr;
 			if (Chooser->GetEnableDebugTesting())
 			{
 				TSharedRef<SWidget> ObjectPicker = CreateObjectPicker(Chooser, ObjectColumn, Row);
@@ -144,18 +126,9 @@ namespace UE::ChooserEditor
 				CellWidget->ObjectPickerFactory = [Chooser, ObjectColumn, Row]() { return CreateObjectPicker(Chooser, ObjectColumn, Row); };
 				CellWidget->ObjectPickerSlot = 0;
 				
-				ColumnHeaderWidget = SNew(SVerticalBox)
-				+ SVerticalBox::Slot()
-				[
-					ColumnHeaderWidget
-				]
-				+ SVerticalBox::Slot()
-				[
-					CellWidget
-				];
+				DebugWidget = CellWidget;
 			}
-	
-			return ColumnHeaderWidget;
+			return MakeColumnHeaderWidget(Chooser, Column, ColumnName, ColumnTooltip, ColumnIcon, DebugWidget);
 		}
 
 
diff --git a/Engine/Plugins/Experimental/Chooser/Source/ChooserEditor/Private/OutputFloatColumnEditor.cpp b/Engine/Plugins/Experimental/Chooser/Source/ChooserEditor/Private/OutputFloatColumnEditor.cpp
--- a/Engine/Plugins/Experimental/Chooser/Source/ChooserEditor/Private/OutputFloatColumnEditor.cpp
+++ b/Engine/Plugins/Experimental/Chooser/Source/ChooserEditor/Private/OutputFloatColumnEditor.cpp
@@ -1,6 +1,9 @@
 // Copyright Epic Games, Inc. All Rights Reserved.
 
 #include "OutputFloatColumnEditor.h"
+
+#include <ChooserColumnHeader.h>
+
 #include "OutputFloatColumn.h"
 #include "ObjectChooserWidgetFactories.h"
 #include "ChooserTableEditor.h"
@@ -21,45 +24,20 @@ TSharedRef<SWidget> CreateOutputFloatColumnWidget(UChooserTable* Chooser, FChoos
 
     if (Row == ColumnWidget_SpecialIndex_Header)
 	{
-		// create column header widget
-		TSharedPtr<SWidget> InputValueWidget = nullptr;
-		if (FChooserParameterBase* InputValue = Column->GetInputValue())
-		{
-			InputValueWidget = FObjectChooserWidgetFactories::CreateWidget(false, Chooser, InputValue, Column->GetInputType(), Chooser->OutputObjectType);
-		}
-		
+    	// create column header widget
 		const FSlateBrush* ColumnIcon = FCoreStyle::Get().GetBrush("Icons.ArrowRight");
-		
-		TSharedRef<SWidget> ColumnHeaderWidget = SNew(SHorizontalBox)
-			+ SHorizontalBox::Slot().AutoWidth()
-			[
-				SNew(SBorder)
-				.BorderBackgroundColor(FLinearColor(0,0,0,0))
-				.Content()
-				[
-					SNew(SImage).Image(ColumnIcon)
-				]
-			]
-			+ SHorizontalBox::Slot()
-			[
-				InputValueWidget ? InputValueWidget.ToSharedRef() : SNullWidget::NullWidget
-			];
-	
+		const FText ColumnTooltip = LOCTEXT("Output Float Tooltip", "Output Float: writes the value from cell in the result row to the bound variable");
+		const FText ColumnName = LOCTEXT("Output Float","Output Float");
+        		
+		TSharedPtr<SWidget> DebugWidget = nullptr;
 		if (Chooser->GetEnableDebugTesting())
 		{
-			ColumnHeaderWidget = SNew(SVerticalBox)
-			+ SVerticalBox::Slot()
-			[
-				ColumnHeaderWidget
-			]
-			+ SVerticalBox::Slot()
-			[
-				SNew(SNumericEntryBox<float>).IsEnabled(false)
-				.Value_Lambda([OutputFloatColumn]() { return OutputFloatColumn->TestValue; })
-			];
+			DebugWidget = SNew(SNumericEntryBox<float>)
+				.IsEnabled(false)
+				.Value_Lambda([OutputFloatColumn]() { return OutputFloatColumn->TestValue; });
 		}
-
-		return ColumnHeaderWidget;
+        
+		return MakeColumnHeaderWidget(Chooser, Column, ColumnName, ColumnTooltip, ColumnIcon, DebugWidget);
 	}
 	if (Row == ColumnWidget_SpecialIndex_Fallback)
     {
diff --git a/Engine/Plugins/Experimental/Chooser/Source/ChooserEditor/Private/OutputObjectColumnEditor.cpp b/Engine/Plugins/Experimental/Chooser/Source/ChooserEditor/Private/OutputObjectColumnEditor.cpp
--- a/Engine/Plugins/Experimental/Chooser/Source/ChooserEditor/Private/OutputObjectColumnEditor.cpp
+++ b/Engine/Plugins/Experimental/Chooser/Source/ChooserEditor/Private/OutputObjectColumnEditor.cpp
@@ -1,6 +1,9 @@
 // Copyright Epic Games, Inc. All Rights Reserved.
 
 #include "OutputObjectColumnEditor.h"
+
+#include <ChooserColumnHeader.h>
+
 #include "SPropertyAccessChainWidget.h"
 #include "GraphEditorSettings.h"
 #include "ObjectChooserWidgetFactories.h"
@@ -72,30 +75,13 @@ namespace UE::ChooserEditor
 				else if (Row == ColumnWidget_SpecialIndex_Header)
 				{
 					// create column header widget
-					TSharedPtr<SWidget> InputValueWidget = nullptr;
-					if (FChooserParameterBase* InputValue = Column->GetInputValue())
-					{
-						InputValueWidget = FObjectChooserWidgetFactories::CreateWidget(false, Chooser, InputValue, Column->GetInputType(), ContextOwner->OutputObjectType);
-					}
-					
 					const FSlateBrush* ColumnIcon = FCoreStyle::Get().GetBrush("Icons.ArrowRight");
+					const FText ColumnTooltip = LOCTEXT("Output Object Tooltip", "Output Object: writes the value from cell in the result row to the bound variable");
+					const FText ColumnName = LOCTEXT("Output Object","Output Object");
+							
+					TSharedPtr<SWidget> DebugWidget = nullptr;
 					
-					TSharedRef<SWidget> ColumnHeaderWidget = SNew(SHorizontalBox)
-						+ SHorizontalBox::Slot().AutoWidth()
-						[
-							SNew(SBorder)
-							.BorderBackgroundColor(FLinearColor(0,0,0,0))
-							.Content()
-							[
-								SNew(SImage).Image(ColumnIcon)
-							]
-						]
-						+ SHorizontalBox::Slot()
-						[
-							InputValueWidget ? InputValueWidget.ToSharedRef() : SNullWidget::NullWidget
-						];
-				
-					return ColumnHeaderWidget;
+					return MakeColumnHeaderWidget(Chooser, Column, ColumnName, ColumnTooltip, ColumnIcon, DebugWidget);
 				}
  
 				if (Column->RowValues.IsValidIndex(Row))
diff --git a/Engine/Plugins/Experimental/Chooser/Source/ChooserEditor/Private/OutputStructColumnEditor.cpp b/Engine/Plugins/Experimental/Chooser/Source/ChooserEditor/Private/OutputStructColumnEditor.cpp
--- a/Engine/Plugins/Experimental/Chooser/Source/ChooserEditor/Private/OutputStructColumnEditor.cpp
+++ b/Engine/Plugins/Experimental/Chooser/Source/ChooserEditor/Private/OutputStructColumnEditor.cpp
@@ -1,6 +1,9 @@
 // Copyright Epic Games, Inc. All Rights Reserved.
 
 #include "OutputStructColumnEditor.h"
+
+#include <ChooserColumnHeader.h>
+
 #include "OutputStructColumn.h"
 #include "SPropertyAccessChainWidget.h"
 #include "ObjectChooserWidgetFactories.h"
@@ -20,37 +23,14 @@ TSharedRef<SWidget> CreateOutputStructColumnWidget(UChooserTable* Chooser, FChoo
 {
 	if (Row == ColumnWidget_SpecialIndex_Header)
 	{
-		// create column header widget
-		TSharedPtr<SWidget> InputValueWidget = nullptr;
-		if (FChooserParameterBase* InputValue = Column->GetInputValue())
-		{
-			InputValueWidget = FObjectChooserWidgetFactories::CreateWidget(false, Chooser, InputValue, Column->GetInputType(), Chooser->OutputObjectType,
-				FChooserWidgetValueChanged::CreateLambda([Column]()
-				{
-					FOutputStructColumn* StructColumn = static_cast<FOutputStructColumn*>(Column);
-					StructColumn->StructTypeChanged();
-				})
-				);
-		}
-		
+    	// create column header widget
 		const FSlateBrush* ColumnIcon = FCoreStyle::Get().GetBrush("Icons.ArrowRight");
-		
-		TSharedRef<SWidget> ColumnHeaderWidget = SNew(SHorizontalBox)
-			+ SHorizontalBox::Slot().AutoWidth()
-			[
-				SNew(SBorder)
-				.BorderBackgroundColor(FLinearColor(0,0,0,0))
-				.Content()
-				[
-					SNew(SImage).Image(ColumnIcon)
-				]
-			]
-			+ SHorizontalBox::Slot()
-			[
-				InputValueWidget ? InputValueWidget.ToSharedRef() : SNullWidget::NullWidget
-			];
-	
-		return ColumnHeaderWidget;
+		const FText ColumnTooltip = LOCTEXT("Output Struct Tooltip", "Output Struct: writes the value from cell in the result row to the bound variable");
+		const FText ColumnName = LOCTEXT("Output Struct","Output Struct");
+        		
+		TSharedPtr<SWidget> DebugWidget = nullptr;
+        
+		return MakeColumnHeaderWidget(Chooser, Column, ColumnName, ColumnTooltip, ColumnIcon, DebugWidget);
 	}
 	
 	FOutputStructColumn* StructColumn = static_cast<FOutputStructColumn*>(Column);
diff --git a/Engine/Plugins/Experimental/Chooser/Source/ChooserEditor/Private/RandomizeColumnEditor.cpp b/Engine/Plugins/Experimental/Chooser/Source/ChooserEditor/Private/RandomizeColumnEditor.cpp
--- a/Engine/Plugins/Experimental/Chooser/Source/ChooserEditor/Private/RandomizeColumnEditor.cpp
+++ b/Engine/Plugins/Experimental/Chooser/Source/ChooserEditor/Private/RandomizeColumnEditor.cpp
@@ -1,6 +1,9 @@
 // Copyright Epic Games, Inc. All Rights Reserved.
 
 #include "RandomizeColumnEditor.h"
+
+#include <ChooserColumnHeader.h>
+
 #include "RandomizeColumn.h"
 #include "SPropertyAccessChainWidget.h"
 #include "ObjectChooserWidgetFactories.h"
@@ -26,28 +29,14 @@ TSharedRef<SWidget> CreateRandomizeColumnWidget(UChooserTable* Chooser, FChooser
 	}
 	if (Row == ColumnWidget_SpecialIndex_Header)
 	{
+		// create column header widget
 		const FSlateBrush* ColumnIcon = FAppStyle::Get().GetBrush("Icons.Help");
-		
-		TSharedRef<SWidget> ColumnHeaderWidget = SNew(SHorizontalBox)
-			+ SHorizontalBox::Slot().AutoWidth()
-			[
-				SNew(SBorder)
-				.BorderBackgroundColor(FLinearColor(0,0,0,0))
-				.Content()
-				[
-					SNew(SImage).Image(ColumnIcon)
-				]
-			]
-			+ SHorizontalBox::Slot().AutoWidth()
-			[
-				SNew(SSpacer).Size(FVector2D(10,10))
-			]
-			+ SHorizontalBox::Slot().VAlign(VAlign_Center)
-			[
-				SNew(STextBlock).Text(LOCTEXT("Randomize","Randomize"))
-			];
+		const FText ColumnTooltip = LOCTEXT("Randomize Tooltip", "Randomize: randomly selects a single result from the rows which passed all other columns, or the rows with equal, minimum cost, for cost based columns.  Optional Randomization Context variable binding can be used to reduce (or eliminate) the probability of selecting the same entry twice in a row.");
+		const FText ColumnName = LOCTEXT("Randomize","Randomize");
+        		
+		TSharedPtr<SWidget> DebugWidget = nullptr;
 
-		return ColumnHeaderWidget;
+		return MakeColumnHeaderWidget(Chooser, Column, ColumnName, ColumnTooltip, ColumnIcon, DebugWidget);
 	}
 
 	// create cell widget
diff --git a/Engine/Plugins/Experimental/Chooser/Source/ChooserEditor/Private/SChooserColumnHandle.cpp b/Engine/Plugins/Experimental/Chooser/Source/ChooserEditor/Private/SChooserColumnHandle.cpp
--- a/Engine/Plugins/Experimental/Chooser/Source/ChooserEditor/Private/SChooserColumnHandle.cpp
+++ b/Engine/Plugins/Experimental/Chooser/Source/ChooserEditor/Private/SChooserColumnHandle.cpp
@@ -3,6 +3,7 @@
 #include "SChooserColumnHandle.h"
 #include "ChooserTableEditor.h"
 #include "RandomizeColumn.h"
+#include "Widgets/Colors/SColorBlock.h"
 #include "Widgets/Layout/SSeparator.h"
 
 #define LOCTEXT_NAMESPACE "ChooserColumnHandle"
@@ -32,6 +33,19 @@ namespace UE::ChooserEditor
 		ChildSlot
 		[
 			SNew(SOverlay)
+				+ SOverlay::Slot()
+				[
+					SNew(SColorBlock).Color_Lambda([this](){ return ChooserEditor->TableHasFocus() ? FStyleColors::Select.GetSpecifiedColor() : FStyleColors::SelectInactive.GetSpecifiedColor(); } )
+							.Visibility_Lambda(
+							[this]()
+							{
+								if (ChooserEditor->IsColumnSelected(ColumnIndex))
+								{
+									return EVisibility::Visible;
+								}
+								return EVisibility::Hidden;
+							})
+				]
 				+ SOverlay::Slot()
 				[
 					InArgs._Content.Widget
diff --git a/Engine/Plugins/Experimental/Chooser/Source/ChooserEditor/Private/SChooserTableRow.cpp b/Engine/Plugins/Experimental/Chooser/Source/ChooserEditor/Private/SChooserTableRow.cpp
--- a/Engine/Plugins/Experimental/Chooser/Source/ChooserEditor/Private/SChooserTableRow.cpp
+++ b/Engine/Plugins/Experimental/Chooser/Source/ChooserEditor/Private/SChooserTableRow.cpp
@@ -127,24 +127,24 @@ namespace UE::ChooserEditor
 				);
 
 				return SNew(SOverlay)
-							+ SOverlay::Slot()
-    						[
-    							ResultWidget.ToSharedRef()
-    						]
-							+ SOverlay::Slot()
-							[
-								SNew(SColorBlock).Color(DisabledColor)
-										.Visibility_Lambda(
-										[this]()
+						+ SOverlay::Slot()
+						[
+							ResultWidget.ToSharedRef()
+						]
+						+ SOverlay::Slot()
+						[
+							SNew(SColorBlock).Color(DisabledColor)
+									.Visibility_Lambda(
+									[this]()
+									{
+										if (Chooser->IsRowDisabled(RowIndex->RowIndex))
 										{
-											if (Chooser->IsRowDisabled(RowIndex->RowIndex))
-											{
-												return EVisibility::HitTestInvisible;
-											}
-											return EVisibility::Hidden;
-										})
-										
-							];
+											return EVisibility::HitTestInvisible;
+										}
+										return EVisibility::Hidden;
+									})
+									
+						];
 			}
 			else
 			{
@@ -161,7 +161,30 @@ namespace UE::ChooserEditor
 						return SNew(SOverlay)
 						+ SOverlay::Slot()
 						[
-							ColumnWidget.ToSharedRef()
+							SNew(SBorder)
+							.BorderBackgroundColor(FLinearColor(0,0,0,0))
+							.Padding(FMargin(4,0))
+							.Content()
+							[
+								SNew(SOverlay)
+								+ SOverlay::Slot()
+								[
+									SNew(SColorBlock).Color_Lambda([this, ColumnIndex](){ return Editor->TableHasFocus() ? FStyleColors::Select.GetSpecifiedColor() : FStyleColors::SelectInactive.GetSpecifiedColor(); } )
+										.Visibility_Lambda(
+										[this, ColumnIndex]()
+										{
+											if (Editor->IsColumnSelected(ColumnIndex))
+											{
+												return EVisibility::Visible;
+											}
+											return EVisibility::Hidden;
+										})
+								]
+								+ SOverlay::Slot()
+								[
+									ColumnWidget.ToSharedRef()
+								]
+							]
 						]
 						+ SOverlay::Slot()
 						[
@@ -267,7 +290,34 @@ namespace UE::ChooserEditor
 				
 					if (ColumnWidget.IsValid())
 					{
-						return ColumnWidget.ToSharedRef();
+						return SNew(SOverlay)
+							+ SOverlay::Slot()
+							[
+								SNew(SBorder)
+								.BorderBackgroundColor(FLinearColor(0,0,0,0))
+								.Padding(FMargin(4,0))
+								.Content()
+								[
+									SNew(SOverlay)
+									+SOverlay::Slot()
+									[
+										SNew(SColorBlock).Color_Lambda([this, ColumnIndex](){ return Editor->TableHasFocus() ? FStyleColors::Select.GetSpecifiedColor() : FStyleColors::SelectInactive.GetSpecifiedColor(); } )
+											.Visibility_Lambda(
+											[this, ColumnIndex]()
+											{
+												if (Editor->IsColumnSelected(ColumnIndex))
+												{
+													return EVisibility::Visible;
+												}
+												return EVisibility::Hidden;
+											})
+									]
+									+ SOverlay::Slot()
+									[
+										ColumnWidget.ToSharedRef()
+									]
+								]
+							];
 					}
 				}
 			}
diff --git a/Engine/Plugins/Experimental/Chooser/Source/ChooserEditor/Public/ChooserColumnHeader.h b/Engine/Plugins/Experimental/Chooser/Source/ChooserEditor/Public/ChooserColumnHeader.h
new file mode 100644
--- /dev/null
+++ b/Engine/Plugins/Experimental/Chooser/Source/ChooserEditor/Public/ChooserColumnHeader.h
@@ -0,0 +1,13 @@
+// Copyright Epic Games, Inc. All Rights Reserved.
+#pragma once
+
+#include "Widgets/SWidget.h"
+
+class UChooserTable;
+struct FChooserColumnBase;
+
+
+namespace UE::ChooserEditor
+{
+	TSharedRef<SWidget> MakeColumnHeaderWidget(UChooserTable* Chooser, FChooserColumnBase* Column, const FText& ColumnName,const FText& ColumnTooltip, const FSlateBrush* ColumnIcon, TSharedPtr<SWidget> DebugWidget);
+}
\ No newline at end of file
