Adding a way to iterate ChooserObjects

diff --git a/Engine/Plugins/Experimental/Chooser/Source/Chooser/Internal/Chooser.h b/Engine/Plugins/Experimental/Chooser/Source/Chooser/Internal/Chooser.h
--- a/Engine/Plugins/Experimental/Chooser/Source/Chooser/Internal/Chooser.h
+++ b/Engine/Plugins/Experimental/Chooser/Source/Chooser/Internal/Chooser.h
@@ -188,6 +188,7 @@ public:
 
 	bool IsCookedData() const { return !CookedResults.IsEmpty(); }
 	static FObjectChooserBase::EIteratorStatus EvaluateChooser(FChooserEvaluationContext& Context, const UChooserTable* Chooser, FObjectChooserBase::FObjectChooserIteratorCallback Callback);
+	static FObjectChooserBase::EIteratorStatus IterateChooser(const UChooserTable* Chooser, FObjectChooserBase::FObjectChooserIteratorCallback Callback);
 };
 
 USTRUCT(BlueprintType, DisplayName = "Nested Chooser")
@@ -197,6 +198,7 @@ struct CHOOSER_API FNestedChooser : public FObjectChooserBase
 
 	virtual UObject* ChooseObject(FChooserEvaluationContext& Context) const final override;
 	virtual EIteratorStatus ChooseMulti(FChooserEvaluationContext &Context, FObjectChooserIteratorCallback Callback) const final override;
+	virtual EIteratorStatus IterateObjects(FObjectChooserIteratorCallback Callback) const final override;
 	virtual void GetDebugName(FString& OutDebugName) const override;
 	
 	public:
@@ -214,6 +216,7 @@ struct CHOOSER_API FEvaluateChooser : public FObjectChooserBase
 
 	virtual UObject* ChooseObject(FChooserEvaluationContext& Context) const final override;
 	virtual EIteratorStatus ChooseMulti(FChooserEvaluationContext &Context, FObjectChooserIteratorCallback Callback) const final override;
+	virtual EIteratorStatus IterateObjects(FObjectChooserIteratorCallback Callback) const final override;
 	virtual void GetDebugName(FString& OutDebugName) const override;
 	
 	public:
diff --git a/Engine/Plugins/Experimental/Chooser/Source/Chooser/Internal/ObjectChooser_Asset.h b/Engine/Plugins/Experimental/Chooser/Source/Chooser/Internal/ObjectChooser_Asset.h
--- a/Engine/Plugins/Experimental/Chooser/Source/Chooser/Internal/ObjectChooser_Asset.h
+++ b/Engine/Plugins/Experimental/Chooser/Source/Chooser/Internal/ObjectChooser_Asset.h
@@ -13,6 +13,7 @@ struct CHOOSER_API FAssetChooser : public FObjectChooserBase
 	
 	// FObjectChooserBase interface
 	virtual UObject* ChooseObject(FChooserEvaluationContext& Context) const final override;
+	virtual EIteratorStatus IterateObjects(FObjectChooserIteratorCallback Callback) const final override;
 
 #if WITH_EDITOR
 	virtual UObject* GetReferencedObject() const override { return Asset; }
@@ -29,6 +30,7 @@ struct CHOOSER_API FSoftAssetChooser : public FObjectChooserBase
 
 	// FObjectChooserBase interface
 	virtual UObject* ChooseObject(FChooserEvaluationContext& Context) const final override;
+	virtual EIteratorStatus IterateObjects(FObjectChooserIteratorCallback Callback) const final override;
 	
 #if WITH_EDITOR
 	virtual UObject* GetReferencedObject() const override { return Asset.LoadSynchronous(); }
diff --git a/Engine/Plugins/Experimental/Chooser/Source/Chooser/Internal/ObjectChooser_Class.h b/Engine/Plugins/Experimental/Chooser/Source/Chooser/Internal/ObjectChooser_Class.h
--- a/Engine/Plugins/Experimental/Chooser/Source/Chooser/Internal/ObjectChooser_Class.h
+++ b/Engine/Plugins/Experimental/Chooser/Source/Chooser/Internal/ObjectChooser_Class.h
@@ -13,6 +13,7 @@ struct CHOOSER_API FClassChooser : public FObjectChooserBase
 	
 	// FObjectChooserBase interface
 	virtual UObject* ChooseObject(FChooserEvaluationContext& Context) const final override;
+	virtual EIteratorStatus IterateObjects(FObjectChooserIteratorCallback Callback) const final override;
 public: 
 	UPROPERTY(EditAnywhere, Category = "Parameters")
 	TObjectPtr<UClass> Class;
diff --git a/Engine/Plugins/Experimental/Chooser/Source/Chooser/Private/Chooser.cpp b/Engine/Plugins/Experimental/Chooser/Source/Chooser/Private/Chooser.cpp
--- a/Engine/Plugins/Experimental/Chooser/Source/Chooser/Private/Chooser.cpp
+++ b/Engine/Plugins/Experimental/Chooser/Source/Chooser/Private/Chooser.cpp
@@ -534,6 +534,59 @@ FObjectChooserBase::EIteratorStatus UChooserTable::EvaluateChooser(FChooserEvalu
 	return bSetOutputs ? FObjectChooserBase::EIteratorStatus::ContinueWithOutputs : FObjectChooserBase::EIteratorStatus::Continue;
 }
 
+FObjectChooserBase::EIteratorStatus UChooserTable::IterateChooser(const UChooserTable* Chooser, FObjectChooserBase::FObjectChooserIteratorCallback Callback)
+{
+	TRACE_CPUPROFILER_EVENT_SCOPE(IterateChooser);
+
+	if (Chooser == nullptr)
+	{
+		return FObjectChooserBase::EIteratorStatus::Continue;
+	}
+
+	const TArray<FInstancedStruct>* ResultsArray = &Chooser->CookedResults;
+
+#if WITH_EDITORONLY_DATA
+	if (!Chooser->IsCookedData())
+	{
+		ResultsArray = &Chooser->ResultsStructs;
+	}
+#endif
+	
+	uint32 Count = ResultsArray->Num();
+	uint32 BufferSize = Count * sizeof(FChooserIndexArray::FIndexData);
+
+	FChooserIndexArray Indices(static_cast<FChooserIndexArray::FIndexData*>(FMemory_Alloca(BufferSize)), Count);
+
+	for(uint32 i=0;i<Count;i++)
+	{
+		if (!Chooser->IsRowDisabled(i))
+		{
+			Indices.Push({i, 0});
+		}
+	}
+
+	for (FChooserIndexArray::FIndexData& SelectedIndexData : Indices)
+	{
+		if (ResultsArray->IsValidIndex(SelectedIndexData.Index))
+		{
+			const FObjectChooserBase& SelectedResult = (*ResultsArray)[SelectedIndexData.Index].Get<FObjectChooserBase>();
+			FObjectChooserBase::EIteratorStatus Status = SelectedResult.IterateObjects(Callback);
+			if (Status == FObjectChooserBase::EIteratorStatus::Stop)
+			{
+				return FObjectChooserBase::EIteratorStatus::Stop;
+			}
+		}
+	}
+
+	if (Chooser->FallbackResult.IsValid())
+	{
+		const FObjectChooserBase& SelectedResult = Chooser->FallbackResult.Get<FObjectChooserBase>();
+		return SelectedResult.IterateObjects(Callback);
+	}
+
+	return FObjectChooserBase::EIteratorStatus::Continue;
+}
+
 UObject* FEvaluateChooser::ChooseObject(FChooserEvaluationContext& Context) const
 {
 	UObject* Result = nullptr;
@@ -551,6 +604,11 @@ FObjectChooserBase::EIteratorStatus FEvaluateChooser::ChooseMulti(FChooserEvalua
 	return UChooserTable::EvaluateChooser(Context, Chooser, Callback);
 }
 
+FObjectChooserBase::EIteratorStatus FEvaluateChooser::IterateObjects(FObjectChooserIteratorCallback Callback) const
+{
+	return UChooserTable::IterateChooser(Chooser, Callback);
+}
+
 void FEvaluateChooser::GetDebugName(FString& OutDebugName) const
 {
 	if (Chooser)
@@ -581,6 +639,11 @@ FObjectChooserBase::EIteratorStatus FNestedChooser::ChooseMulti(FChooserEvaluati
 	return UChooserTable::EvaluateChooser(Context, Chooser, Callback);
 }
 
+FObjectChooserBase::EIteratorStatus FNestedChooser::IterateObjects(FObjectChooserIteratorCallback Callback) const
+{
+	return UChooserTable::IterateChooser(Chooser, Callback);
+}
+
 void FNestedChooser::GetDebugName(FString& OutDebugName) const
 {
 	OutDebugName  = GetNameSafe(Chooser);
diff --git a/Engine/Plugins/Experimental/Chooser/Source/Chooser/Private/ObjectChooser_Asset.cpp b/Engine/Plugins/Experimental/Chooser/Source/Chooser/Private/ObjectChooser_Asset.cpp
--- a/Engine/Plugins/Experimental/Chooser/Source/Chooser/Private/ObjectChooser_Asset.cpp
+++ b/Engine/Plugins/Experimental/Chooser/Source/Chooser/Private/ObjectChooser_Asset.cpp
@@ -7,7 +7,17 @@ UObject* FAssetChooser::ChooseObject(FChooserEvaluationContext& Context) const
 	return Asset;
 }
 
+FObjectChooserBase::EIteratorStatus FAssetChooser::IterateObjects(FObjectChooserIteratorCallback Callback) const
+{
+	return Callback.Execute(Asset);
+}
+
 UObject* FSoftAssetChooser::ChooseObject(FChooserEvaluationContext& Context) const
 {
 	return Asset.LoadSynchronous();
 }
+
+FObjectChooserBase::EIteratorStatus FSoftAssetChooser::IterateObjects(FObjectChooserIteratorCallback Callback) const
+{
+	return Callback.Execute(Asset.LoadSynchronous());
+}
diff --git a/Engine/Plugins/Experimental/Chooser/Source/Chooser/Private/ObjectChooser_Class.cpp b/Engine/Plugins/Experimental/Chooser/Source/Chooser/Private/ObjectChooser_Class.cpp
--- a/Engine/Plugins/Experimental/Chooser/Source/Chooser/Private/ObjectChooser_Class.cpp
+++ b/Engine/Plugins/Experimental/Chooser/Source/Chooser/Private/ObjectChooser_Class.cpp
@@ -5,4 +5,9 @@
 UObject* FClassChooser::ChooseObject(FChooserEvaluationContext& Context) const
 {
 	return Class;
-}
\ No newline at end of file
+}
+
+FObjectChooserBase::EIteratorStatus FClassChooser::IterateObjects(FObjectChooserIteratorCallback Callback) const
+{
+	return Callback.Execute(Class);
+}
diff --git a/Engine/Plugins/Experimental/Chooser/Source/Chooser/Public/IObjectChooser.h b/Engine/Plugins/Experimental/Chooser/Source/Chooser/Public/IObjectChooser.h
--- a/Engine/Plugins/Experimental/Chooser/Source/Chooser/Public/IObjectChooser.h
+++ b/Engine/Plugins/Experimental/Chooser/Source/Chooser/Public/IObjectChooser.h
@@ -108,6 +108,7 @@ public:
 		return EIteratorStatus::Continue;
 	}
 
+	virtual EIteratorStatus IterateObjects(FObjectChooserIteratorCallback Callback) const { return EIteratorStatus::Continue; }
 	virtual void GetDebugName(FString& OutDebugName) const {};
 
 #if WITH_EDITOR
diff --git a/Engine/Plugins/Experimental/Chooser/Source/ProxyTable/Private/ProxyTable.cpp b/Engine/Plugins/Experimental/Chooser/Source/ProxyTable/Private/ProxyTable.cpp
--- a/Engine/Plugins/Experimental/Chooser/Source/ProxyTable/Private/ProxyTable.cpp
+++ b/Engine/Plugins/Experimental/Chooser/Source/ProxyTable/Private/ProxyTable.cpp
@@ -252,6 +252,21 @@ FObjectChooserBase::EIteratorStatus UProxyTable::FindProxyObjectMulti(const FGui
     return FObjectChooserBase::EIteratorStatus::Continue;
 }
 
+FObjectChooserBase::EIteratorStatus UProxyTable::IterateProxyObjects(const FGuid& Key, FObjectChooserBase::FObjectChooserIteratorCallback Callback) const
+{
+	const int FoundIndex = Algo::BinarySearch(Keys, Key);
+	if (FoundIndex != INDEX_NONE)
+	{
+		const FRuntimeProxyValue& EntryValueData = RuntimeValues[FoundIndex];
+		const FObjectChooserBase& EntryValue = EntryValueData.Value.Get<const FObjectChooserBase>();
+
+		FObjectChooserBase::EIteratorStatus Result = EntryValue.IterateObjects(Callback);
+		return Result;
+	}
+
+	return FObjectChooserBase::EIteratorStatus::Continue;
+}
+
 UProxyTable::UProxyTable(const FObjectInitializer& Initializer)
 	:Super(Initializer)
 {
diff --git a/Engine/Plugins/Experimental/Chooser/Source/ProxyTable/Public/ProxyTable.h b/Engine/Plugins/Experimental/Chooser/Source/ProxyTable/Public/ProxyTable.h
--- a/Engine/Plugins/Experimental/Chooser/Source/ProxyTable/Public/ProxyTable.h
+++ b/Engine/Plugins/Experimental/Chooser/Source/ProxyTable/Public/ProxyTable.h
@@ -86,6 +86,8 @@ public:
 	FObjectChooserBase::EIteratorStatus FindProxyObjectMulti(const FGuid& Key, FChooserEvaluationContext &Context, FObjectChooserBase::FObjectChooserIteratorCallback Callback) const;
 	UObject* FindProxyObject(const FGuid& Key, FChooserEvaluationContext& Context) const;
 
+	FObjectChooserBase::EIteratorStatus IterateProxyObjects(const FGuid& Key, FObjectChooserBase::FObjectChooserIteratorCallback Callback) const;
+
 	virtual void PostLoad() override;
 #if WITH_EDITORONLY_DATA
 public:
