Fix duplicate LOCTEXT in EnumColumnEditor, and change != to say Not

diff --git a/Engine/Plugins/Experimental/Chooser/Source/ChooserEditor/Private/EnumColumnEditor.cpp b/Engine/Plugins/Experimental/Chooser/Source/ChooserEditor/Private/EnumColumnEditor.cpp
--- a/Engine/Plugins/Experimental/Chooser/Source/ChooserEditor/Private/EnumColumnEditor.cpp
+++ b/Engine/Plugins/Experimental/Chooser/Source/ChooserEditor/Private/EnumColumnEditor.cpp
@@ -61,7 +61,7 @@ TSharedRef<SWidget> CreateEnumColumnWidget(UChooserTable* Chooser, FChooserColum
 							return LOCTEXT("CompEqual", "=");
 
 						case EEnumColumnCellValueComparison::MatchNotEqual:
-							return LOCTEXT("CompNotEqual", "!=");
+							return LOCTEXT("CompNotEqual", "Not");
 
 						case EEnumColumnCellValueComparison::MatchAny:
 							return LOCTEXT("CompAny", "Any");
@@ -114,7 +114,7 @@ TSharedRef<SWidget> CreateOutputEnumColumnWidget(UChooserTable* Chooser, FChoose
 	if (Row == ColumnWidget_SpecialIndex_Header)
 	{
 		const FSlateBrush* ColumnIcon = FCoreStyle::Get().GetBrush("Icons.ArrowRight");
-		const FText ColumnTooltip = LOCTEXT("Output Enum", "Output Enum:  writes the value from cell in the result row to the bound variable");
+		const FText ColumnTooltip = LOCTEXT("Output Enum Tooltip", "Output Enum:  writes the value from cell in the result row to the bound variable");
 		const FText ColumnName = LOCTEXT("Output Enum","Output Enum");
 		
 		TSharedPtr<SWidget> DebugWidget = nullptr;
diff --git a/Engine/Plugins/Experimental/Chooser/Source/ChooserEditor/Private/ObjectColumnEditor.cpp b/Engine/Plugins/Experimental/Chooser/Source/ChooserEditor/Private/ObjectColumnEditor.cpp
--- a/Engine/Plugins/Experimental/Chooser/Source/ChooserEditor/Private/ObjectColumnEditor.cpp
+++ b/Engine/Plugins/Experimental/Chooser/Source/ChooserEditor/Private/ObjectColumnEditor.cpp
@@ -155,7 +155,7 @@ namespace UE::ChooserEditor
 									return LOCTEXT("CompEqual", "=");
 
 								case EObjectColumnCellValueComparison::MatchNotEqual:
-									return LOCTEXT("CompNotEqual", "!=");
+									return LOCTEXT("CompNotEqual", "Not");
 
 								case EObjectColumnCellValueComparison::MatchAny:
 									return LOCTEXT("CompAny", "Any");
