Add cpu profiling scopes with chooser name

diff --git a/Engine/Plugins/Experimental/Chooser/Source/Chooser/Private/Chooser.cpp b/Engine/Plugins/Experimental/Chooser/Source/Chooser/Private/Chooser.cpp
--- a/Engine/Plugins/Experimental/Chooser/Source/Chooser/Private/Chooser.cpp
+++ b/Engine/Plugins/Experimental/Chooser/Source/Chooser/Private/Chooser.cpp
@@ -382,6 +382,7 @@ void UChooserTable::UpdateDebugging(FChooserEvaluationContext& Context) const
 FObjectChooserBase::EIteratorStatus UChooserTable::EvaluateChooser(FChooserEvaluationContext& Context, const UChooserTable* Chooser, FObjectChooserBase::FObjectChooserIteratorCallback Callback)
 {
 	TRACE_CPUPROFILER_EVENT_SCOPE(EvaluateChooser);
+	TRACE_CPUPROFILER_EVENT_SCOPE_TEXT(*Chooser->GetName());
 
 	if (Chooser == nullptr)
 	{
