Add Duplicate command to Chooser Table editor

diff --git a/Engine/Plugins/Experimental/Chooser/Source/ChooserEditor/Private/ChooserTableEditor.cpp b/Engine/Plugins/Experimental/Chooser/Source/ChooserEditor/Private/ChooserTableEditor.cpp
--- a/Engine/Plugins/Experimental/Chooser/Source/ChooserEditor/Private/ChooserTableEditor.cpp
+++ b/Engine/Plugins/Experimental/Chooser/Source/ChooserEditor/Private/ChooserTableEditor.cpp
@@ -391,6 +391,7 @@ void FChooserTableEditor::RegisterMenus()
 		Section.AddEntry(FToolMenuEntry::InitMenuEntry(FGenericCommands::Get().Copy));
 		Section.AddEntry(FToolMenuEntry::InitMenuEntry(FGenericCommands::Get().Cut));
 		Section.AddEntry(FToolMenuEntry::InitMenuEntry(FGenericCommands::Get().Paste));
+		Section.AddEntry(FToolMenuEntry::InitMenuEntry(FGenericCommands::Get().Duplicate));
 		Section.AddEntry(FToolMenuEntry::InitMenuEntry(FGenericCommands::Get().Delete));
 		Section.AddEntry(FToolMenuEntry::InitMenuEntry(Commands.Disable));
 		Section.AddEntry(FToolMenuEntry::InitMenuEntry(Commands.MoveUp));
@@ -451,6 +452,7 @@ void FChooserTableEditor::RegisterMenus()
 				MenuBuilder.AddMenuEntry(FGenericCommands::Get().Copy, NAME_None);
 				MenuBuilder.AddMenuEntry(FGenericCommands::Get().Cut, NAME_None);
 				MenuBuilder.AddMenuEntry(FGenericCommands::Get().Paste, NAME_None);
+				MenuBuilder.AddMenuEntry(FGenericCommands::Get().Duplicate, NAME_None, LOCTEXT("Duplicate Selection", "Duplicate Selection"));
 				MenuBuilder.AddMenuEntry(FGenericCommands::Get().Delete, NAME_None, LOCTEXT("Delete Selection", "Delete Selection"));
 				MenuBuilder.AddMenuEntry(FChooserTableEditorCommands::Get().Disable, NAME_None, LOCTEXT("Disable Selection", "Disable Selection"));
 				MenuBuilder.AddMenuEntry(FChooserTableEditorCommands::Get().RemoveDisabledData, NAME_None);
@@ -504,6 +506,12 @@ void FChooserTableEditor::BindCommands()
 		FCanExecuteAction::CreateSP(this, &FChooserTableEditor::HasSelection)
 		);
 
+	ToolkitCommands->MapAction(
+		FGenericCommands::Get().Duplicate,
+		FExecuteAction::CreateSP(this, &FChooserTableEditor::DuplicateSelection),
+		FCanExecuteAction::CreateSP(this, &FChooserTableEditor::HasSelection)
+		);
+
 	ToolkitCommands->MapAction(
 		Commands.AutoPopulateSelection,
 		FExecuteAction::CreateSP(this, &FChooserTableEditor::AutoPopulateSelection),
@@ -1759,6 +1767,31 @@ void FChooserTableEditor::DeleteSelection()
 	}
 }
 
+void FChooserTableEditor::DuplicateSelection()
+{
+	if (HasRowsSelected())
+	{
+		const FScopedTransaction Transaction(LOCTEXT("Duplicate Row(s)", "Duplicate Row(s)"));
+		UChooserTable* RowCopy = CopySelectionInternal();
+		int MaxSelectedRow = -1;
+		for(UChooserRowDetails* SelectedRow : SelectedRows)
+		{
+			MaxSelectedRow = FMath::Max(SelectedRow->Row, MaxSelectedRow);
+		}
+		PasteInternal(RowCopy, MaxSelectedRow + 1);
+	}
+	else if(HasColumnSelected())
+	{
+		const FScopedTransaction Transaction(LOCTEXT("Duplicate Column", "Duplicate Column"));
+		UChooserTable* Chooser = GetChooser();
+		Chooser->Modify();
+		FInstancedStruct Column = Chooser->ColumnsStructs[SelectedColumn->Column];
+		Chooser->ColumnsStructs.Insert(Column,SelectedColumn->Column);
+		RefreshAll();
+	}
+}
+
+
 bool FChooserTableEditor::CanMoveRowsUp()
 {
 	if (HasRowsSelected())
diff --git a/Engine/Plugins/Experimental/Chooser/Source/ChooserEditor/Private/ChooserTableEditor.h b/Engine/Plugins/Experimental/Chooser/Source/ChooserEditor/Private/ChooserTableEditor.h
--- a/Engine/Plugins/Experimental/Chooser/Source/ChooserEditor/Private/ChooserTableEditor.h
+++ b/Engine/Plugins/Experimental/Chooser/Source/ChooserEditor/Private/ChooserTableEditor.h
@@ -105,6 +105,7 @@ namespace UE::ChooserEditor
 		bool IsSelectionDisabled() const;
 		void ToggleDisableSelection();
 		void DeleteSelection();
+		void DuplicateSelection();
 
 		bool CanMoveRowsUp();
 		void MoveRowsUp();
