Add ChooserTable creation dialog

diff --git a/Engine/Plugins/Experimental/Chooser/Source/Chooser/Internal/Chooser.h b/Engine/Plugins/Experimental/Chooser/Source/Chooser/Internal/Chooser.h
--- a/Engine/Plugins/Experimental/Chooser/Source/Chooser/Internal/Chooser.h
+++ b/Engine/Plugins/Experimental/Chooser/Source/Chooser/Internal/Chooser.h
@@ -175,15 +175,15 @@ public:
 	// Columns which filter Results
 	UPROPERTY(EditAnywhere, NoClear, DisplayName = "Columns", Category = Hidden, meta = (ExcludeBaseStruct, BaseStruct = "/Script/Chooser.ChooserColumnBase"))
 	TArray<FInstancedStruct> ColumnsStructs;
-	
-	UPROPERTY(EditAnywhere, NoClear, Meta = (ExcludeBaseStruct, BaseStruct = "/Script/Chooser.ContextObjectTypeBase"), Category = "Input")
-	TArray<FInstancedStruct> ContextData;
 
-	UPROPERTY(EditAnywhere, Category="Output", Meta = (AllowAbstract=true))
+	UPROPERTY(EditAnywhere, DisplayName= "Result Class", Category="Result", Meta = (AllowAbstract=true))
 	TObjectPtr<UClass> OutputObjectType;
-	UPROPERTY(EditAnywhere, Category="Output")
+	UPROPERTY(EditAnywhere, DisplayName = "Result Type", Category="Result")
 	EObjectChooserResultType ResultType = EObjectChooserResultType::ObjectResult;
 
+	UPROPERTY(EditAnywhere, NoClear, DisplayName = "Parameters", Meta = (ExcludeBaseStruct, BaseStruct = "/Script/Chooser.ContextObjectTypeBase"), Category = "Parameters")
+	TArray<FInstancedStruct> ContextData;
+
 	virtual TConstArrayView<FInstancedStruct> GetContextData() const override { return GetRootChooser()->ContextData; }
 
 	bool IsCookedData() const { return !CookedResults.IsEmpty(); }
diff --git a/Engine/Plugins/Experimental/Chooser/Source/Chooser/Public/ChooserPropertyAccess.h b/Engine/Plugins/Experimental/Chooser/Source/Chooser/Public/ChooserPropertyAccess.h
--- a/Engine/Plugins/Experimental/Chooser/Source/Chooser/Public/ChooserPropertyAccess.h
+++ b/Engine/Plugins/Experimental/Chooser/Source/Chooser/Public/ChooserPropertyAccess.h
@@ -152,9 +152,9 @@ struct FChooserStructPropertyBinding : public FChooserPropertyBinding
 UENUM()
 enum class EContextObjectDirection
 {
-	Read,
-	Write,
-	ReadWrite
+	Read UMETA(DisplayName="Input", Tooltip = "This Parameter will only be read from"),
+	Write UMETA(DisplayName="Output", Tooltip = "This Parameter will only be written to"),
+	ReadWrite UMETA(DisplayName="Input/Output", Tooltip = "This Parameter can be both read from and written to"),
 };
 
 USTRUCT()
@@ -162,12 +162,12 @@ struct FContextObjectTypeBase
 {
 	GENERATED_BODY()
 	
-	UPROPERTY(EditAnywhere, Category="Type")
+	UPROPERTY(EditAnywhere, Category="Type", meta = (ToolTip="Select weather this property will be read from, written to, or both"))
 	EContextObjectDirection Direction = EContextObjectDirection::Read;
 };
 
 
-USTRUCT()
+USTRUCT(DisplayName="Class Parameter")
 struct FContextObjectTypeClass : public FContextObjectTypeBase
 {
 	GENERATED_BODY()
@@ -175,7 +175,7 @@ struct FContextObjectTypeClass : public FContextObjectTypeBase
 	TObjectPtr<UClass> Class;
 };
 
-USTRUCT()
+USTRUCT(DisplayName = "Struct Parameter")
 struct FContextObjectTypeStruct : public FContextObjectTypeBase
 {
 	GENERATED_BODY()
diff --git a/Engine/Plugins/Experimental/Chooser/Source/Chooser/Public/IHasContext.h b/Engine/Plugins/Experimental/Chooser/Source/Chooser/Public/IHasContext.h
--- a/Engine/Plugins/Experimental/Chooser/Source/Chooser/Public/IHasContext.h
+++ b/Engine/Plugins/Experimental/Chooser/Source/Chooser/Public/IHasContext.h
@@ -37,6 +37,8 @@ DECLARE_MULTICAST_DELEGATE_OneParam(FChooserOutputObjectTypeChanged, const UClas
 UENUM()
 enum class EObjectChooserResultType
 {
-	ObjectResult,
-	ClassResult,
+	// The Chooser returns an Object of the specified Result Class
+	ObjectResult UMETA(DisplayName = "Object Of Type"),
+	// The Chooser returns a Class that is a SubClass of the specified Result Class (eg a chooser could return a type of Character to spawn, or a type of AnimInstance to link)
+	ClassResult UMETA(DisplayName = "SubClass Of"),
 };
diff --git a/Engine/Plugins/Experimental/Chooser/Source/ChooserEditor/Private/ChooserFactory.cpp b/Engine/Plugins/Experimental/Chooser/Source/ChooserEditor/Private/ChooserFactory.cpp
--- a/Engine/Plugins/Experimental/Chooser/Source/ChooserEditor/Private/ChooserFactory.cpp
+++ b/Engine/Plugins/Experimental/Chooser/Source/ChooserEditor/Private/ChooserFactory.cpp
@@ -2,6 +2,156 @@
 
 #include "ChooserFactory.h"
 #include "Chooser.h"
+#include "ChooserInitializer.h"
+#include "Widgets/Input/SButton.h"
+#include "Widgets/Layout/SUniformGridPanel.h"
+#include "Widgets/SCompoundWidget.h"
+#include "Widgets/SWindow.h"
+#include "Editor.h"
+#include "ObjectChooserClassFilter.h"
+#include "StructViewerModule.h"
+#include "PropertyEditorModule.h"
+#include "Modules/ModuleManager.h"
+
+#define LOCTEXT_NAMESPACE "CreateChooser"
+
+/*------------------------------------------------------------------------------
+	Dialog to configure creation properties
+------------------------------------------------------------------------------*/
+class SChooserCreateDialog : public SCompoundWidget
+{
+public:
+	SLATE_BEGIN_ARGS( SChooserCreateDialog ){}
+
+	SLATE_END_ARGS()
+
+	/** Constructs this widget with InArgs */
+	void Construct( const FArguments& InArgs )
+	{
+		bOkClicked = false;
+		
+		FPropertyEditorModule& PropertyEditorModule = FModuleManager::GetModuleChecked<FPropertyEditorModule>( "PropertyEditor" );
+		FDetailsViewArgs DetailsViewArgs;
+		// DetailsViewArgs.NotifyHook = this;
+		DetailsViewArgs.NameAreaSettings = FDetailsViewArgs::HideNameArea;
+		DetailsViewArgs.bAllowSearch = false;
+		DetailsViewArgs.bShowOptions = false;
+		DetailsViewArgs.bAllowFavoriteSystem = false;
+		DetailsViewArgs.bShowLooseProperties = true;
+		DetailsView = PropertyEditorModule.CreateDetailView( DetailsViewArgs );
+		DetailsView->SetObject(ChooserFactory.Get());
+
+		ChildSlot
+		[
+			SNew(SVerticalBox)
+			+SVerticalBox::Slot().FillHeight(1)
+			[
+				DetailsView.ToSharedRef()
+			]
+
+			// Ok/Cancel buttons
+			+SVerticalBox::Slot()
+			.AutoHeight()
+			.HAlign(HAlign_Right)
+			.VAlign(VAlign_Bottom)
+			.Padding(10.0f)
+			[
+				SNew(SUniformGridPanel)
+				.SlotPadding(FAppStyle::GetMargin("StandardDialog.SlotPadding"))
+				.MinDesiredSlotWidth(FAppStyle::GetFloat("StandardDialog.MinDesiredSlotWidth"))
+				.MinDesiredSlotHeight(FAppStyle::GetFloat("StandardDialog.MinDesiredSlotHeight"))
+				+SUniformGridPanel::Slot(0,0)
+				[
+					SNew(SButton)
+					.ToolTipText(LOCTEXT("CreateAnimBlueprintCreate_Tooltip", "Create a new Chooser Table Asset.."))
+					.IsEnabled_Lambda([this]()
+					{
+						return true;
+					})
+					.HAlign(HAlign_Center)
+					.ContentPadding( FAppStyle::GetMargin("StandardDialog.ContentPadding") )
+					.OnClicked(this, &SChooserCreateDialog::OkClicked)
+					.Text(LOCTEXT("Create", "Create"))
+				]
+				+SUniformGridPanel::Slot(1,0)
+				[
+					SNew(SButton)
+					.HAlign(HAlign_Center)
+					.ContentPadding( FAppStyle::GetMargin("StandardDialog.ContentPadding") )
+					.OnClicked(this, &SChooserCreateDialog::CancelClicked)
+					.Text(LOCTEXT("Cancel", "Cancel"))
+				]
+			]
+		];
+	}
+	
+	/** Sets properties for the supplied AnimBlueprintFactory */
+	bool ConfigureProperties(TWeakObjectPtr<UChooserTableFactory> InChooserFactory)
+	{
+		ChooserFactory = InChooserFactory;
+		DetailsView->SetObject(ChooserFactory.Get());
+
+		Window = SNew(SWindow)
+		.Title( LOCTEXT("Create Chooser Options", "Create Chooser Table") )
+		.SizingRule(ESizingRule::FixedSize)
+		.ClientSize(FVector2f(550, 300))
+		.SupportsMinimize(false)
+		.SupportsMaximize(false)
+		[
+			AsShared()
+		];
+
+		GEditor->EditorAddModalWindow(Window.ToSharedRef());
+		ChooserFactory.Reset();
+
+		return bOkClicked;
+	}
+
+private:
+
+	/** Handler for when ok is clicked */
+	FReply OkClicked()
+	{
+		CloseDialog(true);
+
+		return FReply::Handled();
+	}
+
+	void CloseDialog(bool bWasPicked=false)
+	{
+		bOkClicked = bWasPicked;
+		if (Window.IsValid())
+		{
+			Window->RequestDestroyWindow();
+		}
+	}
+
+	/** Handler for when cancel is clicked */
+	FReply CancelClicked()
+	{
+		CloseDialog();
+		return FReply::Handled();
+	}
+
+	FReply OnKeyDown( const FGeometry& MyGeometry, const FKeyEvent& InKeyEvent )
+	{
+		if (InKeyEvent.GetKey() == EKeys::Escape)
+		{
+			CloseDialog();
+			return FReply::Handled();
+		}
+		return SWidget::OnKeyDown(MyGeometry, InKeyEvent);
+	}
+
+private:
+	/** The factory for which we are setting up properties */
+	TWeakObjectPtr<UChooserTableFactory> ChooserFactory;
+	TSharedPtr<SWindow> Window;
+	TSharedPtr<IDetailsView> DetailsView;
+
+	/** True if Ok was clicked */
+	bool bOkClicked = false;
+};
 
 UChooserTableFactory::UChooserTableFactory()
 {
@@ -12,10 +162,18 @@ UChooserTableFactory::UChooserTableFactory()
 
 bool UChooserTableFactory::ConfigureProperties()
 {
-	return true;
+	TSharedRef<SChooserCreateDialog> Dialog = SNew(SChooserCreateDialog);
+	return Dialog->ConfigureProperties(this);
 }
 
 UObject* UChooserTableFactory::FactoryCreateNew(UClass* Class, UObject* InParent, FName Name, EObjectFlags Flags, UObject* Context, FFeedbackContext* Warn, FName CallingContext)
 {
-	return NewObject<UChooserTable>(InParent, Class, Name, Flags | RF_Public | RF_Standalone | RF_Transactional | RF_LoadCompleted);
-}
\ No newline at end of file
+	UChooserTable* NewChooser =  NewObject<UChooserTable>(InParent, Class, Name, Flags | RF_Public | RF_Standalone | RF_Transactional | RF_LoadCompleted);
+	if (ChooserInitializer.IsValid())
+	{
+		ChooserInitializer.Get<FChooserInitializer>().Initialize(NewChooser);
+	}
+	return NewChooser;
+}
+
+#undef LOCTEXT_NAMESPACE
\ No newline at end of file
diff --git a/Engine/Plugins/Experimental/Chooser/Source/ChooserEditor/Private/ChooserFactory.h b/Engine/Plugins/Experimental/Chooser/Source/ChooserEditor/Private/ChooserFactory.h
--- a/Engine/Plugins/Experimental/Chooser/Source/ChooserEditor/Private/ChooserFactory.h
+++ b/Engine/Plugins/Experimental/Chooser/Source/ChooserEditor/Private/ChooserFactory.h
@@ -4,6 +4,8 @@
 
 #include "CoreMinimal.h"
 #include "Factories/Factory.h"
+#include "StructUtils/InstancedStruct.h"
+
 #include "ChooserFactory.generated.h"
 
 UCLASS(MinimalAPI)
@@ -20,4 +22,8 @@ class UChooserTableFactory : public UFactory
 		return FactoryCreateNew(Class, InParent, Name, Flags, Context, Warn, NAME_None);
 	}
 	virtual FName GetNewAssetThumbnailOverride() const override { return "ChooserEditor.ChooserTableIconLarge"; }
+
+private:
+	UPROPERTY(EditAnywhere, DisplayName = "Chooser Type", NoClear, Category="Chooser Setup", meta = (ExcludeBaseStruct, BaseStruct = "/Script/ChooserEditor.ChooserInitializer"))
+	FInstancedStruct ChooserInitializer;
 };
\ No newline at end of file
diff --git a/Engine/Plugins/Experimental/Chooser/Source/ChooserEditor/Private/ChooserInitializer.cpp b/Engine/Plugins/Experimental/Chooser/Source/ChooserEditor/Private/ChooserInitializer.cpp
new file mode 100644
--- /dev/null
+++ b/Engine/Plugins/Experimental/Chooser/Source/ChooserEditor/Private/ChooserInitializer.cpp
@@ -0,0 +1,46 @@
+// Copyright Epic Games, Inc. All Rights Reserved.
+
+#include "ChooserInitializer.h"
+
+#include "AnimNode_ChooserPlayer.h"
+#include "Chooser.h"
+#include "ChooserPropertyAccess.h"
+#include "ObjectChooserWidgetFactories.h"
+#include "Animation/AnimationAsset.h"
+#include "Animation/AnimInstance.h"
+
+#define LOCTEXT_NAMESPACE "ChooserInitializer"
+
+
+void FGenericChooserInitializer::Initialize(UChooserTable* Chooser) const
+{
+	Chooser->ContextData = ContextData;
+	Chooser->ResultType = ResultType;
+	Chooser->OutputObjectType = OutputObjectType;
+}
+
+void FChooserPlayerInitializer::Initialize(UChooserTable* Chooser) const
+{
+	Chooser->ContextData.SetNum(2);
+	Chooser->ContextData[0].InitializeAs(FContextObjectTypeClass::StaticStruct());
+	FContextObjectTypeClass& ClassData = Chooser->ContextData[0].GetMutable<FContextObjectTypeClass>();
+	if (AnimClass)
+	{
+		ClassData.Class = AnimClass;		
+	}
+	else
+	{
+		ClassData.Class = UAnimInstance::StaticClass();
+	}
+	ClassData.Direction = EContextObjectDirection::ReadWrite;
+	
+	Chooser->ContextData[1].InitializeAs(FContextObjectTypeStruct::StaticStruct());
+	FContextObjectTypeStruct& StructData = Chooser->ContextData[1].GetMutable<FContextObjectTypeStruct>();
+	StructData.Struct = FChooserPlayerSettings::StaticStruct();
+	StructData.Direction = EContextObjectDirection::Write;
+	
+	Chooser->OutputObjectType = UAnimationAsset::StaticClass();
+}
+
+
+#undef LOCTEXT_NAMESPACE
diff --git a/Engine/Plugins/Experimental/Chooser/Source/ChooserEditor/Public/ChooserInitializer.h b/Engine/Plugins/Experimental/Chooser/Source/ChooserEditor/Public/ChooserInitializer.h
new file mode 100644
--- /dev/null
+++ b/Engine/Plugins/Experimental/Chooser/Source/ChooserEditor/Public/ChooserInitializer.h
@@ -0,0 +1,44 @@
+// Copyright Epic Games, Inc. All Rights Reserved.
+#pragma once
+
+#include "InstancedStruct.h"
+#include "Chooser.h"
+#include "ChooserPropertyAccess.h"
+
+#include "ChooserInitializer.generated.h"
+
+USTRUCT()
+struct FChooserInitializer
+{
+	GENERATED_BODY()
+	virtual ~FChooserInitializer() {}
+	virtual void Initialize(UChooserTable* Chooser) const {};
+};
+
+
+USTRUCT(DisplayName="Generic Chooser", Meta = (ToolTip="A ChooserTable for use in Blueprint, which can return an arbitrary Asset type or Class, and can take any number of Objects or Structs as Parameters.") )
+struct FGenericChooserInitializer : public FChooserInitializer
+{
+	GENERATED_BODY()
+
+	virtual void Initialize(UChooserTable* Chooser) const override;
+	
+	UPROPERTY(EditAnywhere, DisplayName= "Result Class", Category="Result", Meta = (AllowAbstract=true))
+	TObjectPtr<UClass> OutputObjectType;
+	
+	UPROPERTY(EditAnywhere, DisplayName = "Result Type", Category="Result")
+	EObjectChooserResultType ResultType = EObjectChooserResultType::ObjectResult;
+	
+	UPROPERTY(EditAnywhere, DisplayName = "Parameters", NoClear, Meta = (ExpandByDefault, ExcludeBaseStruct, BaseStruct = "/Script/Chooser.ContextObjectTypeBase"), Category = "Parameters")
+	TArray<FInstancedStruct> ContextData;
+};
+
+USTRUCT(DisplayName="Animation Chooser", meta = (ToolTip="A ChooserTable for use with the ChooserPlayer AnimGraph node.\nReturns an AnimAsset, and takes an AnimInstance, and a ChooserPlayerSettings struct as parameters."))
+struct FChooserPlayerInitializer : public FChooserInitializer
+{
+	GENERATED_BODY()
+	virtual void Initialize(UChooserTable* Chooser) const override;
+
+	UPROPERTY(EditAnywhere, Category="AnimClass")
+	TObjectPtr<UClass> AnimClass;
+};
\ No newline at end of file
