Fix Struct Output column not reinitializing cell data when a new struct is selected

diff --git a/Engine/Plugins/Experimental/Chooser/Source/ChooserEditor/Private/ChooserColumnHeader.cpp b/Engine/Plugins/Experimental/Chooser/Source/ChooserEditor/Private/ChooserColumnHeader.cpp
--- a/Engine/Plugins/Experimental/Chooser/Source/ChooserEditor/Private/ChooserColumnHeader.cpp
+++ b/Engine/Plugins/Experimental/Chooser/Source/ChooserEditor/Private/ChooserColumnHeader.cpp
@@ -15,12 +15,12 @@
 namespace UE::ChooserEditor
 {
 
-	TSharedRef<SWidget> MakeColumnHeaderWidget(UChooserTable* Chooser, FChooserColumnBase* Column, const FText& ColumnName,const FText& ColumnTooltip, const FSlateBrush* ColumnIcon, TSharedPtr<SWidget> DebugWidget)
+	TSharedRef<SWidget> MakeColumnHeaderWidget(UChooserTable* Chooser, FChooserColumnBase* Column, const FText& ColumnName,const FText& ColumnTooltip, const FSlateBrush* ColumnIcon, TSharedPtr<SWidget> DebugWidget, FChooserWidgetValueChanged ValueChanged)
 	{
 		TSharedPtr<SWidget> InputValueWidget = nullptr;
 		if (FChooserParameterBase* InputValue = Column->GetInputValue())
 		{
-			InputValueWidget = FObjectChooserWidgetFactories::CreateWidget(false, Chooser, InputValue, Column->GetInputType(), Chooser->OutputObjectType);
+			InputValueWidget = FObjectChooserWidgetFactories::CreateWidget(false, Chooser, InputValue, Column->GetInputType(), Chooser->OutputObjectType,ValueChanged);
 		}
 	
 		TSharedRef<SVerticalBox> ColumnHeaderWidget =
diff --git a/Engine/Plugins/Experimental/Chooser/Source/ChooserEditor/Private/OutputStructColumnEditor.cpp b/Engine/Plugins/Experimental/Chooser/Source/ChooserEditor/Private/OutputStructColumnEditor.cpp
--- a/Engine/Plugins/Experimental/Chooser/Source/ChooserEditor/Private/OutputStructColumnEditor.cpp
+++ b/Engine/Plugins/Experimental/Chooser/Source/ChooserEditor/Private/OutputStructColumnEditor.cpp
@@ -30,7 +30,13 @@ TSharedRef<SWidget> CreateOutputStructColumnWidget(UChooserTable* Chooser, FChoo
         		
 		TSharedPtr<SWidget> DebugWidget = nullptr;
         
-		return MakeColumnHeaderWidget(Chooser, Column, ColumnName, ColumnTooltip, ColumnIcon, DebugWidget);
+		return MakeColumnHeaderWidget(Chooser, Column, ColumnName, ColumnTooltip, ColumnIcon, DebugWidget,
+				FChooserWidgetValueChanged::CreateLambda([Column]()
+				{
+					FOutputStructColumn* StructColumn = static_cast<FOutputStructColumn*>(Column);
+					StructColumn->StructTypeChanged();
+				})	
+			);
 	}
 	
 	FOutputStructColumn* StructColumn = static_cast<FOutputStructColumn*>(Column);
diff --git a/Engine/Plugins/Experimental/Chooser/Source/ChooserEditor/Private/SPropertyAccessChainWidget.cpp b/Engine/Plugins/Experimental/Chooser/Source/ChooserEditor/Private/SPropertyAccessChainWidget.cpp
--- a/Engine/Plugins/Experimental/Chooser/Source/ChooserEditor/Private/SPropertyAccessChainWidget.cpp
+++ b/Engine/Plugins/Experimental/Chooser/Source/ChooserEditor/Private/SPropertyAccessChainWidget.cpp
@@ -169,16 +169,45 @@ TSharedRef<SWidget> SPropertyAccessChainWidget::CreatePropertyAccessWidget()
 					Chooser::CopyPropertyChain(InBindingChain, *PropertyBindingValue.Get());
 					
 					FField* Property = InBindingChain.Last().Field.ToField();
-					
+						
 					ContextProperty->DisplayName = "";
-					if (Property)
+					if (InBindingChain.Num() == 1)
 					{
-						ContextProperty->DisplayName = Property->GetDisplayNameText().ToString();
-						static const int ShortNameLength = 5;
-						if (ContextProperty->DisplayName.Len() < ShortNameLength && InBindingChain.Num() > 2)
+						// direct binding to a context struct/class, set display name to the struct/clase name
+						if (ContextClassOwner)
 						{
-							FField* ParentProperty = InBindingChain[InBindingChain.Num() - 2].Field.ToField();
-							ContextProperty->DisplayName = ParentProperty->GetDisplayNameText().ToString() + "." + ContextProperty->DisplayName;
+							TConstArrayView<FInstancedStruct> ContextData = ContextClassOwner->GetContextData();
+							if (ContextData.IsValidIndex(InBindingChain[0].ArrayIndex))
+							{
+								if (const FContextObjectTypeStruct* StructContext = ContextData[InBindingChain[0].ArrayIndex].GetPtr<FContextObjectTypeStruct>())
+								{
+									if (StructContext->Struct)
+									{
+										ContextProperty->DisplayName = StructContext->Struct->GetDisplayNameText().ToString();
+									}
+								}
+								else if (const FContextObjectTypeClass* ClassContext = ContextData[InBindingChain[0].ArrayIndex].GetPtr<FContextObjectTypeClass>())
+								{
+									if (ClassContext->Class)
+									{
+										ContextProperty->DisplayName = ClassContext->Class->GetDisplayNameText().ToString();
+									}
+								}
+							}
+						}
+					}
+					else
+					{
+						// set displayname from property name
+						if (Property)
+						{
+							ContextProperty->DisplayName = Property->GetDisplayNameText().ToString();
+							static const int ShortNameLength = 5;
+							if (ContextProperty->DisplayName.Len() < ShortNameLength && InBindingChain.Num() > 2)
+							{
+								FField* ParentProperty = InBindingChain[InBindingChain.Num() - 2].Field.ToField();
+								ContextProperty->DisplayName = ParentProperty->GetDisplayNameText().ToString() + "." + ContextProperty->DisplayName;
+							}
 						}
 					}
 					
diff --git a/Engine/Plugins/Experimental/Chooser/Source/ChooserEditor/Public/ChooserColumnHeader.h b/Engine/Plugins/Experimental/Chooser/Source/ChooserEditor/Public/ChooserColumnHeader.h
--- a/Engine/Plugins/Experimental/Chooser/Source/ChooserEditor/Public/ChooserColumnHeader.h
+++ b/Engine/Plugins/Experimental/Chooser/Source/ChooserEditor/Public/ChooserColumnHeader.h
@@ -1,6 +1,7 @@
 // Copyright Epic Games, Inc. All Rights Reserved.
 #pragma once
 
+#include "ObjectChooserWidgetFactories.h"
 #include "Widgets/SWidget.h"
 
 class UChooserTable;
@@ -9,5 +10,5 @@ struct FChooserColumnBase;
 
 namespace UE::ChooserEditor
 {
-	TSharedRef<SWidget> MakeColumnHeaderWidget(UChooserTable* Chooser, FChooserColumnBase* Column, const FText& ColumnName,const FText& ColumnTooltip, const FSlateBrush* ColumnIcon, TSharedPtr<SWidget> DebugWidget);
+	TSharedRef<SWidget> MakeColumnHeaderWidget(UChooserTable* Chooser, FChooserColumnBase* Column, const FText& ColumnName,const FText& ColumnTooltip, const FSlateBrush* ColumnIcon, TSharedPtr<SWidget> DebugWidget, FChooserWidgetValueChanged ValueChanged = FChooserWidgetValueChanged());
 }
\ No newline at end of file
