Improved compiler errors for chooser property access

diff --git a/Engine/Plugins/Experimental/Chooser/Source/Chooser/Internal/Chooser.h b/Engine/Plugins/Experimental/Chooser/Source/Chooser/Internal/Chooser.h
--- a/Engine/Plugins/Experimental/Chooser/Source/Chooser/Internal/Chooser.h
+++ b/Engine/Plugins/Experimental/Chooser/Source/Chooser/Internal/Chooser.h
@@ -185,6 +185,7 @@ public:
 	TArray<FInstancedStruct> ContextData;
 
 	virtual TConstArrayView<FInstancedStruct> GetContextData() const override { return GetRootChooser()->ContextData; }
+	virtual FString GetContextOwnerName() const override { return GetName(); }
 
 	bool IsCookedData() const { return !CookedResults.IsEmpty(); }
 	static FObjectChooserBase::EIteratorStatus EvaluateChooser(FChooserEvaluationContext& Context, const UChooserTable* Chooser, FObjectChooserBase::FObjectChooserIteratorCallback Callback);
diff --git a/Engine/Plugins/Experimental/Chooser/Source/Chooser/Internal/RandomizeColumn.h b/Engine/Plugins/Experimental/Chooser/Source/Chooser/Internal/RandomizeColumn.h
--- a/Engine/Plugins/Experimental/Chooser/Source/Chooser/Internal/RandomizeColumn.h
+++ b/Engine/Plugins/Experimental/Chooser/Source/Chooser/Internal/RandomizeColumn.h
@@ -18,6 +18,10 @@ public:
 	FChooserPropertyBinding Binding;
 	
 	virtual bool GetValue(FChooserEvaluationContext& Context, const FChooserRandomizationContext*& OutResult) const override;
+	virtual bool IsBound() const override
+	{
+		return Binding.IsBoundToRoot || !Binding.PropertyBindingChain.IsEmpty();
+	}
 
 	CHOOSER_PARAMETER_BOILERPLATE();
 };
@@ -71,6 +75,18 @@ struct CHOOSER_API FRandomizeColumn : public FChooserColumnBase
 		}
 	}
 
+	virtual void Compile(IHasContextClass* Owner, bool bForce) override
+	{
+		if (FChooserParameterRandomizeBase* Input = InputValue.GetMutablePtr<FChooserParameterRandomizeBase>())
+		{
+			// binding on randomize columns is optional, so don't call compile unless it's bound, to avoid error messages
+			if (Input->IsBound())
+			{
+				Input->Compile(Owner, bForce);
+			}
+		}
+	};
+
 	CHOOSER_COLUMN_BOILERPLATE(FChooserParameterRandomizeBase);
 
 #if WITH_EDITOR
diff --git a/Engine/Plugins/Experimental/Chooser/Source/Chooser/Private/ChooserPropertyAccess.cpp b/Engine/Plugins/Experimental/Chooser/Source/Chooser/Private/ChooserPropertyAccess.cpp
--- a/Engine/Plugins/Experimental/Chooser/Source/Chooser/Private/ChooserPropertyAccess.cpp
+++ b/Engine/Plugins/Experimental/Chooser/Source/Chooser/Private/ChooserPropertyAccess.cpp
@@ -17,8 +17,8 @@
 
 TAutoConsoleVariable<bool> CVarEnableDetailedWarnings(
 	TEXT("Choosers.EnableDetailedWarnings"),
-	false,
-	TEXT("Enable detailed context validation with warnings when choosers are evaluated on an incorrect context. \n0: Disable (default), 1: Enable"),
+	true,
+	TEXT("Enable detailed context validation with warnings when choosers are evaluated on an incorrect context. \n0: Disable, 1: Enable (default)"),
 	ECVF_Default);
 
 TAutoConsoleVariable<bool> CVarUseCompiledPropertyChainsInEditor(
@@ -87,6 +87,7 @@ void FChooserPropertyBinding::Compile(IHasContextClass* Owner, bool bForce)
 #if WITH_EDITORONLY_DATA
 		CompileMessage = LOCTEXT("No Property Bound", "No Property Bound");
 #endif
+		UE_LOG(LogChooser, Error, TEXT("Chooser Property Compile Error (%s): Missing property binding."), *Owner->GetContextOwnerName());
 		CompiledBinding = nullptr;
 		return;
 	}
@@ -106,6 +107,7 @@ void FChooserPropertyBinding::Compile(IHasContextClass* Owner, bool bForce)
 #if WITH_EDITORONLY_DATA
 		CompileMessage = FText::Format(LOCTEXT("No valid struct", "No valid Context Object/Struct at index: {0}"), FText::FromString(FString::FromInt(ContextIndex)));
 #endif
+		UE_LOG(LogChooser, Error, TEXT("Chooser Property Compile Error (%s): No valid Context Object/Struct at index: %d"), *Owner->GetContextOwnerName(), ContextIndex);
 		CompiledBinding = nullptr;
 		return;
 	}
@@ -203,6 +205,7 @@ void FChooserPropertyBinding::Compile(IHasContextClass* Owner, bool bForce)
 #if WITH_EDITORONLY_DATA
 			CompileMessage = FText::Format(LOCTEXT("Property Not Found", "Property/Function: {0} not Found on Class/Struct: {1}"), FText::FromName(PropertyBindingChain[PropertyChainIndex]), StructType->GetDisplayNameText());
 #endif
+			UE_LOG(LogChooser, Error, TEXT(" Chooser Property Compile Error (%s): Property/Function: %s not Found on Class/Struct %s"), *Owner->GetContextOwnerName(), *PropertyBindingChain[PropertyChainIndex].ToString(), *StructType->GetName());
 			CompiledBinding = nullptr;
 			return;
 		}
@@ -303,6 +306,7 @@ void FChooserPropertyBinding::Compile(IHasContextClass* Owner, bool bForce)
 #if WITH_EDITORONLY_DATA
 		CompileMessage = FText::Format(LOCTEXT("Property Not Found", "Property/Function: {0} not Found on Class/Struct: {1}"), FText::FromName(PropertyBindingChain.Last()), StructType->GetDisplayNameText());
 #endif
+		UE_LOG(LogChooser, Error, TEXT(" Chooser Property Compile Error (%s): Property/Function: %s not Found on Class/Struct %s"), *Owner->GetContextOwnerName(), *PropertyBindingChain.Last().ToString(), *StructType->GetName());
 		CompiledBinding = nullptr;
  	}
 }
@@ -318,6 +322,7 @@ namespace UE::Chooser
 			return;
 		}
 
+		TRACE_CPUPROFILER_EVENT_SCOPE(ValidateChooserContext);
 		int ContextNum = ContextData.Num();
 		for (int i = 0; i < ContextNum; i++)
 		{
@@ -334,26 +339,26 @@ namespace UE::Chooser
 						{
 							if (!InputObjectParam->Object->GetClass()->IsChildOf(ExpectedClassType->Class))
 							{
-								UE_LOG(LogChooser, Warning, TEXT("Chooser Table: %s ContextData entry %d expects an object of type %s, but an object of type %s was passed in."),
+								UE_LOG(LogChooser, Error, TEXT("Chooser Table: %s ContextData entry %d expects an object of type %s, but an object of type %s was passed in."),
 									ToCStr(Chooser->GetName()), i, ToCStr(ExpectedClassType->Class->GetName()), ToCStr(InputObjectParam->Object->GetClass()->GetName()));
 							}
 						}
 						else
 						{
-							UE_LOG(LogChooser, Warning, TEXT("Chooser Table: %s ContextData entry %d expects an object of type %s, but null was passed in."),
+							UE_LOG(LogChooser, Error, TEXT("Chooser Table: %s ContextData entry %d expects an object of type %s, but null was passed in."),
 								ToCStr(Chooser->GetName()), i, ToCStr(ExpectedClassType->Class->GetName()));
 						}
 
 					}
 					else
 					{
-						UE_LOG(LogChooser, Warning, TEXT("Chooser Table: %s ContextData entry %d expects an object of type %s, but was passed a struct of type %s."),
+						UE_LOG(LogChooser, Error, TEXT("Chooser Table: %s ContextData entry %d expects an object of type %s, but was passed a struct of type %s."),
 							ToCStr(Chooser->GetName()), i, ToCStr(ExpectedClassType->Class->GetName()), ToCStr(Context.Params[i].GetScriptStruct()->GetName()));
 					}
 				}
 				else
 				{
-					UE_LOG(LogChooser, Warning, TEXT("Chooser Table: %s ContextData entry %d expects an object of type %s, but nothing was passed in."),
+					UE_LOG(LogChooser, Error, TEXT("Chooser Table: %s ContextData entry %d expects an object of type %s, but nothing was passed in."),
 						ToCStr(Chooser->GetName()), i, ToCStr(ExpectedClassType->Class->GetName()));
 				}
 			}
@@ -363,27 +368,27 @@ namespace UE::Chooser
 				{
 					if (FChooserEvaluationInputObject* InputObjectParam = Context.Params[i].GetPtr<FChooserEvaluationInputObject>())
 					{
-						UE_LOG(LogChooser, Warning, TEXT("Chooser Table: %s ContextData entry %d expects an struct of type %s, but was passed a object of type %s."),
+						UE_LOG(LogChooser, Error, TEXT("Chooser Table: %s ContextData entry %d expects an struct of type %s, but was passed a object of type %s."),
 							ToCStr(Chooser->GetName()), i, ToCStr(ExpectedStructType->Struct->GetName()), ToCStr(InputObjectParam->Object->GetClass()->GetName()));
 					}
 					else
 					{
 						if (Context.Params[i].GetScriptStruct() != ExpectedStructType->Struct)
 						{
-							UE_LOG(LogChooser, Warning, TEXT("Chooser Table: %s ContextData entry %d expects an struct of type %s, but was passed a struct of type %s."),
+							UE_LOG(LogChooser, Error, TEXT("Chooser Table: %s ContextData entry %d expects an struct of type %s, but was passed a struct of type %s."),
 								ToCStr(Chooser->GetName()), i, ToCStr(ExpectedStructType->Struct->GetName()), ToCStr(Context.Params[i].GetScriptStruct()->GetName()));
 						}
 					}
 				}
 				else
 				{
-					UE_LOG(LogChooser, Warning, TEXT("Chooser Table: %s ContextData entry %d expects an struct of type %s, but nothing was passed in."),
+					UE_LOG(LogChooser, Error, TEXT("Chooser Table: %s ContextData entry %d expects an struct of type %s, but nothing was passed in."),
 						ToCStr(Chooser->GetName()), i, ToCStr(ExpectedStructType->Struct->GetName()));
 				}
 			}
 			else
 			{
-				UE_LOG(LogChooser, Warning, TEXT("Chooser Table: %s  ContextData entry %i is of unknown type or none."), ToCStr(Chooser->GetName()), i);
+				UE_LOG(LogChooser, Error, TEXT("Chooser Table: %s  ContextData entry %i is of unknown type or none."), ToCStr(Chooser->GetName()), i);
 			}
 		}
 	}
@@ -403,7 +408,7 @@ namespace UE::Chooser
 
 		if(!Context.Params.IsValidIndex(CompiledBinding.ContextIndex))
 		{
-			UE_LOG(LogChooser, Warning, TEXT("Invalid Index {%d} while resolving compiled property chain."), CompiledBinding.ContextIndex);
+			UE_LOG(LogChooser, Error, TEXT("Invalid Index {%d} while resolving compiled property chain."), CompiledBinding.ContextIndex);
 			return false;
 		}
 		
@@ -429,7 +434,7 @@ namespace UE::Chooser
 
 		if (!InputType->IsChildOf(CompiledBinding.TargetType))
 		{
-			UE_LOG(LogChooser, Warning, TEXT("Property Binding compiled for type: {%s} is being evaluated on incompatible type: {%s}."), ToCStr(CompiledBinding.TargetType->GetName()), ToCStr(InputType->GetName()));
+			UE_LOG(LogChooser, Error, TEXT("Property Binding compiled for type: {%s} is being evaluated on incompatible type: {%s}."), ToCStr(CompiledBinding.TargetType->GetName()), ToCStr(InputType->GetName()));
 			return false;
 		}
 		
diff --git a/Engine/Plugins/Experimental/Chooser/Source/Chooser/Public/IChooserColumn.h b/Engine/Plugins/Experimental/Chooser/Source/Chooser/Public/IChooserColumn.h
--- a/Engine/Plugins/Experimental/Chooser/Source/Chooser/Public/IChooserColumn.h
+++ b/Engine/Plugins/Experimental/Chooser/Source/Chooser/Public/IChooserColumn.h
@@ -48,6 +48,12 @@ public:
 
 	virtual void Compile(IHasContextClass* Owner, bool bForce)
 	{
+#if WITH_EDITORONLY_DATA
+		if (bDisabled)
+		{
+			return;
+		}
+#endif
 		if (FChooserParameterBase* Input = GetInputValue())
 		{
 			Input->Compile(Owner, bForce);
diff --git a/Engine/Plugins/Experimental/Chooser/Source/Chooser/Public/IChooserParameterRandomize.h b/Engine/Plugins/Experimental/Chooser/Source/Chooser/Public/IChooserParameterRandomize.h
--- a/Engine/Plugins/Experimental/Chooser/Source/Chooser/Public/IChooserParameterRandomize.h
+++ b/Engine/Plugins/Experimental/Chooser/Source/Chooser/Public/IChooserParameterRandomize.h
@@ -27,4 +27,5 @@ struct FChooserParameterRandomizeBase : public FChooserParameterBase
 {
 	GENERATED_BODY()
 	virtual bool GetValue(FChooserEvaluationContext& Context, const FChooserRandomizationContext*& OutResult) const { return false; }
+	virtual bool IsBound() const { return false; }
 };
diff --git a/Engine/Plugins/Experimental/Chooser/Source/Chooser/Public/IHasContext.h b/Engine/Plugins/Experimental/Chooser/Source/Chooser/Public/IHasContext.h
--- a/Engine/Plugins/Experimental/Chooser/Source/Chooser/Public/IHasContext.h
+++ b/Engine/Plugins/Experimental/Chooser/Source/Chooser/Public/IHasContext.h
@@ -25,6 +25,8 @@ public:
 	FContextClassChanged OnContextClassChanged;
 	virtual TConstArrayView<FInstancedStruct> GetContextData() const { return TConstArrayView<FInstancedStruct>(); }
 
+	virtual FString GetContextOwnerName() const { return ""; }
+
 	virtual void Compile(bool bForce = false) {}
 #if WITH_EDITOR
 	virtual void AddCompileDependency(const UStruct* Struct) { }
diff --git a/Engine/Plugins/Experimental/Chooser/Source/ProxyTable/Private/ProxyAsset.cpp b/Engine/Plugins/Experimental/Chooser/Source/ProxyTable/Private/ProxyAsset.cpp
--- a/Engine/Plugins/Experimental/Chooser/Source/ProxyTable/Private/ProxyAsset.cpp
+++ b/Engine/Plugins/Experimental/Chooser/Source/ProxyTable/Private/ProxyAsset.cpp
@@ -106,15 +106,6 @@ void UProxyAsset::PostLoad()
 	CachedPreviousType = Type;
 	CachedPreviousResultType = ResultType;
 #endif
-
-	if (ProxyTable.IsValid())
-	{
-		// compile property access for Proxy Table fallback codepath
-		if (FChooserParameterProxyTableBase* ProxyReference = ProxyTable.GetMutablePtr<FChooserParameterProxyTableBase>())
-		{
-			ProxyReference->Compile(this, false);
-		}
-	}
 }
 
 void UProxyAsset::PostDuplicate(EDuplicateMode::Type DuplicateMode)
@@ -132,42 +123,14 @@ void UProxyAsset::PostDuplicate(EDuplicateMode::Type DuplicateMode)
 UProxyAsset::UProxyAsset(const FObjectInitializer& Initializer)
 	: Super(Initializer)
 {
-	ProxyTable.InitializeAs(FProxyTableContextProperty::StaticStruct());
 }
 
 UObject* UProxyAsset::FindProxyObject(FChooserEvaluationContext& Context) const
 {
-	if (ProxyTable.IsValid())
-	{
-		const UProxyTable* Table;
-		if (ProxyTable.Get<FChooserParameterProxyTableBase>().GetValue(Context, Table))
-		{
-			if(Table)
-			{
-				if (UObject* Value = Table->FindProxyObject(Guid, Context))
-				{
-					return Value;
-				}
-			}
-		}
-	}
-	
 	return nullptr;
 }
 
 FObjectChooserBase::EIteratorStatus UProxyAsset::FindProxyObjectMulti(FChooserEvaluationContext &Context, FObjectChooserBase::FObjectChooserIteratorCallback Callback) const
 {
-	if (ProxyTable.IsValid())
-	{
-		const UProxyTable* Table;
-		if (ProxyTable.Get<FChooserParameterProxyTableBase>().GetValue(Context, Table))
-		{
-			if(Table)
-			{
-				return Table->FindProxyObjectMulti(Guid, Context, Callback);
-			}
-		}
-	}
-	
 	return FObjectChooserBase::EIteratorStatus::Continue;
 }
diff --git a/Engine/Plugins/Experimental/Chooser/Source/ProxyTable/Private/ProxyTable.cpp b/Engine/Plugins/Experimental/Chooser/Source/ProxyTable/Private/ProxyTable.cpp
--- a/Engine/Plugins/Experimental/Chooser/Source/ProxyTable/Private/ProxyTable.cpp
+++ b/Engine/Plugins/Experimental/Chooser/Source/ProxyTable/Private/ProxyTable.cpp
@@ -173,11 +173,17 @@ void UProxyTable::PostLoad()
 			// need to compile results in case one is a LookupProxy
 			Result->Compile(Entry.ProxyAsset, false);
 		}
-
+		
 		/// compile any struct output property references
 		for(FProxyStructOutput& StructOutput : Entry.OutputStructData)
 		{
 			StructOutput.Binding.Compile(Entry.ProxyAsset);
+#if WITH_EDITORONLY_DATA
+			if (!StructOutput.Binding.CompileMessage.IsEmpty())
+			{
+				UE_LOG(LogChooser, Error, TEXT("ProxyTable Struct Output Compile Error (%s - %s): %s."),  *GetName(), *Entry.ProxyAsset.GetName(), *StructOutput.Binding.CompileMessage.ToString());
+			}
+#endif
 		}
 	}
 }
diff --git a/Engine/Plugins/Experimental/Chooser/Source/ProxyTable/Public/ProxyAsset.h b/Engine/Plugins/Experimental/Chooser/Source/ProxyTable/Public/ProxyAsset.h
--- a/Engine/Plugins/Experimental/Chooser/Source/ProxyTable/Public/ProxyAsset.h
+++ b/Engine/Plugins/Experimental/Chooser/Source/ProxyTable/Public/ProxyAsset.h
@@ -44,13 +44,11 @@ public:
 	UPROPERTY(EditAnywhere, Meta = (ExcludeBaseStruct, BaseStruct = "/Script/Chooser.ContextObjectTypeBase"), Category = "Input")
 	TArray<FInstancedStruct> ContextData;
 
-	UPROPERTY(EditAnywhere, Meta = (ExcludeBaseStruct, BaseStruct ="/Script/ProxyTable.ChooserParameterProxyTableBase"), Category = "Proxy Table Reference")
-	FInstancedStruct ProxyTable;
-
 	UPROPERTY()
 	FGuid Guid;
 
 	virtual TConstArrayView<FInstancedStruct> GetContextData() const override { return ContextData; }
+	virtual FString GetContextOwnerName() const override { return GetName(); }
 	UObject* FindProxyObject(struct FChooserEvaluationContext& Context) const;
 	FObjectChooserBase::EIteratorStatus FindProxyObjectMulti(FChooserEvaluationContext &Context, FObjectChooserBase::FObjectChooserIteratorCallback Callback) const;
 
