make enum column data store enum values as FName

diff --git a/Engine/Plugins/Experimental/Chooser/Source/Chooser/Internal/EnumColumn.h b/Engine/Plugins/Experimental/Chooser/Source/Chooser/Internal/EnumColumn.h
--- a/Engine/Plugins/Experimental/Chooser/Source/Chooser/Internal/EnumColumn.h
+++ b/Engine/Plugins/Experimental/Chooser/Source/Chooser/Internal/EnumColumn.h
@@ -69,6 +69,9 @@ struct FChooserEnumRowData
 #if WITH_EDITORONLY_DATA
 	UPROPERTY()
 	bool CompareNotEqual_DEPRECATED = false;
+
+	UPROPERTY(EditAnywhere, Category = EditorOnly)
+	FName ValueName;
 #endif
 	
 	UPROPERTY(EditAnywhere, Category = Runtime, Meta = (ValidEnumValues = "MatchEqual, MatchNotEqual, MatchAny"))
@@ -81,8 +84,16 @@ struct FChooserEnumRowData
 };
 
 
+USTRUCT(Meta=(Hidden))
+struct CHOOSER_API FEnumColumnBase : public FChooserColumnBase
+{
+	GENERATED_BODY()
+	
+	virtual void EnumChanged(const UEnum* Enum) {}
+};
+
 USTRUCT()
-struct CHOOSER_API FEnumColumn : public FChooserColumnBase
+struct CHOOSER_API FEnumColumn : public FEnumColumnBase
 {
 	GENERATED_BODY()
 public:
@@ -104,6 +115,16 @@ public:
 	virtual void Filter(FChooserEvaluationContext& Context, const FChooserIndexArray& IndexListIn, FChooserIndexArray& IndexListOut) const override;
 	
 #if WITH_EDITOR
+	const UEnum* GetEnum() const 
+	{
+		const UEnum* Enum = nullptr;
+		if (const FChooserParameterEnumBase* Input = InputValue.GetPtr<FChooserParameterEnumBase>())
+		{
+			Enum = Input->GetEnum();
+		}
+		return Enum;
+	}
+	
 	mutable uint8 TestValue = 0;
 	virtual bool EditorTestFilter(int32 RowIndex) const override
 	{
@@ -116,8 +137,10 @@ public:
 		Reader << TestValue;
 	}
 	
-	virtual void AddToDetails (FInstancedPropertyBag& PropertyBag, int32 ColumnIndex, int32 RowIndex);
-	virtual void SetFromDetails(FInstancedPropertyBag& PropertyBag, int32 ColumnIndex, int32 RowIndex);
+	virtual void AddToDetails (FInstancedPropertyBag& PropertyBag, int32 ColumnIndex, int32 RowIndex) override;
+	virtual void SetFromDetails(FInstancedPropertyBag& PropertyBag, int32 ColumnIndex, int32 RowIndex) override;
+
+	virtual void EnumChanged(const UEnum* Enum) override;
 #endif
 	
 	CHOOSER_COLUMN_BOILERPLATE(FChooserParameterEnumBase);
diff --git a/Engine/Plugins/Experimental/Chooser/Source/Chooser/Internal/MultiEnumColumn.h b/Engine/Plugins/Experimental/Chooser/Source/Chooser/Internal/MultiEnumColumn.h
--- a/Engine/Plugins/Experimental/Chooser/Source/Chooser/Internal/MultiEnumColumn.h
+++ b/Engine/Plugins/Experimental/Chooser/Source/Chooser/Internal/MultiEnumColumn.h
@@ -50,6 +50,16 @@ public:
 	virtual void Filter(FChooserEvaluationContext& Context, const FChooserIndexArray& IndexListIn, FChooserIndexArray& IndexListOut) const override;
 	
 #if WITH_EDITOR
+	const UEnum* GetEnum() const 
+	{
+		const UEnum* Enum = nullptr;
+		if (const FChooserParameterEnumBase* Input = InputValue.GetPtr<FChooserParameterEnumBase>())
+		{
+			Enum = Input->GetEnum();
+		}
+		return Enum;
+	}
+	
 	mutable uint8 TestValue = 0;
 	virtual bool EditorTestFilter(int32 RowIndex) const override
 	{
diff --git a/Engine/Plugins/Experimental/Chooser/Source/Chooser/Internal/OutputEnumColumn.h b/Engine/Plugins/Experimental/Chooser/Source/Chooser/Internal/OutputEnumColumn.h
--- a/Engine/Plugins/Experimental/Chooser/Source/Chooser/Internal/OutputEnumColumn.h
+++ b/Engine/Plugins/Experimental/Chooser/Source/Chooser/Internal/OutputEnumColumn.h
@@ -6,6 +6,7 @@
 #include "IChooserParameterEnum.h"
 #include "StructUtils/InstancedStruct.h"
 #include "ChooserPropertyAccess.h"
+#include "EnumColumn.h"
 #include "OutputEnumColumn.generated.h"
 
 USTRUCT()
@@ -15,11 +16,16 @@ struct FChooserOutputEnumRowData
 
 	UPROPERTY(EditAnywhere, Category = Runtime)
 	uint8 Value = 0;
+
+#if WITH_EDITORONLY_DATA
+	UPROPERTY(EditAnywhere, Category = EditorOnly)
+	FName ValueName;
+#endif
 };
 
 
 USTRUCT()
-struct CHOOSER_API FOutputEnumColumn : public FChooserColumnBase
+struct CHOOSER_API FOutputEnumColumn : public FEnumColumnBase
 {
 	GENERATED_BODY()
 public:
@@ -38,6 +44,8 @@ public:
 #if WITH_EDITORONLY_DATA
 	UPROPERTY(EditAnywhere, Category = "Data")
 	FChooserOutputEnumRowData DefaultRowValue;
+
+	virtual void PostLoad() override;
 #endif
 	
 	UPROPERTY(EditAnywhere, Category = "Data")
@@ -46,10 +54,22 @@ public:
 	TArray<FChooserOutputEnumRowData> RowValues;
 	
 #if WITH_EDITOR
+	const UEnum* GetEnum() const 
+	{
+		const UEnum* Enum = nullptr;
+		if (const FChooserParameterEnumBase* Input = InputValue.GetPtr<FChooserParameterEnumBase>())
+		{
+			Enum = Input->GetEnum();
+		}
+		return Enum;
+	}
+
 	mutable uint8 TestValue;
 
 	virtual void AddToDetails (FInstancedPropertyBag& PropertyBag, int32 ColumnIndex, int32 RowIndex) override;
 	virtual void SetFromDetails(FInstancedPropertyBag& PropertyBag, int32 ColumnIndex, int32 RowIndex) override;
+	
+	virtual void EnumChanged(const UEnum* Enum) override;
 #endif
 	
 	CHOOSER_COLUMN_BOILERPLATE(FChooserParameterEnumBase);
diff --git a/Engine/Plugins/Experimental/Chooser/Source/Chooser/Private/Chooser.cpp b/Engine/Plugins/Experimental/Chooser/Source/Chooser/Private/Chooser.cpp
--- a/Engine/Plugins/Experimental/Chooser/Source/Chooser/Private/Chooser.cpp
+++ b/Engine/Plugins/Experimental/Chooser/Source/Chooser/Private/Chooser.cpp
@@ -35,6 +35,15 @@ void UChooserTable::PostLoad()
 	Super::PostLoad();
 
 #if WITH_EDITORONLY_DATA
+	for (FInstancedStruct& ColumnData : ColumnsStructs)
+	{
+		if (ColumnData.IsValid())
+		{
+			FChooserColumnBase& Column = ColumnData.GetMutable<FChooserColumnBase>();
+			Column.PostLoad();
+		}
+	}
+
 	CachedPreviousOutputObjectType = OutputObjectType;
 	CachedPreviousResultType = ResultType;
 
diff --git a/Engine/Plugins/Experimental/Chooser/Source/Chooser/Private/EnumColumn.cpp b/Engine/Plugins/Experimental/Chooser/Source/Chooser/Private/EnumColumn.cpp
--- a/Engine/Plugins/Experimental/Chooser/Source/Chooser/Private/EnumColumn.cpp
+++ b/Engine/Plugins/Experimental/Chooser/Source/Chooser/Private/EnumColumn.cpp
@@ -3,6 +3,7 @@
 #include "ChooserIndexArray.h"
 #include "ChooserPropertyAccess.h"
 #include "ChooserTrace.h"
+#include "Engine/UserDefinedEnum.h"
 #if WITH_EDITOR
 #include "StructUtils/PropertyBag.h"
 #endif
@@ -34,14 +35,33 @@ void FEnumColumn::PostLoad()
 	}
 
 #if WITH_EDITORONLY_DATA
-	// upgrade data for "Any" support
+	const UEnum* Enum = nullptr;
+	if (InputValue.IsValid())
+	{
+		Enum = InputValue.Get<FChooserParameterEnumBase>().GetEnum();
+	}
+	
 	for(FChooserEnumRowData& CellData : RowValues)
 	{
+		// upgrade data for "Any" support
 		if (CellData.CompareNotEqual_DEPRECATED)
 		{
 			CellData.CompareNotEqual_DEPRECATED = false;
 			CellData.Comparison = EEnumColumnCellValueComparison::MatchNotEqual;
 		}
+
+		if (Enum)
+		{
+			if (Enum->IsValidEnumName(CellData.ValueName))
+			{
+				CellData.Value = Enum->GetValueByName(CellData.ValueName);
+			}
+			else
+			// if StringValue is empty (or the names in the enum have changed) upgrade old data, to have a valid Name
+			{
+				CellData.ValueName = Enum->GetNameByValue(CellData.Value);
+			}
+		}
 	}
 #endif
 }
@@ -141,4 +161,26 @@ void FEnumColumn::Filter(FChooserEvaluationContext& Context, const FChooserIndex
 			}
 		}
 	}
+
+	void FEnumColumn::EnumChanged(const UEnum* Enum)
+	{
+		if (InputValue.IsValid())
+		{
+			if (Enum == InputValue.Get<FChooserParameterEnumBase>().GetEnum())
+			{
+				for(FChooserEnumRowData& CellData : RowValues)
+				{
+					if (Enum->IsValidEnumName(CellData.ValueName))
+					{
+						CellData.Value = Enum->GetValueByName(CellData.ValueName);
+					}
+					else
+					// if StringValue is empty (or the names in the enum have changed) upgrade old data, to have a valid Name
+					{
+						CellData.ValueName = Enum->GetNameByValue(CellData.Value);
+					}
+				}
+			}
+		}
+	}
 #endif
\ No newline at end of file
diff --git a/Engine/Plugins/Experimental/Chooser/Source/Chooser/Private/OutputEnumColumn.cpp b/Engine/Plugins/Experimental/Chooser/Source/Chooser/Private/OutputEnumColumn.cpp
--- a/Engine/Plugins/Experimental/Chooser/Source/Chooser/Private/OutputEnumColumn.cpp
+++ b/Engine/Plugins/Experimental/Chooser/Source/Chooser/Private/OutputEnumColumn.cpp
@@ -33,7 +33,48 @@ void FOutputEnumColumn::SetOutputs(FChooserEvaluationContext& Context, int RowIn
 #endif
 }
 
+#if WITH_EDITORONLY_DATA
+void FOutputEnumColumn::PostLoad()
+{
+	Super::PostLoad();
+    	
+    if (InputValue.IsValid())
+    {
+    	InputValue.GetMutable<FChooserParameterBase>().PostLoad();
+    	
+	    if (const UEnum* Enum = InputValue.Get<FChooserParameterEnumBase>().GetEnum())
+	    {
+	    	for(FChooserOutputEnumRowData& CellData : RowValues)
+	    	{
+	    		if (Enum)
+	    		{
+	    			if (Enum->IsValidEnumName(CellData.ValueName))
+	    			{
+	    				CellData.Value = Enum->GetValueByName(CellData.ValueName);
+	    			}
+	    			else
+					// if StringValue is empty (or the names in the enum have changed) upgrade old data, to have a valid Name
+					{
+						CellData.ValueName = Enum->GetNameByValue(CellData.Value);
+					}
+	    		}
+	    	}
+	    	if (Enum->IsValidEnumName(FallbackValue.ValueName))
+			{
+				FallbackValue.Value = Enum->GetValueByName(FallbackValue.ValueName);
+			}
+			else
+			// if StringValue is empty (or the names in the enum have changed) upgrade old data, to have a valid Name
+			{
+				FallbackValue.ValueName = Enum->GetNameByValue(FallbackValue.Value);
+			}
+	    }
+    }
+}
+#endif
+
 #if WITH_EDITOR
+
 	void FOutputEnumColumn::AddToDetails (FInstancedPropertyBag& PropertyBag, int32 ColumnIndex, int32 RowIndex)
 	{
 		FText DisplayName;
@@ -58,4 +99,36 @@ void FOutputEnumColumn::SetOutputs(FChooserEvaluationContext& Context, int RowIn
 			RowValues[RowIndex].Value = *Value;
 		}
 	}
-#endif
\ No newline at end of file
+
+void FOutputEnumColumn::EnumChanged(const UEnum* Enum)
+{
+	if (InputValue.IsValid())
+	{
+		if (Enum == InputValue.Get<FChooserParameterEnumBase>().GetEnum())
+		{
+			for(FChooserOutputEnumRowData& CellData : RowValues)
+			{
+				if (Enum->IsValidEnumName(CellData.ValueName))
+				{
+					CellData.Value = Enum->GetValueByName(CellData.ValueName);
+				}
+				else
+				// if StringValue is empty (or the names in the enum have changed) upgrade old data, to have a valid Name
+				{
+					CellData.ValueName = Enum->GetNameByValue(CellData.Value);
+				}
+			}
+			
+			if (Enum->IsValidEnumName(FallbackValue.ValueName))
+			{
+				FallbackValue.Value = Enum->GetValueByName(FallbackValue.ValueName);
+			}
+			else
+			// if StringValue is empty (or the names in the enum have changed) upgrade old data, to have a valid Name
+			{
+				FallbackValue.ValueName = Enum->GetNameByValue(FallbackValue.Value);
+			}
+		}
+	}
+}
+#endif
diff --git a/Engine/Plugins/Experimental/Chooser/Source/ChooserEditor/Private/ChooserEditorModule.cpp b/Engine/Plugins/Experimental/Chooser/Source/ChooserEditor/Private/ChooserEditorModule.cpp
--- a/Engine/Plugins/Experimental/Chooser/Source/ChooserEditor/Private/ChooserEditorModule.cpp
+++ b/Engine/Plugins/Experimental/Chooser/Source/ChooserEditor/Private/ChooserEditorModule.cpp
@@ -26,6 +26,9 @@
 #include "PropertyEditorModule.h"
 #include "RandomizeColumnEditor.h"
 #include "ChooserTrack.h"
+#include "Kismet2/EnumEditorUtils.h"
+#include "EnumColumn.h"
+#include "UObject/UObjectIterator.h"
 
 #define LOCTEXT_NAMESPACE "ChooserEditorModule"
 
@@ -34,6 +37,22 @@ namespace UE::ChooserEditor
 	
 FChoosersTrackCreator GChoosersTrackCreator;
 
+void FEnumChangedListener::PostChange(const UUserDefinedEnum* Changed, FEnumEditorUtils::EEnumEditorChangeInfo ChangedType)
+{
+	// iterate over all loaded choosers
+	for (TObjectIterator<UChooserTable> It; It; ++It)
+	{
+		UChooserTable* Chooser = *It;
+		for (FInstancedStruct& ColumnData : Chooser->ColumnsStructs)
+		{
+			if (FEnumColumnBase* EnumColumn = ColumnData.GetMutablePtr<FEnumColumnBase>())
+			{
+				EnumColumn->EnumChanged(Changed);
+			}
+		}
+	}
+}
+
 void FModule::StartupModule()
 {
 	FChooserEditorStyle::Initialize();
@@ -67,6 +86,8 @@ void FModule::StartupModule()
 	IModularFeatures::Get().RegisterModularFeature( IRewindDebuggerExtension::ModularFeatureName, &RewindDebuggerChooser);
 	IModularFeatures::Get().RegisterModularFeature(RewindDebugger::IRewindDebuggerTrackCreator::ModularFeatureName, &GChoosersTrackCreator);
 	IModularFeatures::Get().RegisterModularFeature(TraceServices::ModuleFeatureName, &ChooserTraceModule);
+
+	FEnumEditorUtils::FEnumEditorManager::Get().AddListener(&EnumChanged);
 }
 
 void FModule::ShutdownModule()
@@ -77,6 +98,8 @@ void FModule::ShutdownModule()
 	FChooserTableEditorCommands::Unregister();
 	
 	FChooserEditorStyle::Shutdown();
+	
+	FEnumEditorUtils::FEnumEditorManager::Get().RemoveListener(&EnumChanged);
 }
 
 }
diff --git a/Engine/Plugins/Experimental/Chooser/Source/ChooserEditor/Private/ChooserEditorModule.h b/Engine/Plugins/Experimental/Chooser/Source/ChooserEditor/Private/ChooserEditorModule.h
--- a/Engine/Plugins/Experimental/Chooser/Source/ChooserEditor/Private/ChooserEditorModule.h
+++ b/Engine/Plugins/Experimental/Chooser/Source/ChooserEditor/Private/ChooserEditorModule.h
@@ -6,10 +6,17 @@
 #include "ChooserTraceModule.h"
 #include "Modules/ModuleInterface.h"
 #include "RewindDebuggerChooser.h"
+#include "Kismet2/EnumEditorUtils.h"
 
 namespace UE::ChooserEditor
 {
 
+class FEnumChangedListener : public FEnumEditorUtils::FEnumEditorManager::BaseNotifyOnChanged
+{
+	virtual void PostChange(const UUserDefinedEnum* Changed, FEnumEditorUtils::EEnumEditorChangeInfo ChangedType) override;
+	virtual void PreChange(const UUserDefinedEnum* Changed, FEnumEditorUtils::EEnumEditorChangeInfo ChangedType) override {};
+};
+
 class FModule : public IModuleInterface
 {
 private:
@@ -18,6 +25,8 @@ private:
 	
 	FRewindDebuggerChooser RewindDebuggerChooser;
 	FChooserTraceModule ChooserTraceModule;
+
+	FEnumChangedListener EnumChanged;
 };
 
-}
\ No newline at end of file
+}
diff --git a/Engine/Plugins/Experimental/Chooser/Source/ChooserEditor/Private/EnumColumnEditor.cpp b/Engine/Plugins/Experimental/Chooser/Source/ChooserEditor/Private/EnumColumnEditor.cpp
--- a/Engine/Plugins/Experimental/Chooser/Source/ChooserEditor/Private/EnumColumnEditor.cpp
+++ b/Engine/Plugins/Experimental/Chooser/Source/ChooserEditor/Private/EnumColumnEditor.cpp
@@ -1,23 +1,88 @@
 // Copyright Epic Games, Inc. All Rights Reserved.
 
 #include "EnumColumnEditor.h"
+#include "ChooserColumnHeader.h"
+#include "ChooserTableEditor.h"
+#include "DetailLayoutBuilder.h"
 #include "EnumColumn.h"
+#include "GraphEditorSettings.h"
+#include "ObjectChooserWidgetFactories.h"
 #include "OutputEnumColumn.h"
+#include "ScopedTransaction.h"
 #include "SPropertyAccessChainWidget.h"
-#include "ObjectChooserWidgetFactories.h"
-#include "ChooserTableEditor.h"
-#include "ChooserColumnHeader.h"
-#include "GraphEditorSettings.h"
-#include "SEnumCombo.h"
-#include "TransactionCommon.h"
+#include "Framework/MultiBox/MultiBoxBuilder.h"
 #include "Widgets/Input/SButton.h"
-#include "ScopedTransaction.h"
+#include "Widgets/Input/SComboButton.h"
 
 #define LOCTEXT_NAMESPACE "EnumColumnEditor"
 
 namespace UE::ChooserEditor
 {
 
+TSharedRef<SWidget> SEnumCell::GenerateEnumMenu() const
+{
+	if (const UEnum* EnumSource = Enum.Get())
+	{
+		FMenuBuilder MenuBuilder(true, nullptr);
+		for (int32 EnumIndex = 0; EnumIndex < EnumSource->NumEnums() - 1; ++EnumIndex)
+		{
+			const bool bIsHidden = EnumSource->HasMetaData(TEXT("Hidden"), EnumIndex);
+			if (!bIsHidden)
+			{
+				MenuBuilder.AddMenuEntry(EnumSource->GetDisplayNameTextByIndex(EnumIndex),
+					FText()/*todo tooltip*/,
+					FSlateIcon(),
+					FUIAction(
+					FExecuteAction::CreateLambda([this, EnumValue = EnumSource->GetValueByIndex(EnumIndex)]() 
+					{
+						OnValueSet.ExecuteIfBound(EnumValue);
+					})	
+					));
+			}
+		}
+		return MenuBuilder.MakeWidget();
+	}
+	return SNullWidget::NullWidget;
+}
+	
+TSharedRef<SWidget> SEnumCell::CreateEnumComboBox()
+{
+	return SNew(SComboButton)
+		.IsEnabled_Lambda([this](){ return IsEnabled(); } )
+		.OnGetMenuContent(this, &SEnumCell::GenerateEnumMenu)
+		.VAlign(VAlign_Center)
+		.ButtonContent()
+		[
+			SNew(STextBlock)
+			.Font(IDetailLayoutBuilder::GetDetailFont())
+			.Text_Lambda([this]()
+			{
+				if (const UEnum* EnumSource = Enum.Get())
+				{
+					return EnumSource->GetDisplayNameTextByValue(EnumValue.Get());
+				}
+				return FText();
+			})
+		];
+	
+	return SNullWidget::NullWidget;
+}
+
+void SEnumCell::Construct( const FArguments& InArgs)
+{
+	SetEnabled(InArgs._IsEnabled);
+	
+	Enum = InArgs._Enum;
+	EnumValue = InArgs._EnumValue;
+	OnValueSet = InArgs._OnValueSet;
+
+	ChildSlot
+	[
+		CreateEnumComboBox()
+	];
+}
+
+
 TSharedRef<SWidget> CreateEnumColumnWidget(UChooserTable* Chooser, FChooserColumnBase* Column, int Row)
 {
 	FEnumColumn* EnumColumn = static_cast<FEnumColumn*>(Column);
@@ -35,13 +100,18 @@ TSharedRef<SWidget> CreateEnumColumnWidget(UChooserTable* Chooser, FChooserColum
 		TSharedPtr<SWidget> DebugWidget = nullptr;
 		if (Chooser->GetEnableDebugTesting())
 		{
-			DebugWidget = SNew(SEnumCell<FEnumColumn>).TransactionObject(Chooser).EnumColumn(EnumColumn)
+			DebugWidget = SNew(SEnumCell)
+								.Enum_Lambda([EnumColumn]() { return EnumColumn->GetEnum(); })
 								.OnValueSet_Lambda([EnumColumn](int Value) { EnumColumn->TestValue = Value; })
 								.EnumValue_Lambda([EnumColumn]() { return EnumColumn->TestValue; })
 								.IsEnabled_Lambda([Chooser] { return !Chooser->HasDebugTarget(); });
 		}
 
-		return MakeColumnHeaderWidget(Chooser, Column, ColumnName, ColumnTooltip, ColumnIcon, DebugWidget);
+		return MakeColumnHeaderWidget(Chooser, Column, ColumnName, ColumnTooltip, ColumnIcon, DebugWidget,
+			FChooserWidgetValueChanged::CreateLambda([EnumColumn]()
+			{
+				EnumColumn->EnumChanged(EnumColumn->InputValue.Get<FChooserParameterEnumBase>().GetEnum());
+			}));
 	}
 
 	// create cell widget
@@ -85,19 +155,31 @@ TSharedRef<SWidget> CreateEnumColumnWidget(UChooserTable* Chooser, FChooserColum
 			]
 			+ SHorizontalBox::Slot().FillWidth(1)
 			[
-				SNew(SEnumCell<FEnumColumn>).TransactionObject(Chooser).EnumColumn(EnumColumn)
-					.OnValueSet_Lambda([EnumColumn, Row](int Value)
+				SNew(SEnumCell)
+					.Enum_Lambda([EnumColumn]() { return EnumColumn->GetEnum(); })
+					.OnValueSet_Lambda([Chooser, EnumColumn, Row](int Value)
 					{
+						FScopedTransaction Transaction(LOCTEXT("Set Enum Value", "Set Enum Value"));
+						Chooser->Modify();
+						
 						if (EnumColumn->RowValues.IsValidIndex(Row))
 						{
 							EnumColumn->RowValues[Row].Value = static_cast<uint8>(Value);
+
+							if (EnumColumn->InputValue.IsValid())
+							{
+								if (const UEnum* Enum = EnumColumn->InputValue.Get<FChooserParameterEnumBase>().GetEnum())
+								{
+									EnumColumn->RowValues[Row].ValueName = Enum->GetNameByValue(Value);
+								}
+							}
 						}
 					})
 					.EnumValue_Lambda([EnumColumn, Row]()
 					{
 						return EnumColumn->RowValues.IsValidIndex(Row) ? static_cast<int32>(EnumColumn->RowValues[Row].Value) : 0;
 					})
-				.Visibility_Lambda([EnumColumn,Column, Row]()
+					.Visibility_Lambda([EnumColumn,Column, Row]()
 					{
 						return (EnumColumn->RowValues.IsValidIndex(Row) &&
 								EnumColumn->RowValues[Row].Comparison == EEnumColumnCellValueComparison::MatchAny)
@@ -120,27 +202,60 @@ TSharedRef<SWidget> CreateOutputEnumColumnWidget(UChooserTable* Chooser, FChoose
 		TSharedPtr<SWidget> DebugWidget = nullptr;
 		if (Chooser->GetEnableDebugTesting())
 		{
-			DebugWidget = SNew(SEnumCell<FOutputEnumColumn>).TransactionObject(Chooser).EnumColumn(EnumColumn).IsEnabled(false)
+			DebugWidget = SNew(SEnumCell)
+					.IsEnabled(false)
+					.Enum_Lambda([EnumColumn](){ return EnumColumn->GetEnum(); } )
 					.EnumValue_Lambda([EnumColumn]() { return static_cast<int32>(EnumColumn->TestValue); });
 		}
 
-		return MakeColumnHeaderWidget(Chooser, Column, ColumnName, ColumnTooltip, ColumnIcon, DebugWidget);
+		return MakeColumnHeaderWidget(Chooser, Column, ColumnName, ColumnTooltip, ColumnIcon, DebugWidget,
+				FChooserWidgetValueChanged::CreateLambda([EnumColumn]()
+				{
+					EnumColumn->EnumChanged(EnumColumn->InputValue.Get<FChooserParameterEnumBase>().GetEnum());
+				}));	
 	}
 	else if (Row == ColumnWidget_SpecialIndex_Fallback)
 	{
-		return 	SNew(SEnumCell<FOutputEnumColumn>).TransactionObject(Chooser).EnumColumn(EnumColumn)
-        			.OnValueSet_Lambda([EnumColumn](int Value) { EnumColumn->FallbackValue.Value = Value; })
+		return SNew(SEnumCell)
+					.Enum_Lambda([EnumColumn](){ return EnumColumn->GetEnum(); } )
+        			.OnValueSet_Lambda([Chooser, EnumColumn](int Value)
+        			{
+						FScopedTransaction Transaction(LOCTEXT("Set Enum Value", "Set Enum Value"));
+						Chooser->Modify();
+					
+						EnumColumn->FallbackValue.Value = Value;
+
+						if (EnumColumn->InputValue.IsValid())
+						{
+							if (const UEnum* Enum = EnumColumn->InputValue.Get<FChooserParameterEnumBase>().GetEnum())
+							{
+								EnumColumn->FallbackValue.ValueName = Enum->GetNameByValue(Value);
+							}
+						}
+        			})
         			.EnumValue_Lambda([EnumColumn]() { return EnumColumn->FallbackValue.Value; });
 	}
 
 	// create cell widget
 	
-	return SNew(SEnumCell<FOutputEnumColumn>).TransactionObject(Chooser).EnumColumn(EnumColumn)
-		.OnValueSet_Lambda([EnumColumn, Row](int Value)
+	return SNew(SEnumCell)
+		.Enum_Lambda([EnumColumn](){ return EnumColumn->GetEnum(); } )
+		.OnValueSet_Lambda([EnumColumn, Chooser, Row](int Value)
 		{
+			FScopedTransaction Transaction(LOCTEXT("Set Enum Value", "Set Enum Value"));
+			Chooser->Modify();
+			
 			if (EnumColumn->RowValues.IsValidIndex(Row))
 			{
 				EnumColumn->RowValues[Row].Value = static_cast<uint8>(Value);
+
+				if (EnumColumn->InputValue.IsValid())
+				{
+					if (const UEnum* Enum = EnumColumn->InputValue.Get<FChooserParameterEnumBase>().GetEnum())
+					{
+						EnumColumn->RowValues[Row].ValueName = Enum->GetNameByValue(Value);
+					}
+				}
 			}
 		})
 		.EnumValue_Lambda([EnumColumn, Row]()
diff --git a/Engine/Plugins/Experimental/Chooser/Source/ChooserEditor/Private/EnumColumnEditor.h b/Engine/Plugins/Experimental/Chooser/Source/ChooserEditor/Private/EnumColumnEditor.h
--- a/Engine/Plugins/Experimental/Chooser/Source/ChooserEditor/Private/EnumColumnEditor.h
+++ b/Engine/Plugins/Experimental/Chooser/Source/ChooserEditor/Private/EnumColumnEditor.h
@@ -3,19 +3,14 @@
 #pragma once
 
 #include "CoreMinimal.h"
-#include "SEnumCombo.h"
 #include "IChooserParameterEnum.h"
-#include "Widgets/Layout/SBorder.h"
 #include "Widgets/SCompoundWidget.h"
-#include "ScopedTransaction.h"
 
 #define LOCTEXT_NAMESPACE "EnumColumnEditor"
 
 namespace UE::ChooserEditor
 {
-	
-	// Wrapper widget for EnumComboBox which will reconstruct the combo box when the Enum has changed
-	template <typename ColumnType>
+	// Enum combo box widget with a dynamic UEnum
 	class SEnumCell : public SCompoundWidget
 	{
 	public:
@@ -25,104 +20,24 @@ namespace UE::ChooserEditor
 		SLATE_BEGIN_ARGS(SEnumCell)
 		{}
 
-		SLATE_ARGUMENT(UObject*, TransactionObject)
-		SLATE_ARGUMENT(ColumnType*, EnumColumn)
+		SLATE_ATTRIBUTE(const UEnum*, Enum);
 		SLATE_ATTRIBUTE(int32, EnumValue);
 		SLATE_EVENT(FOnValueSet, OnValueSet)
 				
 		SLATE_END_ARGS()
 
-		TSharedRef<SWidget> CreateEnumComboBox()
-		{
-			if (const ColumnType* EnumColumnPointer = EnumColumn)
-			{
-				if (EnumColumnPointer->InputValue.IsValid())
-				{
-					if (const UEnum* Enum = EnumColumnPointer->InputValue.template Get<FChooserParameterEnumBase>().GetEnum())
-					{
-						return SNew(SEnumComboBox, Enum)
-							.IsEnabled_Lambda([this](){ return IsEnabled(); } )
-							.CurrentValue(EnumValue)
-							.OnEnumSelectionChanged_Lambda([this](int32 InEnumValue, ESelectInfo::Type)
-							{
-								const FScopedTransaction Transaction(LOCTEXT("Edit RHS", "Edit Enum Value"));
-								TransactionObject->Modify(true);
-								OnValueSet.ExecuteIfBound(InEnumValue);
-							});
-					}
-				}
-			}
-			
-			return SNullWidget::NullWidget;
-		}
-
-		void UpdateEnumComboBox()
-		{
-			ChildSlot[ CreateEnumComboBox()	];
-		}
-
-		virtual void Tick(const FGeometry& AllottedGeometry, const double InCurrentTime, const float InDeltaTime) override
-		{
-			const UEnum* CurrentEnumSource = nullptr;
-			if (EnumColumn->InputValue.IsValid())
-			{
-				CurrentEnumSource = EnumColumn->InputValue.template Get<FChooserParameterEnumBase>().GetEnum(); 
-			}
-			if (EnumSource != CurrentEnumSource)
-			{
-				EnumComboBorder->SetContent(CreateEnumComboBox());
-				EnumSource = CurrentEnumSource;
-			}
-		}
-							
-		void Construct( const FArguments& InArgs)
-		{
-			SetEnabled(InArgs._IsEnabled);
-			
-			SetCanTick(true);
-			EnumColumn = InArgs._EnumColumn;
-			TransactionObject = InArgs._TransactionObject;
-			EnumValue = InArgs._EnumValue;
-			OnValueSet = InArgs._OnValueSet;
-
-			if (EnumColumn)
-			{
-				if (EnumColumn->InputValue.IsValid())
-				{
-					EnumSource = EnumColumn->InputValue.template Get<FChooserParameterEnumBase>().GetEnum();
-				}
-			}
-
-			UpdateEnumComboBox();
-
-			int Row = RowIndex.Get();
-
-			ChildSlot
-			[
-				SAssignNew(EnumComboBorder, SBorder).Padding(0).BorderBackgroundColor(FLinearColor(0,0,0,0))
-				[
-					CreateEnumComboBox()
-				]
-			];
-			
-		}
+		TSharedRef<SWidget> GenerateEnumMenu() const;
+		TSharedRef<SWidget> CreateEnumComboBox();
+		
+		void Construct( const FArguments& InArgs);
 
-		~SEnumCell()
-		{
-		}
+		virtual ~SEnumCell() {}
 
 	private:
-		UObject* TransactionObject = nullptr;
-		ColumnType* EnumColumn = nullptr;
-		const UEnum* EnumSource = nullptr;
-		TSharedPtr<SBorder> EnumComboBorder;
-		TAttribute<int> RowIndex;
-		FDelegateHandle EnumChangedHandle;
-		
 		FOnValueSet OnValueSet;
 		TAttribute<int32> EnumValue;
+		TAttribute<const UEnum*> Enum;
 	};
-
 	
 	void RegisterEnumWidgets();
 }
diff --git a/Engine/Plugins/Experimental/Chooser/Source/ChooserEditor/Private/MultiEnumColumnEditor.cpp b/Engine/Plugins/Experimental/Chooser/Source/ChooserEditor/Private/MultiEnumColumnEditor.cpp
--- a/Engine/Plugins/Experimental/Chooser/Source/ChooserEditor/Private/MultiEnumColumnEditor.cpp
+++ b/Engine/Plugins/Experimental/Chooser/Source/ChooserEditor/Private/MultiEnumColumnEditor.cpp
@@ -149,10 +149,11 @@ TSharedRef<SWidget> CreateMultiEnumColumnWidget(UChooserTable* Chooser, FChooser
 		TSharedPtr<SWidget> DebugWidget = nullptr;
 		if (Chooser->GetEnableDebugTesting())
 		{
-			DebugWidget = SNew(SEnumCell<FMultiEnumColumn>).TransactionObject(Chooser).EnumColumn(MultiEnumColumn)
-								.OnValueSet_Lambda([MultiEnumColumn](int Value) { MultiEnumColumn->TestValue = Value; })
-								.EnumValue_Lambda([MultiEnumColumn]() { return MultiEnumColumn->TestValue; })
-								.IsEnabled_Lambda([Chooser] { return !Chooser->HasDebugTarget(); });
+			DebugWidget = SNew(SEnumCell)
+							.Enum_Lambda([MultiEnumColumn] () { return MultiEnumColumn->GetEnum(); })
+							.OnValueSet_Lambda([MultiEnumColumn](int Value) { MultiEnumColumn->TestValue = Value; })
+							.EnumValue_Lambda([MultiEnumColumn]() { return MultiEnumColumn->TestValue; })
+							.IsEnabled_Lambda([Chooser] { return !Chooser->HasDebugTarget(); });
 			
 			// need to fix support for bitfield enums:
 			// DebugWidget = SNew(SMultiEnumCell<FMultiEnumColumn>).TransactionObject(Chooser).MultiEnumColumn(MultiEnumColumn)
