Fix clang compile error

diff --git a/Engine/Plugins/Experimental/Chooser/Source/ChooserEditor/Private/EnumColumnEditor.cpp b/Engine/Plugins/Experimental/Chooser/Source/ChooserEditor/Private/EnumColumnEditor.cpp
--- a/Engine/Plugins/Experimental/Chooser/Source/ChooserEditor/Private/EnumColumnEditor.cpp
+++ b/Engine/Plugins/Experimental/Chooser/Source/ChooserEditor/Private/EnumColumnEditor.cpp
@@ -33,9 +33,9 @@ TSharedRef<SWidget> SEnumCell::GenerateEnumMenu() const
 					FText()/*todo tooltip*/,
 					FSlateIcon(),
 					FUIAction(
-					FExecuteAction::CreateLambda([this, EnumValue = EnumSource->GetValueByIndex(EnumIndex)]() 
+					FExecuteAction::CreateLambda([this, MenuEntryEnumValue = EnumSource->GetValueByIndex(EnumIndex)]() 
 					{
-						OnValueSet.ExecuteIfBound(EnumValue);
+						OnValueSet.ExecuteIfBound(MenuEntryEnumValue);
 					})	
 					));
 			}
