add  null checking for null class/struct in chooser property access error logging

diff --git a/Engine/Plugins/Experimental/Chooser/Source/Chooser/Private/ChooserPropertyAccess.cpp b/Engine/Plugins/Experimental/Chooser/Source/Chooser/Private/ChooserPropertyAccess.cpp
--- a/Engine/Plugins/Experimental/Chooser/Source/Chooser/Private/ChooserPropertyAccess.cpp
+++ b/Engine/Plugins/Experimental/Chooser/Source/Chooser/Private/ChooserPropertyAccess.cpp
@@ -331,60 +331,66 @@ namespace UE::Chooser
 
 			if (ExpectedClassType)
 			{
-				if (Context.Params.IsValidIndex(i) && Context.Params[i].IsValid())
+				if (ExpectedClassType->Class)
 				{
-					if (FChooserEvaluationInputObject* InputObjectParam = Context.Params[i].GetPtr<FChooserEvaluationInputObject>())
+					if (Context.Params.IsValidIndex(i) && Context.Params[i].IsValid())
 					{
-						if (InputObjectParam->Object)
+						if (FChooserEvaluationInputObject* InputObjectParam = Context.Params[i].GetPtr<FChooserEvaluationInputObject>())
 						{
-							if (!InputObjectParam->Object->GetClass()->IsChildOf(ExpectedClassType->Class))
+							if (InputObjectParam->Object)
 							{
-								UE_LOG(LogChooser, Error, TEXT("Chooser Table: %s ContextData entry %d expects an object of type %s, but an object of type %s was passed in."),
-									ToCStr(Chooser->GetName()), i, ToCStr(ExpectedClassType->Class->GetName()), ToCStr(InputObjectParam->Object->GetClass()->GetName()));
+								if (!InputObjectParam->Object->GetClass()->IsChildOf(ExpectedClassType->Class))
+								{
+									UE_LOG(LogChooser, Error, TEXT("Chooser Table: %s ContextData entry %d expects an object of type %s, but an object of type %s was passed in."),
+										ToCStr(Chooser->GetName()), i, ToCStr(ExpectedClassType->Class->GetName()), ToCStr(InputObjectParam->Object->GetClass()->GetName()));
+								}
 							}
+							else
+							{
+								UE_LOG(LogChooser, Error, TEXT("Chooser Table: %s ContextData entry %d expects an object of type %s, but null was passed in."),
+									ToCStr(Chooser->GetName()), i, ToCStr(ExpectedClassType->Class->GetName()));
+							}
+
 						}
 						else
 						{
-							UE_LOG(LogChooser, Error, TEXT("Chooser Table: %s ContextData entry %d expects an object of type %s, but null was passed in."),
-								ToCStr(Chooser->GetName()), i, ToCStr(ExpectedClassType->Class->GetName()));
+							UE_LOG(LogChooser, Error, TEXT("Chooser Table: %s ContextData entry %d expects an object of type %s, but was passed a struct of type %s."),
+								ToCStr(Chooser->GetName()), i, ToCStr(ExpectedClassType->Class->GetName()), ToCStr(Context.Params[i].GetScriptStruct()->GetName()));
 						}
-
 					}
 					else
 					{
-						UE_LOG(LogChooser, Error, TEXT("Chooser Table: %s ContextData entry %d expects an object of type %s, but was passed a struct of type %s."),
-							ToCStr(Chooser->GetName()), i, ToCStr(ExpectedClassType->Class->GetName()), ToCStr(Context.Params[i].GetScriptStruct()->GetName()));
+						UE_LOG(LogChooser, Error, TEXT("Chooser Table: %s ContextData entry %d expects an object of type %s, but nothing was passed in."),
+							ToCStr(Chooser->GetName()), i, ToCStr(ExpectedClassType->Class->GetName()));
 					}
 				}
-				else
-				{
-					UE_LOG(LogChooser, Error, TEXT("Chooser Table: %s ContextData entry %d expects an object of type %s, but nothing was passed in."),
-						ToCStr(Chooser->GetName()), i, ToCStr(ExpectedClassType->Class->GetName()));
-				}
 			}
 			else if (ExpectedStructType)
 			{
-				if (Context.Params.IsValidIndex(i) && Context.Params[i].IsValid())
+				if (ExpectedStructType->Struct)
 				{
-					if (FChooserEvaluationInputObject* InputObjectParam = Context.Params[i].GetPtr<FChooserEvaluationInputObject>())
+					if (Context.Params.IsValidIndex(i) && Context.Params[i].IsValid())
 					{
-						UE_LOG(LogChooser, Error, TEXT("Chooser Table: %s ContextData entry %d expects an struct of type %s, but was passed a object of type %s."),
-							ToCStr(Chooser->GetName()), i, ToCStr(ExpectedStructType->Struct->GetName()), ToCStr(InputObjectParam->Object->GetClass()->GetName()));
+						if (FChooserEvaluationInputObject* InputObjectParam = Context.Params[i].GetPtr<FChooserEvaluationInputObject>())
+						{
+							UE_LOG(LogChooser, Error, TEXT("Chooser Table: %s ContextData entry %d expects an struct of type %s, but was passed a object of type %s."),
+								ToCStr(Chooser->GetName()), i, ToCStr(ExpectedStructType->Struct->GetName()), ToCStr(InputObjectParam->Object->GetClass()->GetName()));
+						}
+						else
+						{
+							if (Context.Params[i].GetScriptStruct() != ExpectedStructType->Struct)
+							{
+								UE_LOG(LogChooser, Error, TEXT("Chooser Table: %s ContextData entry %d expects an struct of type %s, but was passed a struct of type %s."),
+									ToCStr(Chooser->GetName()), i, ToCStr(ExpectedStructType->Struct->GetName()), ToCStr(Context.Params[i].GetScriptStruct()->GetName()));
+							}
+						}
 					}
 					else
 					{
-						if (Context.Params[i].GetScriptStruct() != ExpectedStructType->Struct)
-						{
-							UE_LOG(LogChooser, Error, TEXT("Chooser Table: %s ContextData entry %d expects an struct of type %s, but was passed a struct of type %s."),
-								ToCStr(Chooser->GetName()), i, ToCStr(ExpectedStructType->Struct->GetName()), ToCStr(Context.Params[i].GetScriptStruct()->GetName()));
-						}
+						UE_LOG(LogChooser, Error, TEXT("Chooser Table: %s ContextData entry %d expects an struct of type %s, but nothing was passed in."),
+							ToCStr(Chooser->GetName()), i, ToCStr(ExpectedStructType->Struct->GetName()));
 					}
 				}
-				else
-				{
-					UE_LOG(LogChooser, Error, TEXT("Chooser Table: %s ContextData entry %d expects an struct of type %s, but nothing was passed in."),
-						ToCStr(Chooser->GetName()), i, ToCStr(ExpectedStructType->Struct->GetName()));
-				}
 			}
 			else
 			{
