Update chooser compile errors use UE_ASSET_LOG

diff --git a/Engine/Plugins/Experimental/Chooser/Source/Chooser/Internal/Chooser.h b/Engine/Plugins/Experimental/Chooser/Source/Chooser/Internal/Chooser.h
--- a/Engine/Plugins/Experimental/Chooser/Source/Chooser/Internal/Chooser.h
+++ b/Engine/Plugins/Experimental/Chooser/Source/Chooser/Internal/Chooser.h
@@ -186,6 +186,7 @@ public:
 
 	virtual TConstArrayView<FInstancedStruct> GetContextData() const override { return GetRootChooser()->ContextData; }
 	virtual FString GetContextOwnerName() const override { return GetName(); }
+	virtual UObject* GetContextOwnerAsset() override { return this; }
 
 	bool IsCookedData() const { return !CookedResults.IsEmpty(); }
 	static FObjectChooserBase::EIteratorStatus EvaluateChooser(FChooserEvaluationContext& Context, const UChooserTable* Chooser, FObjectChooserBase::FObjectChooserIteratorCallback Callback);
diff --git a/Engine/Plugins/Experimental/Chooser/Source/Chooser/Private/ChooserPropertyAccess.cpp b/Engine/Plugins/Experimental/Chooser/Source/Chooser/Private/ChooserPropertyAccess.cpp
--- a/Engine/Plugins/Experimental/Chooser/Source/Chooser/Private/ChooserPropertyAccess.cpp
+++ b/Engine/Plugins/Experimental/Chooser/Source/Chooser/Private/ChooserPropertyAccess.cpp
@@ -87,7 +87,7 @@ void FChooserPropertyBinding::Compile(IHasContextClass* Owner, bool bForce)
 #if WITH_EDITORONLY_DATA
 		CompileMessage = LOCTEXT("No Property Bound", "No Property Bound");
 #endif
-		UE_LOG(LogChooser, Error, TEXT("Chooser Property Compile Error (%s): Missing property binding."), *Owner->GetContextOwnerName());
+		UE_ASSET_LOG(LogChooser, Error, Owner->GetContextOwnerAsset(), TEXT("Missing property binding."));
 		CompiledBinding = nullptr;
 		return;
 	}
@@ -107,7 +107,7 @@ void FChooserPropertyBinding::Compile(IHasContextClass* Owner, bool bForce)
 #if WITH_EDITORONLY_DATA
 		CompileMessage = FText::Format(LOCTEXT("No valid struct", "No valid Context Object/Struct at index: {0}"), FText::FromString(FString::FromInt(ContextIndex)));
 #endif
-		UE_LOG(LogChooser, Error, TEXT("Chooser Property Compile Error (%s): No valid Context Object/Struct at index: %d"), *Owner->GetContextOwnerName(), ContextIndex);
+		UE_ASSET_LOG(LogChooser, Error, Owner->GetContextOwnerAsset(), TEXT("No valid Context Object/Struct at index: %d"), ContextIndex);
 		CompiledBinding = nullptr;
 		return;
 	}
@@ -205,7 +205,7 @@ void FChooserPropertyBinding::Compile(IHasContextClass* Owner, bool bForce)
 #if WITH_EDITORONLY_DATA
 			CompileMessage = FText::Format(LOCTEXT("Property Not Found", "Property/Function: {0} not Found on Class/Struct: {1}"), FText::FromName(PropertyBindingChain[PropertyChainIndex]), StructType->GetDisplayNameText());
 #endif
-			UE_LOG(LogChooser, Error, TEXT(" Chooser Property Compile Error (%s): Property/Function: %s not Found on Class/Struct %s"), *Owner->GetContextOwnerName(), *PropertyBindingChain[PropertyChainIndex].ToString(), *StructType->GetName());
+			UE_ASSET_LOG(LogChooser, Error, Owner->GetContextOwnerAsset(), TEXT("Property/Function: %s not Found on Class/Struct %s"), *PropertyBindingChain[PropertyChainIndex].ToString(), *StructType->GetName());
 			CompiledBinding = nullptr;
 			return;
 		}
@@ -306,7 +306,7 @@ void FChooserPropertyBinding::Compile(IHasContextClass* Owner, bool bForce)
 #if WITH_EDITORONLY_DATA
 		CompileMessage = FText::Format(LOCTEXT("Property Not Found", "Property/Function: {0} not Found on Class/Struct: {1}"), FText::FromName(PropertyBindingChain.Last()), StructType->GetDisplayNameText());
 #endif
-		UE_LOG(LogChooser, Error, TEXT(" Chooser Property Compile Error (%s): Property/Function: %s not Found on Class/Struct %s"), *Owner->GetContextOwnerName(), *PropertyBindingChain.Last().ToString(), *StructType->GetName());
+		UE_ASSET_LOG(LogChooser, Error, Owner->GetContextOwnerAsset(), TEXT("Property/Function: %s not Found on Class/Struct %s"), *PropertyBindingChain.Last().ToString(), *StructType->GetName());
 		CompiledBinding = nullptr;
  	}
 }
diff --git a/Engine/Plugins/Experimental/Chooser/Source/Chooser/Public/IHasContext.h b/Engine/Plugins/Experimental/Chooser/Source/Chooser/Public/IHasContext.h
--- a/Engine/Plugins/Experimental/Chooser/Source/Chooser/Public/IHasContext.h
+++ b/Engine/Plugins/Experimental/Chooser/Source/Chooser/Public/IHasContext.h
@@ -26,6 +26,7 @@ public:
 	virtual TConstArrayView<FInstancedStruct> GetContextData() const { return TConstArrayView<FInstancedStruct>(); }
 
 	virtual FString GetContextOwnerName() const { return ""; }
+	virtual UObject* GetContextOwnerAsset() { return nullptr; }
 
 	virtual void Compile(bool bForce = false) {}
 #if WITH_EDITOR
diff --git a/Engine/Plugins/Experimental/Chooser/Source/ProxyTable/Public/ProxyAsset.h b/Engine/Plugins/Experimental/Chooser/Source/ProxyTable/Public/ProxyAsset.h
--- a/Engine/Plugins/Experimental/Chooser/Source/ProxyTable/Public/ProxyAsset.h
+++ b/Engine/Plugins/Experimental/Chooser/Source/ProxyTable/Public/ProxyAsset.h
@@ -49,6 +49,8 @@ public:
 
 	virtual TConstArrayView<FInstancedStruct> GetContextData() const override { return ContextData; }
 	virtual FString GetContextOwnerName() const override { return GetName(); }
+	virtual UObject* GetContextOwnerAsset() override { return this; }
+
 	UObject* FindProxyObject(struct FChooserEvaluationContext& Context) const;
 	FObjectChooserBase::EIteratorStatus FindProxyObjectMulti(FChooserEvaluationContext &Context, FObjectChooserBase::FObjectChooserIteratorCallback Callback) const;
 
