Chooser Table Editor fix to apply correct asset filtering in details panel

diff --git a/Engine/Plugins/Experimental/Chooser/Source/Chooser/Internal/Chooser.h b/Engine/Plugins/Experimental/Chooser/Source/Chooser/Internal/Chooser.h
--- a/Engine/Plugins/Experimental/Chooser/Source/Chooser/Internal/Chooser.h
+++ b/Engine/Plugins/Experimental/Chooser/Source/Chooser/Internal/Chooser.h
@@ -1,6 +1,7 @@
 // Copyright Epic Games, Inc. All Rights Reserved.
 #pragma once
 
+#include "AssetRegistry/AssetData.h"
 #include "CoreMinimal.h"
 #include "UObject/Object.h"
 #include "IObjectChooser.h"
@@ -168,6 +169,9 @@ public:
 	UPROPERTY()
 	uint32 Version = 0;
 #endif
+
+	UFUNCTION()
+	bool ResultAssetFilter(const FAssetData& AssetData);
 	
 	UPROPERTY()
 	TArray<FInstancedStruct> CookedResults;
diff --git a/Engine/Plugins/Experimental/Chooser/Source/Chooser/Internal/ObjectChooser_Asset.h b/Engine/Plugins/Experimental/Chooser/Source/Chooser/Internal/ObjectChooser_Asset.h
--- a/Engine/Plugins/Experimental/Chooser/Source/Chooser/Internal/ObjectChooser_Asset.h
+++ b/Engine/Plugins/Experimental/Chooser/Source/Chooser/Internal/ObjectChooser_Asset.h
@@ -19,7 +19,7 @@ struct CHOOSER_API FAssetChooser : public FObjectChooserBase
 	virtual UObject* GetReferencedObject() const override { return Asset; }
 #endif
 	
-	UPROPERTY(EditAnywhere, Category = "Parameters")
+	UPROPERTY(EditAnywhere, Category = "Parameters", meta = (GetAssetFilter="ResultAssetFilter"))
 	TObjectPtr<UObject> Asset;
 };
 
diff --git a/Engine/Plugins/Experimental/Chooser/Source/Chooser/Private/Chooser.cpp b/Engine/Plugins/Experimental/Chooser/Source/Chooser/Private/Chooser.cpp
--- a/Engine/Plugins/Experimental/Chooser/Source/Chooser/Private/Chooser.cpp
+++ b/Engine/Plugins/Experimental/Chooser/Source/Chooser/Private/Chooser.cpp
@@ -388,6 +388,11 @@ void UChooserTable::UpdateDebugging(FChooserEvaluationContext& Context) const
 
 #endif
 
+bool UChooserTable::ResultAssetFilter(const FAssetData& AssetData)
+{
+	return !AssetData.IsInstanceOf(OutputObjectType);
+}
+
 FObjectChooserBase::EIteratorStatus UChooserTable::EvaluateChooser(FChooserEvaluationContext& Context, const UChooserTable* Chooser, FObjectChooserBase::FObjectChooserIteratorCallback Callback)
 {
 	TRACE_CPUPROFILER_EVENT_SCOPE(EvaluateChooser);
diff --git a/Engine/Plugins/Experimental/Chooser/Source/ChooserEditor/Private/ChooserDetails.cpp b/Engine/Plugins/Experimental/Chooser/Source/ChooserEditor/Private/ChooserDetails.cpp
--- a/Engine/Plugins/Experimental/Chooser/Source/ChooserEditor/Private/ChooserDetails.cpp
+++ b/Engine/Plugins/Experimental/Chooser/Source/ChooserEditor/Private/ChooserDetails.cpp
@@ -6,6 +6,7 @@
 
 #include "Chooser.h"
 #include "ChooserEditorWidgets.h"
+#include "ChooserPropertyAccess.h"
 #include "DetailCategoryBuilder.h"
 #include "DetailLayoutBuilder.h"
 #include "IDetailsView.h"
@@ -15,6 +16,15 @@
 
 #define LOCTEXT_NAMESPACE "ChooserDetails"
 
+bool UChooserRowDetails::ResultAssetFilter(const FAssetData& AssetData)
+{
+	if (Chooser)
+	{
+		return Chooser->ResultAssetFilter(AssetData);
+	}
+	return true;
+}
+
 namespace UE::ChooserEditor
 {
 	
diff --git a/Engine/Plugins/Experimental/Chooser/Source/ChooserEditor/Private/ChooserDetails.h b/Engine/Plugins/Experimental/Chooser/Source/ChooserEditor/Private/ChooserDetails.h
--- a/Engine/Plugins/Experimental/Chooser/Source/ChooserEditor/Private/ChooserDetails.h
+++ b/Engine/Plugins/Experimental/Chooser/Source/ChooserEditor/Private/ChooserDetails.h
@@ -15,6 +15,8 @@ class UChooserRowDetails : public UObject
 {
 	GENERATED_BODY()
 public:
+	UFUNCTION()
+	bool ResultAssetFilter(const FAssetData& AssetData); 
 
 	UPROPERTY(EditAnywhere, Category="Properties", meta=(FixedLayout, ShowOnlyInnerProperties))
     FInstancedPropertyBag Properties;
