Silence PVS false positive
#rnx

diff --git a/Engine/Plugins/Experimental/Chooser/Source/ProxyTable/Private/ProxyTableFunctionLibrary.cpp b/Engine/Plugins/Experimental/Chooser/Source/ProxyTable/Private/ProxyTableFunctionLibrary.cpp
--- a/Engine/Plugins/Experimental/Chooser/Source/ProxyTable/Private/ProxyTableFunctionLibrary.cpp
+++ b/Engine/Plugins/Experimental/Chooser/Source/ProxyTable/Private/ProxyTableFunctionLibrary.cpp
@@ -16,7 +16,7 @@ UObject* UProxyTableFunctionLibrary::EvaluateProxyAsset(const UObject* ContextOb
 		FChooserEvaluationContext Context(const_cast<UObject*>(ContextObject));
 		
 		Result = Proxy->FindProxyObject(Context);
-		if (ObjectClass && Result && !Result->IsA(ObjectClass))
+		if (ObjectClass && Result && !Result->IsA(ObjectClass)) //-V522
 		{
 			return nullptr;
 		}
