Special copy/paste handling for nested choosers

diff --git a/Engine/Plugins/Experimental/Chooser/Source/ChooserEditor/Private/ChooserTableEditor.cpp b/Engine/Plugins/Experimental/Chooser/Source/ChooserEditor/Private/ChooserTableEditor.cpp
--- a/Engine/Plugins/Experimental/Chooser/Source/ChooserEditor/Private/ChooserTableEditor.cpp
+++ b/Engine/Plugins/Experimental/Chooser/Source/ChooserEditor/Private/ChooserTableEditor.cpp
@@ -1905,6 +1905,34 @@ void FChooserTableEditor::MoveColumnRight()
 		SelectColumn(GetChooser(), MoveColumn(SelectedColumn->Column, SelectedColumn->Column + 2));
 	}
 }
+
+UChooserTable* DuplicateNestedChooser(UChooserTable* Chooser, UChooserTable* NewOuter)
+{
+	if (TObjectPtr<UChooserTable>* FoundTable = NewOuter->NestedChoosers.FindByPredicate([Chooser](UChooserTable* Table)
+		{
+			return Table->GetName() == Chooser->GetName();
+		}))
+	{
+		// we already duplicated this table
+		return *FoundTable;
+	}
+	
+	UChooserTable* NewTable = NewObject<UChooserTable>(NewOuter, Chooser->GetFName());
+	NewTable->ColumnsStructs = Chooser->ColumnsStructs;
+	NewTable->ResultsStructs = Chooser->ResultsStructs;
+	NewTable->RootChooser = NewOuter;
+	NewOuter->NestedChoosers.Add(NewTable);
+
+	for (FInstancedStruct& ResultData : NewTable->ResultsStructs)
+	{
+		if (FNestedChooser* NestedChooser = ResultData.GetMutablePtr<FNestedChooser>())
+		{
+			NestedChooser->Chooser = DuplicateNestedChooser(NestedChooser->Chooser, NewOuter);
+		}
+	}
+
+	return NewTable;
+}
 	
 UChooserTable* FChooserTableEditor::CopySelectionInternal()
 {
@@ -1924,7 +1952,6 @@ UChooserTable* FChooserTableEditor::CopySelectionInternal()
 	}
 	else if (CurrentSelectionType == ESelectionType::Rows)
 	{
-
 		TArray<UChooserRowDetails*> SelectedRowsCopy = SelectedRows;
 		Algo::Sort(SelectedRowsCopy, [](UChooserRowDetails *A, UChooserRowDetails *B) { return A->Row<B->Row;});
 		
@@ -1945,7 +1972,14 @@ UChooserTable* FChooserTableEditor::CopySelectionInternal()
 		
 		for (int RowIndex = 0; RowIndex < SelectedRowsCopy.Num(); RowIndex++)
 		{
-			CopyData->ResultsStructs[RowIndex] = Chooser->ResultsStructs[SelectedRowsCopy[RowIndex]->Row];
+			const FInstancedStruct& Result = Chooser->ResultsStructs[SelectedRowsCopy[RowIndex]->Row];
+
+			CopyData->ResultsStructs[RowIndex] = Result;
+			if (FNestedChooser* CopiedNestedChooser = CopyData->ResultsStructs[RowIndex].GetMutablePtr<FNestedChooser>())
+			{
+				// if the result for this row was a nested chooser, duplicate it
+				CopiedNestedChooser->Chooser = DuplicateNestedChooser(CopiedNestedChooser->Chooser, CopyData);
+			}
 
 			for (int ColumnIndex = 0; ColumnIndex < CopyData->ColumnsStructs.Num(); ColumnIndex++)
 			{
@@ -2072,6 +2106,18 @@ void FChooserTableEditor::PasteInternal(UChooserTable* PastedContent, int PasteR
 
 		Chooser->ResultsStructs.Insert(PastedContent->ResultsStructs, InsertIndex);
 		
+		if (!PastedContent->NestedChoosers.IsEmpty())
+		{
+			// if there were nested choosers in the copy buffer we have to remap or paste them here
+			for (int ResultIndex = InsertIndex; ResultIndex < PastedContent->ResultsStructs.Num() + InsertIndex; ResultIndex++)
+			{
+				if (FNestedChooser* NestedChooser = Chooser->ResultsStructs[ResultIndex].GetMutablePtr<FNestedChooser>())
+				{
+					NestedChooser->Chooser = DuplicateNestedChooser(NestedChooser->Chooser, Chooser);
+				}
+			}
+		}
+		
 		// Make sure each column has the same number of row datas as there are results
 		for(FInstancedStruct& ColumnData : Chooser->ColumnsStructs)
 		{
