Fix chooser Find/Replace dialog to handle unindexed / out of date assets

diff --git a/Engine/Plugins/Experimental/Chooser/Source/ChooserEditor/Private/ChooserFindProperties.cpp b/Engine/Plugins/Experimental/Chooser/Source/ChooserEditor/Private/ChooserFindProperties.cpp
--- a/Engine/Plugins/Experimental/Chooser/Source/ChooserEditor/Private/ChooserFindProperties.cpp
+++ b/Engine/Plugins/Experimental/Chooser/Source/ChooserEditor/Private/ChooserFindProperties.cpp
@@ -90,7 +90,7 @@ bool UChooserFindProperties::ShouldFilterOutAsset(const FAssetData& InAssetData,
 		}
 	}
 	
-	bOutIsOldAsset = false;
+	bOutIsOldAsset = true;
 	return true;
 }
 
