Replace Eye Icon for Chooser Table Columns with one from EditorStyle

diff --git a/Engine/Plugins/Experimental/Chooser/Source/ChooserEditor/Private/ChooserColumnHeader.cpp b/Engine/Plugins/Experimental/Chooser/Source/ChooserEditor/Private/ChooserColumnHeader.cpp
--- a/Engine/Plugins/Experimental/Chooser/Source/ChooserEditor/Private/ChooserColumnHeader.cpp
+++ b/Engine/Plugins/Experimental/Chooser/Source/ChooserEditor/Private/ChooserColumnHeader.cpp
@@ -62,12 +62,12 @@ namespace UE::ChooserEditor
 						SNew(SOverlay)
 							+ SOverlay::Slot()
 							[
-								SNew(SImage).Image(FChooserEditorStyle::Get().GetBrush("ChooserEditor.EyeOpened"))
+								SNew(SImage).Image(FAppStyle::Get().GetBrush("Icons.Visible"))
 								.Visibility_Lambda([Column]() { return Column->bDisabled ? EVisibility::Hidden : EVisibility::Visible; })
 							]
 							+ SOverlay::Slot()
 							[
-								SNew(SImage).Image(FChooserEditorStyle::Get().GetBrush("ChooserEditor.EyeClosed"))
+								SNew(SImage).Image(FAppStyle::Get().GetBrush("Icons.Hidden"))
 								.Visibility_Lambda([Column]() { return Column->bDisabled ? EVisibility::Visible : EVisibility::Hidden; })
 							]
 					]
diff --git a/Engine/Plugins/Experimental/Chooser/Source/ChooserEditor/Private/ChooserEditorStyle.cpp b/Engine/Plugins/Experimental/Chooser/Source/ChooserEditor/Private/ChooserEditorStyle.cpp
--- a/Engine/Plugins/Experimental/Chooser/Source/ChooserEditor/Private/ChooserEditorStyle.cpp
+++ b/Engine/Plugins/Experimental/Chooser/Source/ChooserEditor/Private/ChooserEditorStyle.cpp
@@ -28,11 +28,6 @@ FChooserEditorStyle::FChooserEditorStyle() :
 	Set("ChooserEditor.ChooserTableIconSmall", new IMAGE_BRUSH_SVG("ChooserIcon_16", Icon16x16));
 	
 	Set("ChooserEditor.FallbackIcon", new IMAGE_BRUSH_SVG("Fallback", Icon16x16));
-
-	// visibility bottons
-	Set("ChooserEditor.EyeClosed", new IMAGE_BRUSH(TEXT("icon_EyeClosed_40x"), Icon16x16));
-	Set("ChooserEditor.EyeOpened", new IMAGE_BRUSH(TEXT("icon_EyeOpened_40x"), Icon16x16));
-	
 }
 
 void FChooserEditorStyle::Initialize()
