Fix crash in Chooser Table Editor when editing in details panel with a struct output column that is not bound to anything

diff --git a/Engine/Plugins/Experimental/Chooser/Source/Chooser/Private/OutputStructColumn.cpp b/Engine/Plugins/Experimental/Chooser/Source/Chooser/Private/OutputStructColumn.cpp
--- a/Engine/Plugins/Experimental/Chooser/Source/Chooser/Private/OutputStructColumn.cpp
+++ b/Engine/Plugins/Experimental/Chooser/Source/Chooser/Private/OutputStructColumn.cpp
@@ -57,23 +57,39 @@ void FOutputStructColumn::StructTypeChanged()
 
 void FOutputStructColumn::AddToDetails(FInstancedPropertyBag& PropertyBag, int32 ColumnIndex, int32 RowIndex)
 {
-	FText DisplayName;
-	InputValue.Get<FChooserParameterStructBase>().GetDisplayName(DisplayName);
-	FName PropertyName("RowData",ColumnIndex);
-	FPropertyBagPropertyDesc PropertyDesc(PropertyName,  EPropertyBagPropertyType::Struct, RowValues[RowIndex].GetScriptStruct());
-	PropertyDesc.MetaData.Add(FPropertyBagPropertyDescMetaData("DisplayName", DisplayName.ToString()));
-	PropertyBag.AddProperties({PropertyDesc});
-	PropertyBag.SetValueStruct(PropertyName, FConstStructView(RowValues[RowIndex].GetScriptStruct(), RowValues[RowIndex].GetMemory()));
+	if (InputValue.IsValid())
+	{
+		const FChooserParameterStructBase& StructInput = InputValue.Get<FChooserParameterStructBase>();
+
+		// if there is no struct bound, the row values have no type, and can't be added to the details panel, so just skip them
+		if (StructInput.GetStructType())
+		{
+			FText DisplayName;
+			StructInput.GetDisplayName(DisplayName);
+			FName PropertyName("RowData",ColumnIndex);
+			FPropertyBagPropertyDesc PropertyDesc(PropertyName,  EPropertyBagPropertyType::Struct, RowValues[RowIndex].GetScriptStruct());
+			PropertyDesc.MetaData.Add(FPropertyBagPropertyDescMetaData("DisplayName", DisplayName.ToString()));
+			PropertyBag.AddProperties({PropertyDesc});
+			PropertyBag.SetValueStruct(PropertyName, FConstStructView(RowValues[RowIndex].GetScriptStruct(), RowValues[RowIndex].GetMemory()));
+		}
+	}
 }
 
 void FOutputStructColumn::SetFromDetails(FInstancedPropertyBag& PropertyBag, int32 ColumnIndex, int32 RowIndex)
 {
-	FName PropertyName("RowData", ColumnIndex);
-	
-	TValueOrError<FStructView, EPropertyBagResult> Result = PropertyBag.GetValueStruct(PropertyName, RowValues[RowIndex].GetScriptStruct());
-	if (FStructView* StructView = Result.TryGetValue())
+	if (InputValue.IsValid())
 	{
-		RowValues[RowIndex].GetScriptStruct()->CopyScriptStruct(RowValues[RowIndex].GetMutableMemory(), StructView->GetMemory());
+		const FChooserParameterStructBase& StructInput = InputValue.Get<FChooserParameterStructBase>();
+		if (StructInput.GetStructType())
+		{
+			FName PropertyName("RowData", ColumnIndex);
+
+			TValueOrError<FStructView, EPropertyBagResult> Result = PropertyBag.GetValueStruct(PropertyName, RowValues[RowIndex].GetScriptStruct());
+			if (FStructView* StructView = Result.TryGetValue())
+			{
+				RowValues[RowIndex].GetScriptStruct()->CopyScriptStruct(RowValues[RowIndex].GetMutableMemory(), StructView->GetMemory());
+			}
+		}
 	}
 }
 
