ChooserTable asset creation workflow improvements

diff --git a/Engine/Plugins/Experimental/Chooser/Source/ChooserEditor/ChooserEditor.Build.cs b/Engine/Plugins/Experimental/Chooser/Source/ChooserEditor/ChooserEditor.Build.cs
--- a/Engine/Plugins/Experimental/Chooser/Source/ChooserEditor/ChooserEditor.Build.cs
+++ b/Engine/Plugins/Experimental/Chooser/Source/ChooserEditor/ChooserEditor.Build.cs
@@ -39,6 +39,7 @@ namespace UnrealBuildTool.Rules
 					"RewindDebuggerInterface",
 					"GameplayInsights",
 					"Persona",
+					"DeveloperSettings",
 					// ... add private dependencies that you statically link with here ...
 				}
 			);
diff --git a/Engine/Plugins/Experimental/Chooser/Source/ChooserEditor/Private/AssetDefinition_ChooserTable.h b/Engine/Plugins/Experimental/Chooser/Source/ChooserEditor/Private/AssetDefinition_ChooserTable.h
--- a/Engine/Plugins/Experimental/Chooser/Source/ChooserEditor/Private/AssetDefinition_ChooserTable.h
+++ b/Engine/Plugins/Experimental/Chooser/Source/ChooserEditor/Private/AssetDefinition_ChooserTable.h
@@ -21,7 +21,7 @@ public:
 	virtual TSoftClassPtr<UObject> GetAssetClass() const override { return UChooserTable::StaticClass(); }
 	virtual TConstArrayView<FAssetCategoryPath> GetAssetCategories() const override
 	{
-		static const auto Categories = { FAssetCategoryPath(EAssetCategoryPaths::Misc) };
+		static const auto Categories = { FAssetCategoryPath(EAssetCategoryPaths::Misc) , FAssetCategoryPath(EAssetCategoryPaths::Animation) };
 		return Categories;
 	}
 	virtual EAssetCommandResult OpenAssets(const FAssetOpenArgs& OpenArgs) const override;
diff --git a/Engine/Plugins/Experimental/Chooser/Source/ChooserEditor/Private/ChooserEditorSettings.cpp b/Engine/Plugins/Experimental/Chooser/Source/ChooserEditor/Private/ChooserEditorSettings.cpp
new file mode 100644
--- /dev/null
+++ b/Engine/Plugins/Experimental/Chooser/Source/ChooserEditor/Private/ChooserEditorSettings.cpp
@@ -0,0 +1,42 @@
+﻿// Copyright Epic Games, Inc. All Rights Reserved.
+
+#include "ChooserEditorSettings.h"
+#include "Misc/CoreDelegates.h"
+
+#include UE_INLINE_GENERATED_CPP_BY_NAME(ChooserEditorSettings)
+
+UChooserEditorSettings::UChooserEditorSettings()
+{
+	FCoreDelegates::OnPreExit.AddLambda([]()
+	{
+		Get().SaveConfig();
+	});
+}
+
+#if WITH_EDITOR
+
+FText UChooserEditorSettings::GetSectionText() const
+{
+	return NSLOCTEXT("Chooser", "ChooserEditorSettingsName", "Chooser Editor");
+}
+
+FText UChooserEditorSettings::GetSectionDescription() const
+{
+	return NSLOCTEXT("Chooser", "ChooserEditorSettingsDesc", "Configure options for the Chooser Plugin.");
+}
+
+#endif
+
+FName UChooserEditorSettings::GetCategoryName() const
+{
+	return FName(TEXT("Plugins"));
+}
+
+UChooserEditorSettings& UChooserEditorSettings::Get()
+{
+	UChooserEditorSettings* MutableCDO = GetMutableDefault<UChooserEditorSettings>();
+	check(MutableCDO != nullptr)
+	
+	return *MutableCDO;
+}
+
diff --git a/Engine/Plugins/Experimental/Chooser/Source/ChooserEditor/Private/ChooserEditorSettings.h b/Engine/Plugins/Experimental/Chooser/Source/ChooserEditor/Private/ChooserEditorSettings.h
new file mode 100644
--- /dev/null
+++ b/Engine/Plugins/Experimental/Chooser/Source/ChooserEditor/Private/ChooserEditorSettings.h
@@ -0,0 +1,32 @@
+﻿// Copyright Epic Games, Inc. All Rights Reserved.
+
+#pragma once
+
+#include "CoreMinimal.h"
+#include "Engine/DeveloperSettings.h"
+#include "Logging/LogVerbosity.h"
+
+#include "ChooserEditorSettings.generated.h"
+
+UCLASS(config = EditorPerProjectUserSettings, meta=(DisplayName="Chooser Editor"))
+class UChooserEditorSettings : public UDeveloperSettings
+{
+	GENERATED_BODY()
+	
+public:
+	UChooserEditorSettings();
+
+#if WITH_EDITOR
+	virtual FText GetSectionText() const override;
+	virtual FText GetSectionDescription() const override;
+#endif
+
+	virtual FName GetCategoryName() const override;
+	
+	/** Most Recently Used Chooser Initializer, will be used as the default on following chooser table creates */
+	UPROPERTY(Config)
+	FString DefaultCreateType = "";
+
+	/** Get Mutable CDO of UChooserEditorSettings */
+	static UChooserEditorSettings& Get();
+};
diff --git a/Engine/Plugins/Experimental/Chooser/Source/ChooserEditor/Private/ChooserFactory.cpp b/Engine/Plugins/Experimental/Chooser/Source/ChooserEditor/Private/ChooserFactory.cpp
--- a/Engine/Plugins/Experimental/Chooser/Source/ChooserEditor/Private/ChooserFactory.cpp
+++ b/Engine/Plugins/Experimental/Chooser/Source/ChooserEditor/Private/ChooserFactory.cpp
@@ -3,6 +3,7 @@
 #include "ChooserFactory.h"
 #include "Chooser.h"
 #include "ChooserInitializer.h"
+#include "ChooserEditorSettings.h"
 #include "Widgets/Input/SButton.h"
 #include "Widgets/Layout/SUniformGridPanel.h"
 #include "Widgets/SCompoundWidget.h"
@@ -89,6 +90,15 @@ public:
 	bool ConfigureProperties(TWeakObjectPtr<UChooserTableFactory> InChooserFactory)
 	{
 		ChooserFactory = InChooserFactory;
+		if (!UChooserEditorSettings::Get().DefaultCreateType.IsEmpty())
+		{
+			FTopLevelAssetPath StructPath(UChooserEditorSettings::Get().DefaultCreateType);
+			if (UScriptStruct* DefaultInitializer = FindObject<UScriptStruct>(StructPath))
+			{
+				ChooserFactory->ChooserInitializer.InitializeAs(DefaultInitializer);
+			}
+		}
+		
 		DetailsView->SetObject(ChooserFactory.Get());
 
 		Window = SNew(SWindow)
@@ -172,6 +182,7 @@ UObject* UChooserTableFactory::FactoryCreateNew(UClass* Class, UObject* InParent
 	if (ChooserInitializer.IsValid())
 	{
 		ChooserInitializer.Get<FChooserInitializer>().Initialize(NewChooser);
+		UChooserEditorSettings::Get().DefaultCreateType = ChooserInitializer.GetScriptStruct()->GetStructPathName().ToString();
 	}
 	return NewChooser;
 }
diff --git a/Engine/Plugins/Experimental/Chooser/Source/ChooserEditor/Private/ChooserFactory.h b/Engine/Plugins/Experimental/Chooser/Source/ChooserEditor/Private/ChooserFactory.h
--- a/Engine/Plugins/Experimental/Chooser/Source/ChooserEditor/Private/ChooserFactory.h
+++ b/Engine/Plugins/Experimental/Chooser/Source/ChooserEditor/Private/ChooserFactory.h
@@ -13,6 +13,7 @@ class UChooserTableFactory : public UFactory
 {
 	GENERATED_BODY()
 
+public:
 	UChooserTableFactory();
 	
 	virtual bool ConfigureProperties() override;
@@ -23,7 +24,6 @@ class UChooserTableFactory : public UFactory
 	}
 	virtual FName GetNewAssetThumbnailOverride() const override { return "ChooserEditor.ChooserTableIconLarge"; }
 
-private:
 	UPROPERTY(EditAnywhere, DisplayName = "Chooser Type", NoClear, Category="Chooser Setup", meta = (ExcludeBaseStruct, BaseStruct = "/Script/ChooserEditor.ChooserInitializer"))
 	FInstancedStruct ChooserInitializer;
 };
\ No newline at end of file
