Fix type filtering for Result property on Chooser table row selection details

diff --git a/Engine/Plugins/Experimental/Chooser/Source/ChooserEditor/Private/ChooserDetails.h b/Engine/Plugins/Experimental/Chooser/Source/ChooserEditor/Private/ChooserDetails.h
--- a/Engine/Plugins/Experimental/Chooser/Source/ChooserEditor/Private/ChooserDetails.h
+++ b/Engine/Plugins/Experimental/Chooser/Source/ChooserEditor/Private/ChooserDetails.h
@@ -3,6 +3,7 @@
 #pragma once
 
 #include "CoreMinimal.h"
+#include "AssetRegistry/AssetData.h"
 #include "IDetailCustomization.h"
 #include "StructUtils/PropertyBag.h"
 #include "ChooserDetails.generated.h"
diff --git a/Engine/Plugins/Experimental/Chooser/Source/ChooserEditor/Private/ChooserTableEditor.cpp b/Engine/Plugins/Experimental/Chooser/Source/ChooserEditor/Private/ChooserTableEditor.cpp
--- a/Engine/Plugins/Experimental/Chooser/Source/ChooserEditor/Private/ChooserTableEditor.cpp
+++ b/Engine/Plugins/Experimental/Chooser/Source/ChooserEditor/Private/ChooserTableEditor.cpp
@@ -611,6 +611,8 @@ void FChooserTableEditor::OnObjectsTransacted(UObject* Object, const FTransactio
 			
 			TValueOrError<bool, EPropertyBagResult> DisabledResult = RowDetails->Properties.GetValueBool("Disabled");
 			RowDetails->Chooser->DisabledRows[RowDetails->Row] = DisabledResult.GetValue();
+
+			RefreshAll();
 		}
 	}
 }
@@ -1162,6 +1164,10 @@ void FChooserTableEditor::RefreshRowSelectionDetails()
 	}
 	SelectedRows.SetNum(0);
 	UChooserTable* Chooser = GetChooser();
+
+	FPropertyBagPropertyDesc ResultPropertyDesc ("Result", EPropertyBagPropertyType::Struct, FInstancedStruct::StaticStruct());
+	ResultPropertyDesc.MetaData.Add({"ExcludeBaseStruct",""});
+	ResultPropertyDesc.MetaData.Add({"BaseStruct","/Script/Chooser.ObjectChooserBase"});
 	
 	// Get the list of objects to edit the details of
 	TArray<TSharedPtr<FChooserTableRow>> SelectedItems = TableView->GetSelectedItems();
@@ -1175,7 +1181,7 @@ void FChooserTableEditor::RefreshRowSelectionDetails()
 			Selection->SetFlags(RF_Standalone | RF_Transactional);
 
 			FInstancedStruct& Result = Chooser->ResultsStructs[SelectedItem->RowIndex];
-			Selection->Properties.AddProperty("Result", EPropertyBagPropertyType::Struct, FInstancedStruct::StaticStruct());
+			Selection->Properties.AddProperties({ResultPropertyDesc});
 			Selection->Properties.SetValueStruct("Result", FConstStructView(FInstancedStruct::StaticStruct(), reinterpret_cast<uint8*>(&Result)));
 
 			int ColumnIndex = 0;
