Nicer add row/column workflow in the Chooser Table Editor.
Hiding the fallback row unless a fallback is set, and selection support for the fallback row, with copy/paste/delete and details panel support.
Fixed a bug where row disable states were not copied and pasted

diff --git a/Engine/Plugins/Experimental/Chooser/Source/Chooser/Internal/BoolColumn.h b/Engine/Plugins/Experimental/Chooser/Source/Chooser/Internal/BoolColumn.h
--- a/Engine/Plugins/Experimental/Chooser/Source/Chooser/Internal/BoolColumn.h
+++ b/Engine/Plugins/Experimental/Chooser/Source/Chooser/Internal/BoolColumn.h
@@ -46,7 +46,7 @@ public:
 	
 };
 
-USTRUCT()
+USTRUCT(DisplayName = "Bool", Meta = (Category = "Filter", Tooltip = "A column that filters rows based on the value of a Bool."))
 struct CHOOSER_API FBoolColumn : public FChooserColumnBase
 {
 	GENERATED_BODY()
diff --git a/Engine/Plugins/Experimental/Chooser/Source/Chooser/Internal/Chooser.h b/Engine/Plugins/Experimental/Chooser/Source/Chooser/Internal/Chooser.h
--- a/Engine/Plugins/Experimental/Chooser/Source/Chooser/Internal/Chooser.h
+++ b/Engine/Plugins/Experimental/Chooser/Source/Chooser/Internal/Chooser.h
@@ -197,7 +197,7 @@ public:
 	static FObjectChooserBase::EIteratorStatus IterateChooser(const UChooserTable* Chooser, FObjectChooserBase::FObjectChooserIteratorCallback Callback);
 };
 
-USTRUCT(BlueprintType, DisplayName = "Nested Chooser")
+USTRUCT(DisplayName = "Nested Chooser", Meta = (Category = "Chooser", Tooltip = "Reference another ChooserTable embedded in this asset, which will be evaluated at runtime if this row is selected."))
 struct CHOOSER_API FNestedChooser : public FObjectChooserBase
 {
 	GENERATED_BODY()
@@ -215,7 +215,7 @@ struct CHOOSER_API FNestedChooser : public FObjectChooserBase
 	TObjectPtr<UChooserTable> Chooser;
 };
 
-USTRUCT(BlueprintType, DisplayName = "Evaluate Chooser")
+USTRUCT(DisplayName = "Evaluate Chooser", Meta = (Category = "Chooser", Tooltip = "Reference another ChooserTable asset, which will be evaluated at runtime if this row is selected."))
 struct CHOOSER_API FEvaluateChooser : public FObjectChooserBase
 {
 	GENERATED_BODY()
diff --git a/Engine/Plugins/Experimental/Chooser/Source/Chooser/Internal/EnumColumn.h b/Engine/Plugins/Experimental/Chooser/Source/Chooser/Internal/EnumColumn.h
--- a/Engine/Plugins/Experimental/Chooser/Source/Chooser/Internal/EnumColumn.h
+++ b/Engine/Plugins/Experimental/Chooser/Source/Chooser/Internal/EnumColumn.h
@@ -92,7 +92,7 @@ struct CHOOSER_API FEnumColumnBase : public FChooserColumnBase
 	virtual void EnumChanged(const UEnum* Enum) {}
 };
 
-USTRUCT()
+USTRUCT(DisplayName = "Enum", Meta = (Category = "Filter", Tooltip = "A column that filters rows based on the value of an Enum, with Equal Not Equal, or Any as cell comparison options."))
 struct CHOOSER_API FEnumColumn : public FEnumColumnBase
 {
 	GENERATED_BODY()
diff --git a/Engine/Plugins/Experimental/Chooser/Source/Chooser/Internal/FloatDistanceColumn.h b/Engine/Plugins/Experimental/Chooser/Source/Chooser/Internal/FloatDistanceColumn.h
--- a/Engine/Plugins/Experimental/Chooser/Source/Chooser/Internal/FloatDistanceColumn.h
+++ b/Engine/Plugins/Experimental/Chooser/Source/Chooser/Internal/FloatDistanceColumn.h
@@ -35,7 +35,7 @@ public:
 	virtual void NativeAutoPopulate(UObject* InObject, bool& OutSuccess, float& OutValue) { }
 };
 
-USTRUCT()
+USTRUCT(DisplayName = "Float Difference", Meta = (Category = "Scoring", Tooltip = "A column which scores rows based on their difference from an Input float."))
 struct CHOOSER_API FFloatDistanceColumn : public FChooserColumnBase
 {
 	GENERATED_BODY()
diff --git a/Engine/Plugins/Experimental/Chooser/Source/Chooser/Internal/FloatRangeColumn.h b/Engine/Plugins/Experimental/Chooser/Source/Chooser/Internal/FloatRangeColumn.h
--- a/Engine/Plugins/Experimental/Chooser/Source/Chooser/Internal/FloatRangeColumn.h
+++ b/Engine/Plugins/Experimental/Chooser/Source/Chooser/Internal/FloatRangeColumn.h
@@ -58,7 +58,7 @@ struct FChooserFloatRangeRowData
 };
 
 
-USTRUCT()
+USTRUCT(DisplayName = "Float Range", Meta = (Category = "Filter", Tooltip = "A column which filters rows if an input value is not within the range specified for the row."))
 struct CHOOSER_API FFloatRangeColumn : public FChooserColumnBase
 {
 	GENERATED_BODY()
diff --git a/Engine/Plugins/Experimental/Chooser/Source/Chooser/Internal/GameplayTagColumn.h b/Engine/Plugins/Experimental/Chooser/Source/Chooser/Internal/GameplayTagColumn.h
--- a/Engine/Plugins/Experimental/Chooser/Source/Chooser/Internal/GameplayTagColumn.h
+++ b/Engine/Plugins/Experimental/Chooser/Source/Chooser/Internal/GameplayTagColumn.h
@@ -47,7 +47,7 @@ enum class EGameplayTagMatchDirection : uint8
 	InputInRowValue,
 };
 
-USTRUCT()
+USTRUCT(DisplayName = "Gameplay Tag", Meta = (Category = "Filter", Tooltip = "A column which filters rows by comparing Gameplay Tags."))
 struct CHOOSER_API FGameplayTagColumn : public FChooserColumnBase
 {
 	GENERATED_BODY()
diff --git a/Engine/Plugins/Experimental/Chooser/Source/Chooser/Internal/MultiEnumColumn.h b/Engine/Plugins/Experimental/Chooser/Source/Chooser/Internal/MultiEnumColumn.h
--- a/Engine/Plugins/Experimental/Chooser/Source/Chooser/Internal/MultiEnumColumn.h
+++ b/Engine/Plugins/Experimental/Chooser/Source/Chooser/Internal/MultiEnumColumn.h
@@ -27,7 +27,7 @@ struct FChooserMultiEnumRowData
 };
 
 
-USTRUCT()
+USTRUCT(DisplayName = "Enum (Or)", Meta = (Category = "Filter", Tooltip = "A column which filters rows using an Enum variable, where rows pass if the enum is one of any of the checked values."))
 struct CHOOSER_API FMultiEnumColumn : public FChooserColumnBase
 {
 	GENERATED_BODY()
diff --git a/Engine/Plugins/Experimental/Chooser/Source/Chooser/Internal/ObjectChooser_Asset.h b/Engine/Plugins/Experimental/Chooser/Source/Chooser/Internal/ObjectChooser_Asset.h
--- a/Engine/Plugins/Experimental/Chooser/Source/Chooser/Internal/ObjectChooser_Asset.h
+++ b/Engine/Plugins/Experimental/Chooser/Source/Chooser/Internal/ObjectChooser_Asset.h
@@ -6,7 +6,7 @@
 #include "IObjectChooser.h"
 #include "ObjectChooser_Asset.generated.h"
 
-USTRUCT(DisplayName = "Asset")
+USTRUCT(DisplayName = "Asset", Meta = (ResultType = "Object", Category = "Basic", Tooltip = "A hard reference to a specific asset"))
 struct CHOOSER_API FAssetChooser : public FObjectChooserBase
 {
 	GENERATED_BODY()
@@ -23,7 +23,7 @@ struct CHOOSER_API FAssetChooser : public FObjectChooserBase
 	TObjectPtr<UObject> Asset;
 };
 
-USTRUCT(DisplayName = "Soft Asset")
+USTRUCT(DisplayName = "Asset (Soft Reference)", Meta = (ResultType = "Object", Category = "Basic", Tooltip = "A soft object reference to a specific asset\nAssets will need to be preloaded manually to avoid a hitch if they are selected."))
 struct CHOOSER_API FSoftAssetChooser : public FObjectChooserBase
 {
 	GENERATED_BODY()
diff --git a/Engine/Plugins/Experimental/Chooser/Source/Chooser/Internal/ObjectChooser_Class.h b/Engine/Plugins/Experimental/Chooser/Source/Chooser/Internal/ObjectChooser_Class.h
--- a/Engine/Plugins/Experimental/Chooser/Source/Chooser/Internal/ObjectChooser_Class.h
+++ b/Engine/Plugins/Experimental/Chooser/Source/Chooser/Internal/ObjectChooser_Class.h
@@ -6,7 +6,7 @@
 #include "IObjectChooser.h"
 #include "ObjectChooser_Class.generated.h"
 
-USTRUCT(DisplayName = "Class")
+USTRUCT(DisplayName = "Class", Meta = (ResultType = "Class", Category = "Basic", Tooltip = "A reference to a Class.\nOnly for use in Choosers with ResultType set to Sub Class Of"))
 struct CHOOSER_API FClassChooser : public FObjectChooserBase
 {
 	GENERATED_BODY()
diff --git a/Engine/Plugins/Experimental/Chooser/Source/Chooser/Internal/ObjectClassColumn.h b/Engine/Plugins/Experimental/Chooser/Source/Chooser/Internal/ObjectClassColumn.h
--- a/Engine/Plugins/Experimental/Chooser/Source/Chooser/Internal/ObjectClassColumn.h
+++ b/Engine/Plugins/Experimental/Chooser/Source/Chooser/Internal/ObjectClassColumn.h
@@ -34,7 +34,7 @@ struct FChooserObjectClassRowData
 	bool Evaluate(const UObject* LeftHandSide) const;
 };
 
-USTRUCT()
+USTRUCT(DisplayName = "Object Class", Meta = (Category = "Filter", Tooltip = "A column which filters rows using an Object reference variable, by checking if that object is of a certain Class."))
 struct CHOOSER_API FObjectClassColumn : public FChooserColumnBase
 {
 	GENERATED_BODY()
diff --git a/Engine/Plugins/Experimental/Chooser/Source/Chooser/Internal/ObjectColumn.h b/Engine/Plugins/Experimental/Chooser/Source/Chooser/Internal/ObjectColumn.h
--- a/Engine/Plugins/Experimental/Chooser/Source/Chooser/Internal/ObjectColumn.h
+++ b/Engine/Plugins/Experimental/Chooser/Source/Chooser/Internal/ObjectColumn.h
@@ -54,7 +54,7 @@ struct FChooserObjectRowData
 	bool Evaluate(const FSoftObjectPath& LeftHandSide) const;
 };
 
-USTRUCT()
+USTRUCT(DisplayName = "Object", Meta = (Category = "Filter", Tooltip = "A column which filters rows by an input Object to specified Objects for each row."))
 struct CHOOSER_API FObjectColumn : public FChooserColumnBase
 {
 	GENERATED_BODY()
diff --git a/Engine/Plugins/Experimental/Chooser/Source/Chooser/Internal/OutputBoolColumn.h b/Engine/Plugins/Experimental/Chooser/Source/Chooser/Internal/OutputBoolColumn.h
--- a/Engine/Plugins/Experimental/Chooser/Source/Chooser/Internal/OutputBoolColumn.h
+++ b/Engine/Plugins/Experimental/Chooser/Source/Chooser/Internal/OutputBoolColumn.h
@@ -8,7 +8,7 @@
 #include "StructUtils/InstancedStruct.h"
 #include "OutputBoolColumn.generated.h"
 
-USTRUCT()
+USTRUCT(DisplayName = "Output Bool", Meta = (Category = "Output", Tooltip = "A column which writes a Bool value."))
 struct CHOOSER_API FOutputBoolColumn : public FChooserColumnBase
 {
 	GENERATED_BODY()
@@ -26,7 +26,18 @@ struct CHOOSER_API FOutputBoolColumn : public FChooserColumnBase
 
 	virtual void AddToDetails(FInstancedPropertyBag& PropertyBag, int32 ColumnIndex, int32 RowIndex) override;
 	virtual void SetFromDetails(FInstancedPropertyBag& PropertyBag, int32 ColumnIndex, int32 RowIndex) override;
+	virtual void CopyFallback(FChooserColumnBase& SourceColumn) override { bFallbackValue = static_cast<FOutputBoolColumn&>(SourceColumn).bFallbackValue; }
 #endif
+		
+	bool& GetValueForIndex(int32 Index)
+	{
+		return Index == ChooserColumn_SpecialIndex_Fallback ? bFallbackValue : RowValues[Index];
+	}
+	
+	bool GetValueForIndex(int32 Index) const
+	{
+		return Index == ChooserColumn_SpecialIndex_Fallback ? bFallbackValue : RowValues[Index];
+	}
 	
 	// FallbackValue will be used as the output value if the all rows in the chooser fail, and the FallbackResult from the chooser is used.
 	UPROPERTY(EditAnywhere, Category=Data);
diff --git a/Engine/Plugins/Experimental/Chooser/Source/Chooser/Internal/OutputEnumColumn.h b/Engine/Plugins/Experimental/Chooser/Source/Chooser/Internal/OutputEnumColumn.h
--- a/Engine/Plugins/Experimental/Chooser/Source/Chooser/Internal/OutputEnumColumn.h
+++ b/Engine/Plugins/Experimental/Chooser/Source/Chooser/Internal/OutputEnumColumn.h
@@ -24,7 +24,7 @@ struct FChooserOutputEnumRowData
 };
 
 
-USTRUCT()
+USTRUCT(DisplayName = "Output Enum", Meta = (Category = "Output", Tooltip = "A column which writes an Enum value."))
 struct CHOOSER_API FOutputEnumColumn : public FEnumColumnBase
 {
 	GENERATED_BODY()
@@ -68,9 +68,18 @@ public:
 
 	virtual void AddToDetails (FInstancedPropertyBag& PropertyBag, int32 ColumnIndex, int32 RowIndex) override;
 	virtual void SetFromDetails(FInstancedPropertyBag& PropertyBag, int32 ColumnIndex, int32 RowIndex) override;
+	virtual void CopyFallback(FChooserColumnBase& SourceColumn) override { FallbackValue = static_cast<FOutputEnumColumn&>(SourceColumn).FallbackValue; }
 	
 	virtual void EnumChanged(const UEnum* Enum) override;
 #endif
+	FChooserOutputEnumRowData& GetValueForIndex(int32 Index)
+	{
+		return Index == ChooserColumn_SpecialIndex_Fallback ? FallbackValue : RowValues[Index];
+	}
+	const FChooserOutputEnumRowData& GetValueForIndex(int32 Index) const
+	{
+		return Index == ChooserColumn_SpecialIndex_Fallback ? FallbackValue : RowValues[Index];
+	}
 	
 	CHOOSER_COLUMN_BOILERPLATE(FChooserParameterEnumBase);
 };
\ No newline at end of file
diff --git a/Engine/Plugins/Experimental/Chooser/Source/Chooser/Internal/OutputFloatColumn.h b/Engine/Plugins/Experimental/Chooser/Source/Chooser/Internal/OutputFloatColumn.h
--- a/Engine/Plugins/Experimental/Chooser/Source/Chooser/Internal/OutputFloatColumn.h
+++ b/Engine/Plugins/Experimental/Chooser/Source/Chooser/Internal/OutputFloatColumn.h
@@ -8,7 +8,7 @@
 #include "StructUtils/InstancedStruct.h"
 #include "OutputFloatColumn.generated.h"
 
-USTRUCT()
+USTRUCT(DisplayName = "Output Float", Meta = (Category = "Output", Tooltip = "A column which writes a Float value."))
 struct CHOOSER_API FOutputFloatColumn : public FChooserColumnBase
 {
 	GENERATED_BODY()
@@ -26,7 +26,18 @@ struct CHOOSER_API FOutputFloatColumn : public FChooserColumnBase
 
 	virtual void AddToDetails(FInstancedPropertyBag& PropertyBag, int32 ColumnIndex, int32 RowIndex) override;
 	virtual void SetFromDetails(FInstancedPropertyBag& PropertyBag, int32 ColumnIndex, int32 RowIndex) override;
+	virtual void CopyFallback(FChooserColumnBase& SourceColumn) override { FallbackValue = static_cast<FOutputFloatColumn&>(SourceColumn).FallbackValue; }
 #endif
+	
+	double& GetValueForIndex(int32 Index)
+	{
+		return Index == ChooserColumn_SpecialIndex_Fallback ? FallbackValue : RowValues[Index];
+	}
+
+	double GetValueForIndex(int32 Index) const
+   	{
+   		return Index == ChooserColumn_SpecialIndex_Fallback ? FallbackValue : RowValues[Index];
+   	}
 
 	// FallbackValue will be used as the output value if the all rows in the chooser fail, and the FallbackResult from the chooser is used.
 	UPROPERTY(EditAnywhere, Category=Data);
diff --git a/Engine/Plugins/Experimental/Chooser/Source/Chooser/Internal/OutputObjectColumn.h b/Engine/Plugins/Experimental/Chooser/Source/Chooser/Internal/OutputObjectColumn.h
--- a/Engine/Plugins/Experimental/Chooser/Source/Chooser/Internal/OutputObjectColumn.h
+++ b/Engine/Plugins/Experimental/Chooser/Source/Chooser/Internal/OutputObjectColumn.h
@@ -11,7 +11,7 @@
 
 struct FBindingChainElement;
 
-USTRUCT()
+USTRUCT(DisplayName = "Output Object", Meta = (Category = "Output", Tooltip = "A column which writes an Object reference."))
 struct FChooserOutputObjectRowData
 {
 	GENERATED_BODY()
@@ -20,7 +20,7 @@ struct FChooserOutputObjectRowData
 	FInstancedStruct Value;
 };
 
-USTRUCT()
+USTRUCT(DisplayName = "Output Object", Meta = (Category = "Output", Tooltip = "A column which an Object Reference."))
 struct CHOOSER_API FOutputObjectColumn : public FChooserColumnBase
 {
 	GENERATED_BODY()
@@ -51,7 +51,18 @@ struct CHOOSER_API FOutputObjectColumn : public FChooserColumnBase
 #if WITH_EDITOR
 	virtual void AddToDetails(FInstancedPropertyBag& PropertyBag, int32 ColumnIndex, int32 RowIndex) override;
 	virtual void SetFromDetails(FInstancedPropertyBag& PropertyBag, int32 ColumnIndex, int32 RowIndex) override;
+	virtual void CopyFallback(FChooserColumnBase& SourceColumn) override { FallbackValue = static_cast<FOutputObjectColumn&>(SourceColumn).FallbackValue; }
 #endif
+	
+	FChooserOutputObjectRowData& GetValueForIndex(int32 Index)
+	{
+		return Index == ChooserColumn_SpecialIndex_Fallback ? FallbackValue : RowValues[Index];
+	}
+	
+	const FChooserOutputObjectRowData& GetValueForIndex(int32 Index) const
+	{
+		return Index == ChooserColumn_SpecialIndex_Fallback ? FallbackValue : RowValues[Index];
+	}
 
 	CHOOSER_COLUMN_BOILERPLATE(FChooserParameterObjectBase);
 
diff --git a/Engine/Plugins/Experimental/Chooser/Source/Chooser/Internal/OutputStructColumn.h b/Engine/Plugins/Experimental/Chooser/Source/Chooser/Internal/OutputStructColumn.h
--- a/Engine/Plugins/Experimental/Chooser/Source/Chooser/Internal/OutputStructColumn.h
+++ b/Engine/Plugins/Experimental/Chooser/Source/Chooser/Internal/OutputStructColumn.h
@@ -27,7 +27,7 @@ struct CHOOSER_API FStructContextProperty : public FChooserParameterStructBase
 #endif
 };
 
-USTRUCT()
+USTRUCT(DisplayName = "Output Struct", Meta = (Category = "Output", Tooltip = "A column which writes all elements of a Struct.\n The data for each row must be set in the details panel, and the table cells will display values which differ from the struct default."))
 struct CHOOSER_API FOutputStructColumn : public FChooserColumnBase
 {
 	GENERATED_BODY()
@@ -46,8 +46,19 @@ struct CHOOSER_API FOutputStructColumn : public FChooserColumnBase
 
 	virtual void AddToDetails(FInstancedPropertyBag& PropertyBag, int32 ColumnIndex, int32 RowIndex) override;
 	virtual void SetFromDetails(FInstancedPropertyBag& PropertyBag, int32 ColumnIndex, int32 RowIndex) override;
+	virtual void CopyFallback(FChooserColumnBase& SourceColumn) override { FallbackValue = static_cast<FOutputStructColumn&>(SourceColumn).FallbackValue; }
 #endif
 	
+	FInstancedStruct& GetValueForIndex(int32 Index)
+	{
+		return Index == ChooserColumn_SpecialIndex_Fallback ? FallbackValue : RowValues[Index];
+	}
+	
+	const FInstancedStruct& GetValueForIndex(int32 Index) const
+	{
+		return Index == ChooserColumn_SpecialIndex_Fallback ? FallbackValue : RowValues[Index];
+	}
+	
 	// FallbackValue will be used as the output value if the all rows in the chooser fail, and the FallbackResult from the chooser is used.
 	UPROPERTY(EditAnywhere, Meta = (StructTypeConst), Category=Data);
    	FInstancedStruct FallbackValue;
diff --git a/Engine/Plugins/Experimental/Chooser/Source/Chooser/Internal/RandomizeColumn.h b/Engine/Plugins/Experimental/Chooser/Source/Chooser/Internal/RandomizeColumn.h
--- a/Engine/Plugins/Experimental/Chooser/Source/Chooser/Internal/RandomizeColumn.h
+++ b/Engine/Plugins/Experimental/Chooser/Source/Chooser/Internal/RandomizeColumn.h
@@ -26,7 +26,7 @@ public:
 	CHOOSER_PARAMETER_BOILERPLATE();
 };
 
-USTRUCT()
+USTRUCT(DisplayName = "Randomize", Meta = (Category = "Random", Tooltip = "The Randomize column will randomly select between whatever values have passed all filters.\n The value specified in each cell is a probability weighting for the row.\n  A row with a value twice as likely as another will be twice as likely to be selected.\n Using the optional RandomizationContext binding, it can track the most recent selection, and reduce the probability of randomly picking the same entry twice"))
 struct CHOOSER_API FRandomizeColumn : public FChooserColumnBase
 {
 	GENERATED_BODY()
diff --git a/Engine/Plugins/Experimental/Chooser/Source/Chooser/Private/OutputBoolColumn.cpp b/Engine/Plugins/Experimental/Chooser/Source/Chooser/Private/OutputBoolColumn.cpp
--- a/Engine/Plugins/Experimental/Chooser/Source/Chooser/Private/OutputBoolColumn.cpp
+++ b/Engine/Plugins/Experimental/Chooser/Source/Chooser/Private/OutputBoolColumn.cpp
@@ -15,19 +15,13 @@ void FOutputBoolColumn::SetOutputs(FChooserEvaluationContext& Context, int RowIn
 {
 	if (InputValue.IsValid())
 	{
-		bool bOutputValue = bFallbackValue;
-		if (RowValues.IsValidIndex(RowIndex))
-		{
-			bOutputValue = RowValues[RowIndex];
-		}
-		
-		InputValue.Get<FChooserParameterBoolBase>().SetValue(Context, bOutputValue);
+		InputValue.Get<FChooserParameterBoolBase>().SetValue(Context, GetValueForIndex(RowIndex));
 	}
 	
 #if WITH_EDITOR
 	if (Context.DebuggingInfo.bCurrentDebugTarget)
 	{
-		TestValue = RowValues[RowIndex];
+		TestValue = GetValueForIndex(RowIndex);
 	}
 #endif
 }
@@ -41,7 +35,7 @@ void FOutputBoolColumn::SetOutputs(FChooserEvaluationContext& Context, int RowIn
 		FPropertyBagPropertyDesc PropertyDesc(PropertyName, EPropertyBagPropertyType::Bool);
 		PropertyDesc.MetaData.Add(FPropertyBagPropertyDescMetaData("DisplayName", DisplayName.ToString()));
 		PropertyBag.AddProperties({PropertyDesc});
-		PropertyBag.SetValueFloat(PropertyName, RowValues[RowIndex]);
+		PropertyBag.SetValueFloat(PropertyName, GetValueForIndex(RowIndex));
 	}
 
 	void FOutputBoolColumn::SetFromDetails(FInstancedPropertyBag& PropertyBag, int32 ColumnIndex, int32 RowIndex)
@@ -51,7 +45,7 @@ void FOutputBoolColumn::SetOutputs(FChooserEvaluationContext& Context, int RowIn
 		TValueOrError<bool, EPropertyBagResult> Result = PropertyBag.GetValueBool(PropertyName);
 		if (bool* Value = Result.TryGetValue())
 		{
-			RowValues[RowIndex] = *Value;
+			GetValueForIndex(RowIndex) = *Value;
 		}
 	} 
 
diff --git a/Engine/Plugins/Experimental/Chooser/Source/Chooser/Private/OutputEnumColumn.cpp b/Engine/Plugins/Experimental/Chooser/Source/Chooser/Private/OutputEnumColumn.cpp
--- a/Engine/Plugins/Experimental/Chooser/Source/Chooser/Private/OutputEnumColumn.cpp
+++ b/Engine/Plugins/Experimental/Chooser/Source/Chooser/Private/OutputEnumColumn.cpp
@@ -17,18 +17,13 @@ void FOutputEnumColumn::SetOutputs(FChooserEvaluationContext& Context, int RowIn
 {
 	if (InputValue.IsValid())
 	{
-		uint8 OutputValue = FallbackValue.Value;
-		if (RowValues.IsValidIndex(RowIndex))
-		{
-			OutputValue = RowValues[RowIndex].Value;
-		}
-		InputValue.Get<FChooserParameterEnumBase>().SetValue(Context, OutputValue);
+		InputValue.Get<FChooserParameterEnumBase>().SetValue(Context, GetValueForIndex(RowIndex).Value);
 	}
 	
 #if WITH_EDITOR
 	if (Context.DebuggingInfo.bCurrentDebugTarget)
 	{
-		TestValue = RowValues[RowIndex].Value;
+		TestValue = GetValueForIndex(RowIndex).Value;
 	}
 #endif
 }
@@ -75,7 +70,7 @@ void FOutputEnumColumn::PostLoad()
 
 #if WITH_EDITOR
 
-	void FOutputEnumColumn::AddToDetails (FInstancedPropertyBag& PropertyBag, int32 ColumnIndex, int32 RowIndex)
+	void FOutputEnumColumn::AddToDetails(FInstancedPropertyBag& PropertyBag, int32 ColumnIndex, int32 RowIndex)
 	{
 		FText DisplayName;
 		InputValue.Get<FChooserParameterBase>().GetDisplayName(DisplayName);
@@ -86,7 +81,7 @@ void FOutputEnumColumn::PostLoad()
 		FPropertyBagPropertyDesc PropertyDesc(PropertyName,  EPropertyBagPropertyType::Enum, InputValue.Get<FChooserParameterEnumBase>().GetEnum());
 		PropertyDesc.MetaData.Add(FPropertyBagPropertyDescMetaData("DisplayName", DisplayName.ToString()));
 		PropertyBag.AddProperties({PropertyDesc});
-		PropertyBag.SetValueEnum("Value", RowValues[RowIndex].Value, InputValue.Get<FChooserParameterEnumBase>().GetEnum());
+		PropertyBag.SetValueEnum(PropertyName, GetValueForIndex(RowIndex).Value, InputValue.Get<FChooserParameterEnumBase>().GetEnum());
 	}
 	
 	void FOutputEnumColumn::SetFromDetails(FInstancedPropertyBag& PropertyBag, int32 ColumnIndex, int32 RowIndex)
@@ -96,7 +91,7 @@ void FOutputEnumColumn::PostLoad()
 		TValueOrError<uint8, EPropertyBagResult> Result = PropertyBag.GetValueEnum(PropertyName, InputValue.Get<FChooserParameterEnumBase>().GetEnum());
 		if (uint8* Value = Result.TryGetValue())
 		{
-			RowValues[RowIndex].Value = *Value;
+			GetValueForIndex(RowIndex).Value = *Value;
 		}
 	}
 
diff --git a/Engine/Plugins/Experimental/Chooser/Source/Chooser/Private/OutputFloatColumn .cpp b/Engine/Plugins/Experimental/Chooser/Source/Chooser/Private/OutputFloatColumn .cpp
--- a/Engine/Plugins/Experimental/Chooser/Source/Chooser/Private/OutputFloatColumn .cpp	
+++ b/Engine/Plugins/Experimental/Chooser/Source/Chooser/Private/OutputFloatColumn .cpp	
@@ -15,19 +15,13 @@ void FOutputFloatColumn::SetOutputs(FChooserEvaluationContext& Context, int RowI
 {
 	if (InputValue.IsValid())
 	{
-		double OutputValue = FallbackValue;
-		if (RowValues.IsValidIndex(RowIndex))
-		{
-			OutputValue = RowValues[RowIndex];
-		}
-	
-		InputValue.Get<FChooserParameterFloatBase>().SetValue(Context, OutputValue);
+		InputValue.Get<FChooserParameterFloatBase>().SetValue(Context, GetValueForIndex(RowIndex));
 	}
 	
 #if WITH_EDITOR
 	if (Context.DebuggingInfo.bCurrentDebugTarget)
 	{
-		TestValue = RowValues[RowIndex];
+		TestValue = GetValueForIndex(RowIndex);
 	}
 #endif
 }
@@ -40,7 +34,7 @@ void FOutputFloatColumn::SetOutputs(FChooserEvaluationContext& Context, int RowI
 		FPropertyBagPropertyDesc PropertyDesc(PropertyName, EPropertyBagPropertyType::Float);
 		PropertyDesc.MetaData.Add(FPropertyBagPropertyDescMetaData("DisplayName", DisplayName.ToString()));
 		PropertyBag.AddProperties({PropertyDesc});
-		PropertyBag.SetValueFloat(PropertyName, RowValues[RowIndex]);
+		PropertyBag.SetValueFloat(PropertyName, GetValueForIndex(RowIndex));
 	}
 
 	void FOutputFloatColumn::SetFromDetails(FInstancedPropertyBag& PropertyBag, int32 ColumnIndex, int32 RowIndex)
@@ -50,7 +44,7 @@ void FOutputFloatColumn::SetOutputs(FChooserEvaluationContext& Context, int RowI
 		TValueOrError<float, EPropertyBagResult> Result = PropertyBag.GetValueFloat(PropertyName);
 		if (float* Value = Result.TryGetValue())
 		{
-			RowValues[RowIndex] = *Value;
+			GetValueForIndex(RowIndex) = *Value;
 		}
 	}
 #endif
\ No newline at end of file
diff --git a/Engine/Plugins/Experimental/Chooser/Source/Chooser/Private/OutputObjectColumn.cpp b/Engine/Plugins/Experimental/Chooser/Source/Chooser/Private/OutputObjectColumn.cpp
--- a/Engine/Plugins/Experimental/Chooser/Source/Chooser/Private/OutputObjectColumn.cpp
+++ b/Engine/Plugins/Experimental/Chooser/Source/Chooser/Private/OutputObjectColumn.cpp
@@ -32,21 +32,10 @@ void FOutputObjectColumn::Compile(IHasContextClass* Owner, bool bForce)
 
 void FOutputObjectColumn::SetOutputs(FChooserEvaluationContext& Context, int RowIndex) const
 {
-	if (RowValues.IsValidIndex(RowIndex))
+	if (const FObjectChooserBase* Value = GetValueForIndex(RowIndex).Value.GetPtr<FObjectChooserBase>())
 	{
-		if (const FObjectChooserBase* Value = RowValues[RowIndex].Value.GetPtr<FObjectChooserBase>())
-		{
-			UObject* Result = Value->ChooseObject(Context);
-			InputValue.Get<FChooserParameterObjectBase>().SetValue(Context, Result);
-		}
-	}
-	else
-	{
-		if (const FObjectChooserBase* Value = FallbackValue.Value.GetPtr<FObjectChooserBase>())
-		{
-			UObject* Result = Value->ChooseObject(Context);
-			InputValue.Get<FChooserParameterObjectBase>().SetValue(Context, Result);
-		}
+		UObject* Result = Value->ChooseObject(Context);
+		InputValue.Get<FChooserParameterObjectBase>().SetValue(Context, Result);
 	}
 }
 
@@ -59,7 +48,8 @@ void FOutputObjectColumn::SetOutputs(FChooserEvaluationContext& Context, int Row
 		FPropertyBagPropertyDesc PropertyDesc(PropertyName,  EPropertyBagPropertyType::Struct, FChooserOutputObjectRowData::StaticStruct());
 		PropertyDesc.MetaData.Add(FPropertyBagPropertyDescMetaData("DisplayName", DisplayName.ToString()));
 		PropertyBag.AddProperties({PropertyDesc});
-		PropertyBag.SetValueStruct(PropertyName, RowValues[RowIndex]);
+	
+		PropertyBag.SetValueStruct(PropertyName, GetValueForIndex(RowIndex));
 	}
 
 	void FOutputObjectColumn::SetFromDetails(FInstancedPropertyBag& PropertyBag, int32 ColumnIndex, int32 RowIndex)
@@ -69,7 +59,7 @@ void FOutputObjectColumn::SetOutputs(FChooserEvaluationContext& Context, int Row
 		TValueOrError<FStructView, EPropertyBagResult> Result = PropertyBag.GetValueStruct(PropertyName, FChooserOutputObjectRowData::StaticStruct());
 		if (FStructView* StructView = Result.TryGetValue())
 		{
-			RowValues[RowIndex] = StructView->Get<FChooserOutputObjectRowData>();
+			GetValueForIndex(RowIndex) = StructView->Get<FChooserOutputObjectRowData>();
 		}
 	}
 
diff --git a/Engine/Plugins/Experimental/Chooser/Source/Chooser/Private/OutputStructColumn.cpp b/Engine/Plugins/Experimental/Chooser/Source/Chooser/Private/OutputStructColumn.cpp
--- a/Engine/Plugins/Experimental/Chooser/Source/Chooser/Private/OutputStructColumn.cpp
+++ b/Engine/Plugins/Experimental/Chooser/Source/Chooser/Private/OutputStructColumn.cpp
@@ -67,10 +67,12 @@ void FOutputStructColumn::AddToDetails(FInstancedPropertyBag& PropertyBag, int32
 			FText DisplayName;
 			StructInput.GetDisplayName(DisplayName);
 			FName PropertyName("RowData",ColumnIndex);
-			FPropertyBagPropertyDesc PropertyDesc(PropertyName,  EPropertyBagPropertyType::Struct, RowValues[RowIndex].GetScriptStruct());
+			
+			FInstancedStruct& Value = GetValueForIndex(RowIndex);
+			FPropertyBagPropertyDesc PropertyDesc(PropertyName,  EPropertyBagPropertyType::Struct, Value.GetScriptStruct());
 			PropertyDesc.MetaData.Add(FPropertyBagPropertyDescMetaData("DisplayName", DisplayName.ToString()));
 			PropertyBag.AddProperties({PropertyDesc});
-			PropertyBag.SetValueStruct(PropertyName, FConstStructView(RowValues[RowIndex].GetScriptStruct(), RowValues[RowIndex].GetMemory()));
+			PropertyBag.SetValueStruct(PropertyName, FConstStructView(Value.GetScriptStruct(), Value.GetMemory()));
 		}
 	}
 }
@@ -84,10 +86,12 @@ void FOutputStructColumn::SetFromDetails(FInstancedPropertyBag& PropertyBag, int
 		{
 			FName PropertyName("RowData", ColumnIndex);
 
-			TValueOrError<FStructView, EPropertyBagResult> Result = PropertyBag.GetValueStruct(PropertyName, RowValues[RowIndex].GetScriptStruct());
+			FInstancedStruct& Value = GetValueForIndex(RowIndex);
+
+			TValueOrError<FStructView, EPropertyBagResult> Result = PropertyBag.GetValueStruct(PropertyName, Value.GetScriptStruct());
 			if (FStructView* StructView = Result.TryGetValue())
 			{
-				RowValues[RowIndex].GetScriptStruct()->CopyScriptStruct(RowValues[RowIndex].GetMutableMemory(), StructView->GetMemory());
+				RowValues[RowIndex].GetScriptStruct()->CopyScriptStruct(Value.GetMutableMemory(), StructView->GetMemory());
 			}
 		}
 	}
@@ -104,22 +108,18 @@ void FOutputStructColumn::SetOutputs(FChooserEvaluationContext& Context, int Row
 {
 	if (InputValue.IsValid())
 	{
-		const FInstancedStruct* OutputValue = &FallbackValue;
-		if (RowValues.IsValidIndex(RowIndex))
-		{
-			OutputValue = &RowValues[RowIndex];
-		}
+		const FInstancedStruct& OutputValue = GetValueForIndex(RowIndex);
 
-		if (OutputValue && OutputValue->IsValid())
+		if (OutputValue.IsValid())
 		{
-			InputValue.Get<FChooserParameterStructBase>().SetValue(Context, *OutputValue);
+			InputValue.Get<FChooserParameterStructBase>().SetValue(Context, OutputValue);
 		}
 	}
 	
 #if WITH_EDITOR
 	if (Context.DebuggingInfo.bCurrentDebugTarget)
 	{
-		TestValue = RowValues[RowIndex];
+		TestValue = GetValueForIndex(RowIndex);
 	}
 #endif
 }
diff --git a/Engine/Plugins/Experimental/Chooser/Source/Chooser/Private/RandomizeColumn.cpp b/Engine/Plugins/Experimental/Chooser/Source/Chooser/Private/RandomizeColumn.cpp
--- a/Engine/Plugins/Experimental/Chooser/Source/Chooser/Private/RandomizeColumn.cpp
+++ b/Engine/Plugins/Experimental/Chooser/Source/Chooser/Private/RandomizeColumn.cpp
@@ -186,11 +186,14 @@ void FRandomizeColumn::SetOutputs(FChooserEvaluationContext& Context, int RowInd
 
 	void FRandomizeColumn::AddToDetails(FInstancedPropertyBag& PropertyBag, int32 ColumnIndex, int32 RowIndex)
 	{
-		FName PropertyName("RowData",ColumnIndex);
-		FPropertyBagPropertyDesc PropertyDesc(PropertyName, EPropertyBagPropertyType::Float);
-		PropertyDesc.MetaData.Add(FPropertyBagPropertyDescMetaData("DisplayName", "Randomize"));
-		PropertyBag.AddProperties({PropertyDesc});
-		PropertyBag.SetValueFloat(PropertyName, RowValues[RowIndex]);
+		if (RowValues.IsValidIndex(RowIndex))
+		{
+			FName PropertyName("RowData",ColumnIndex);
+			FPropertyBagPropertyDesc PropertyDesc(PropertyName, EPropertyBagPropertyType::Float);
+			PropertyDesc.MetaData.Add(FPropertyBagPropertyDescMetaData("DisplayName", "Randomize"));
+			PropertyBag.AddProperties({PropertyDesc});
+			PropertyBag.SetValueFloat(PropertyName, RowValues[RowIndex]);
+		}
 	}
 
 	void FRandomizeColumn::SetFromDetails(FInstancedPropertyBag& PropertyBag, int32 ColumnIndex, int32 RowIndex)
diff --git a/Engine/Plugins/Experimental/Chooser/Source/Chooser/Public/IChooserColumn.h b/Engine/Plugins/Experimental/Chooser/Source/Chooser/Public/IChooserColumn.h
--- a/Engine/Plugins/Experimental/Chooser/Source/Chooser/Public/IChooserColumn.h
+++ b/Engine/Plugins/Experimental/Chooser/Source/Chooser/Public/IChooserColumn.h
@@ -69,6 +69,7 @@ public:
 	virtual void MoveRow(int SourceIndex, int TargetIndex) {}
 	virtual void InsertRows(int Index, int Count) {}
 	virtual void CopyRow(FChooserColumnBase& SourceColumn, int SourceIndex, int TargetIndex) {}
+	virtual void CopyFallback(FChooserColumnBase& SourceColumn) {}
 	
 	virtual UScriptStruct* GetInputBaseType() const { return nullptr; };
 	virtual const UScriptStruct* GetInputType() const { return nullptr; };
@@ -94,6 +95,7 @@ public:
 #endif
 };
 
+enum {  ChooserColumn_SpecialIndex_Fallback = -2 };
 
 #if WITH_EDITOR
 #define CHOOSER_COLUMN_BOILERPLATE2(ParameterType, RowValuesProperty) \
diff --git a/Engine/Plugins/Experimental/Chooser/Source/ChooserEditor/Private/BoolColumnEditor.cpp b/Engine/Plugins/Experimental/Chooser/Source/ChooserEditor/Private/BoolColumnEditor.cpp
--- a/Engine/Plugins/Experimental/Chooser/Source/ChooserEditor/Private/BoolColumnEditor.cpp
+++ b/Engine/Plugins/Experimental/Chooser/Source/ChooserEditor/Private/BoolColumnEditor.cpp
@@ -120,8 +120,13 @@ TSharedRef<SWidget> CreateOutputBoolColumnWidget(UChooserTable* Chooser, FChoose
 		+ SHorizontalBox::Slot().AutoWidth().VAlign(VAlign_Center)
 		[
 			SNew(SCheckBox)
-			.IsChecked_Lambda([BoolColumn]() { return BoolColumn->bFallbackValue ? ECheckBoxState::Checked : ECheckBoxState::Unchecked; })
-			.OnCheckStateChanged_Lambda([BoolColumn](ECheckBoxState State) { BoolColumn->bFallbackValue = State == ECheckBoxState::Checked; })
+			.IsChecked_Lambda([ BoolColumn]() { return BoolColumn->bFallbackValue ? ECheckBoxState::Checked : ECheckBoxState::Unchecked; })
+			.OnCheckStateChanged_Lambda([Chooser, BoolColumn](ECheckBoxState State)
+			{
+				const FScopedTransaction Transaction(LOCTEXT("Change Bool Value", "Change Bool Value"));
+				Chooser->Modify(true);
+				BoolColumn->bFallbackValue = State == ECheckBoxState::Checked;
+			})
 		]
 		+ SHorizontalBox::Slot().FillWidth(1);
 	}
diff --git a/Engine/Plugins/Experimental/Chooser/Source/ChooserEditor/Private/ChooserTableEditor.cpp b/Engine/Plugins/Experimental/Chooser/Source/ChooserEditor/Private/ChooserTableEditor.cpp
--- a/Engine/Plugins/Experimental/Chooser/Source/ChooserEditor/Private/ChooserTableEditor.cpp
+++ b/Engine/Plugins/Experimental/Chooser/Source/ChooserEditor/Private/ChooserTableEditor.cpp
@@ -50,6 +50,7 @@
 #include "Exporters/Exporter.h"
 #include "HAL/PlatformApplicationMisc.h"
 #include "Framework/Commands/GenericCommands.h"
+#include "SPositiveActionButton.h"
 
 #define LOCTEXT_NAMESPACE "ChooserEditor"
 
@@ -594,23 +595,42 @@ void FChooserTableEditor::OnObjectsTransacted(UObject* Object, const FTransactio
 		// if this is for the chooser we're editing
 		if (GetChooser() == RowDetails->Chooser)
 		{
-			// copy all the values over
-			TValueOrError<FStructView, EPropertyBagResult> Result = RowDetails->Properties.GetValueStruct("Result", FInstancedStruct::StaticStruct());
-			if (Result.IsValid())
+			if (RowDetails->Chooser->ResultsStructs.IsValidIndex(RowDetails->Row))
 			{
-				RowDetails->Chooser->ResultsStructs[RowDetails->Row] = Result.GetValue().Get<FInstancedStruct>();
-			}
+				// copy all the values over
+				TValueOrError<FStructView, EPropertyBagResult> Result = RowDetails->Properties.GetValueStruct("Result", FInstancedStruct::StaticStruct());
+				if (Result.IsValid())
+				{
+					RowDetails->Chooser->ResultsStructs[RowDetails->Row] = Result.GetValue().Get<FInstancedStruct>();
+				}
 
-			int ColumnIndex = 0;
-			for (FInstancedStruct& ColumnData : RowDetails->Chooser->ColumnsStructs)
+				int ColumnIndex = 0;
+				for (FInstancedStruct& ColumnData : RowDetails->Chooser->ColumnsStructs)
+				{
+					FChooserColumnBase& Column = ColumnData.GetMutable<FChooserColumnBase>();
+					Column.SetFromDetails(RowDetails->Properties, ColumnIndex, RowDetails->Row);
+					ColumnIndex++;
+				}
+			
+				TValueOrError<bool, EPropertyBagResult> DisabledResult = RowDetails->Properties.GetValueBool("Disabled");
+				RowDetails->Chooser->DisabledRows[RowDetails->Row] = DisabledResult.GetValue();
+			}
+			else if (RowDetails->Row == ColumnWidget_SpecialIndex_Fallback)
 			{
-				FChooserColumnBase& Column = ColumnData.GetMutable<FChooserColumnBase>();
-				Column.SetFromDetails(RowDetails->Properties, ColumnIndex, RowDetails->Row);
-				ColumnIndex++;
+				TValueOrError<FStructView, EPropertyBagResult> Result = RowDetails->Properties.GetValueStruct("Result", FInstancedStruct::StaticStruct());
+				if (Result.IsValid())
+				{
+					RowDetails->Chooser->FallbackResult = Result.GetValue().Get<FInstancedStruct>();
+				}
+				
+				int ColumnIndex = 0;
+				for (FInstancedStruct& ColumnData : RowDetails->Chooser->ColumnsStructs)
+				{
+					FChooserColumnBase& Column = ColumnData.GetMutable<FChooserColumnBase>();
+					Column.SetFromDetails(RowDetails->Properties, ColumnIndex, RowDetails->Row);
+					ColumnIndex++;
+				}
 			}
-			
-			TValueOrError<bool, EPropertyBagResult> DisabledResult = RowDetails->Properties.GetValueBool("Disabled");
-			RowDetails->Chooser->DisabledRows[RowDetails->Row] = DisabledResult.GetValue();
 
 			RefreshAll();
 		}
@@ -998,15 +1018,18 @@ int FChooserTableEditor::MoveRow(int SourceRowIndex, int TargetRowIndex)
 	
 void FChooserTableEditor::SelectRow(int32 RowIndex, bool bClear)
 {
-	if (TableRows.IsValidIndex(RowIndex))
+	if (TSharedPtr<FChooserTableRow>* Row = TableRows.FindByPredicate([RowIndex](const TSharedPtr<FChooserTableRow>& InRow)
+		{
+			return InRow->RowIndex == RowIndex;
+		}))
 	{
-		if (!TableView->IsItemSelected(TableRows[RowIndex]))
+		if (!TableView->IsItemSelected(*Row))
 		{
 			if (bClear)
 			{
 				TableView->ClearSelection();
 			}
-			TableView->SetItemSelection(TableRows[RowIndex], true, ESelectInfo::OnMouseClick);
+			TableView->SetItemSelection(*Row, true, ESelectInfo::OnMouseClick);
 		}
 	}
 }
@@ -1095,10 +1118,14 @@ void FChooserTableEditor::UpdateTableColumns()
 		.FillWidth(1.0)
 		.HeaderContent( )
 		[
-			SNew(SHorizontalBox)
-			+ SHorizontalBox::Slot().MaxWidth(120)
-			[
-				CreateColumnComboButton.ToSharedRef()
+			SNew(SVerticalBox)
+			 + SVerticalBox::Slot().AutoHeight()
+			 [
+				SNew(SHorizontalBox)
+				+ SHorizontalBox::Slot().MaxWidth(150)
+				[
+					CreateColumnComboButton.ToSharedRef()
+				]
 			]
 		]
 		);
@@ -1107,7 +1134,7 @@ void FChooserTableEditor::UpdateTableColumns()
 
 void FChooserTableEditor::AddColumn(const UScriptStruct* ColumnType)
 {
-	CreateColumnComboButton->SetIsOpen(false);
+	FSlateApplication::Get().DismissAllMenus();
 	UChooserTable* Chooser = GetChooser();
 	const FScopedTransaction Transaction(LOCTEXT("Add Column Transaction", "Add Column"));
 	Chooser->Modify(true);
@@ -1199,7 +1226,30 @@ void FChooserTableEditor::RefreshRowSelectionDetails()
 			}
 
 			SelectedRows.Add(Selection);
-			
+		}
+		else if (SelectedItem->RowIndex == ColumnWidget_SpecialIndex_Fallback)
+		{
+			TObjectPtr<UChooserRowDetails> Selection = NewObject<UChooserRowDetails>();
+			Selection->Chooser = Chooser;
+			Selection->Row = SelectedItem->RowIndex;
+			Selection->SetFlags(RF_Standalone | RF_Transactional);
+
+			FInstancedStruct& Result = Chooser->FallbackResult;
+			Selection->Properties.AddProperties({ResultPropertyDesc});
+			Selection->Properties.SetValueStruct("Result", FConstStructView(FInstancedStruct::StaticStruct(), reinterpret_cast<uint8*>(&Result)));
+
+			int ColumnIndex = 0;
+			for (FInstancedStruct& ColumnData : Chooser->ColumnsStructs)
+			{
+				FChooserColumnBase& Column = ColumnData.GetMutable<FChooserColumnBase>();
+				if (Column.HasOutputs())
+				{
+					Column.AddToDetails(Selection->Properties, ColumnIndex, SelectedItem->RowIndex);
+				}
+				ColumnIndex++;
+			}
+
+			SelectedRows.Add(Selection);
 		}
 	}
 	
@@ -1260,63 +1310,244 @@ TSharedRef<SDockTab> FChooserTableEditor::SpawnNestedTablesTreeTab( const FSpawn
 		];
 }
 
-TSharedRef<SDockTab> FChooserTableEditor::SpawnTableTab( const FSpawnTabArgs& Args )
+DECLARE_DELEGATE_OneParam(FCreateStructDelegate, UScriptStruct*);
+
+struct FColumnTypeInfoStruct
 {
-	check( Args.GetTabId() == TableTabId );
+	int SortOrder = 100;
+	FString Category;
+	UScriptStruct* Type;
+	bool operator < (const FColumnTypeInfoStruct& Other) const
+	{
+		if (Category == Other.Category)
+		{
+			return Type->GetDisplayNameText().ToString() < Other.Type->GetDisplayNameText().ToString();
+		}
+		else if (SortOrder == Other.SortOrder)
+		{
+			return Category < Other.Category;
+		}
+		else
+		{
+			return SortOrder < Other.SortOrder;
+		}
+	}
+};
 
-	UChooserTable* Chooser = GetChooser();
+TSharedRef<SWidget>	FChooserTableEditor::MakeCreateColumnMenu()
+{
+	FMenuBuilder MenuBuilder(true, nullptr);
+	static TArray<FColumnTypeInfoStruct> ColumnTypes;
 
-	// + button to create new columns
+	if (ColumnTypes.IsEmpty())
+	{
+		UScriptStruct* BaseType = FChooserColumnBase::StaticStruct();
+		for (TObjectIterator<UScriptStruct> StructIt; StructIt; ++StructIt)
+		{
+			if (*StructIt != BaseType && StructIt->IsChildOf(BaseType))
+			{
+				if (!StructIt->HasMetaData("Hidden"))
+				{
+					FColumnTypeInfoStruct Info;
+					Info.Type = *StructIt;
+					Info.Category = StructIt->HasMetaData("Category") ? StructIt->GetMetaData("Category") : "Other";
+
+					if (Info.Category == "Filter")
+					{
+						Info.SortOrder = 1;
+					}
+					else if (Info.Category == "Scoring")
+                   	{
+                   		Info.SortOrder = 2;
+                   	}
+					else if (Info.Category == "Output")
+                   	{
+                   		Info.SortOrder = 3;
+                   	}
+					else if (Info.Category == "Random")
+					{
+						Info.SortOrder = 4;
+					}
+
+					ColumnTypes.Add(Info);
+				}
+			}
+		}
+		ColumnTypes.Sort();
+	}
+
+	FString Section = "";
+	for(FColumnTypeInfoStruct& Type : ColumnTypes)
+	{
+		if (Section != Type.Category)
+		{
+			if (Section != "")
+			{
+				MenuBuilder.EndSection();
+			}
+			Section = Type.Category;
+			MenuBuilder.BeginSection(FName(Section), FText::FromString(Section));
+		}
+		
+		MenuBuilder.AddMenuEntry(Type.Type->GetDisplayNameText(), Type.Type->GetToolTipText(), FSlateIcon(),
+			FUIAction(
+				FExecuteAction::CreateLambda([this, Type]()
+				{
+					AddColumn(Type.Type);
+				}))
+			);
+							
+	}
+	return MenuBuilder.MakeWidget();
+}
+
+struct FResultTypeInfoStruct
+{
+	bool ObjectOnly = false;
+	bool ClassOnly = false;
+	FString Category;
+	UScriptStruct* Type = nullptr;
 	
-	CreateColumnComboButton = SNew(SComboButton).OnGetMenuContent_Lambda([this]()
+	bool operator < (const FResultTypeInfoStruct& Other) const
+	{
+		if (Category == Other.Category)
+		{
+			return Type->GetDisplayNameText().ToString() < Other.Type->GetDisplayNameText().ToString();
+		}
+		else
+		{
+			return Category < Other.Category;
+		}
+	}
+};
+
+void MakeCreateResultMenu(FMenuBuilder& MenuBuilder, EObjectChooserResultType ChooserResultType, FCreateStructDelegate CreateStruct)
+{
+	static TArray<FResultTypeInfoStruct> ResultTypes;
+
+	if (ResultTypes.IsEmpty())
+	{
+		UScriptStruct* BaseType = FObjectChooserBase::StaticStruct();
+		for (TObjectIterator<UScriptStruct> StructIt; StructIt; ++StructIt)
+		{
+			if (*StructIt != BaseType && StructIt->IsChildOf(BaseType))
+			{
+				if (!StructIt->HasMetaData("Hidden"))
+				{
+					FResultTypeInfoStruct Info;
+					Info.Type = *StructIt;
+					Info.Category = StructIt->HasMetaData("Category") ? StructIt->GetMetaData("Category") : "Other";
+
+					if (StructIt->HasMetaData("ResultType"))
+					{
+						FString ResultTypeString = StructIt->GetMetaData("ResultType");
+						Info.ClassOnly = ResultTypeString == "Class";
+						Info.ObjectOnly = ResultTypeString == "Object";
+					}
+					
+					ResultTypes.Add(Info);
+				}
+			}
+		}
+		ResultTypes.Sort();
+	}
+
+	FString Section = "";
+	for(FResultTypeInfoStruct& Type : ResultTypes)
 	{
-		FStructViewerInitializationOptions Options;
-		Options.StructFilter = MakeShared<FStructFilter>(FChooserColumnBase::StaticStruct());
-		Options.bAllowViewOptions = false;
-		Options.bShowNoneOption = false;
-		Options.NameTypeToDisplay = EStructViewerNameTypeToDisplay::DisplayName;
+		if (Section != Type.Category)
+		{
+			if (Section != "")
+			{
+				MenuBuilder.EndSection();
+			}
+			Section = Type.Category;
+			MenuBuilder.BeginSection(FName(Section), FText::FromString(Section));
+		}
 		
-		// Add class filter for columns here
-		FStructViewerModule& StructViewerModule = FModuleManager::LoadModuleChecked<FStructViewerModule>("StructViewer");
-		TSharedRef<SWidget> Widget = StructViewerModule.CreateStructViewer(Options, FOnStructPicked::CreateRaw(this, &FChooserTableEditor::AddColumn));
-		return Widget;
-	})
-	.ComboButtonStyle(FAppStyle::Get(), "SimpleComboButton")
-	.ButtonContent()
-	[
-		SNew(STextBlock).Text(LOCTEXT("AddColumn", "+ Add Column"))
-	];
+		MenuBuilder.AddMenuEntry(Type.Type->GetDisplayNameText(), Type.Type->GetToolTipText(), FSlateIcon(),
+			FUIAction(
+				FExecuteAction::CreateLambda([Type, CreateStruct]()
+				{
+					CreateStruct.Execute(Type.Type);
+				}),
+				FCanExecuteAction::CreateLambda([Type, ChooserResultType]()
+				{
+					if (Type.ClassOnly && ChooserResultType == EObjectChooserResultType::ObjectResult)
+					{
+						return false;
+					}
+					if (Type.ObjectOnly && ChooserResultType == EObjectChooserResultType::ClassResult)
+					{
+						return false;
+					}
+					return true;
+				})
+				)
+			);
+							
+	}
+}
 
+TSharedRef<SWidget>	FChooserTableEditor::MakeCreateRowMenu()
+{
+	FMenuBuilder MenuBuilder(true, nullptr);
 
-	CreateRowComboButton = SNew(SComboButton)
-		.ComboButtonStyle(FAppStyle::Get(), "SimpleComboButton")
-		.ButtonContent()
-		[
-			SNew(STextBlock).Text(LOCTEXT("AddRow", "+ Add Row"))
-		]
+    UChooserTable* Chooser = GetChooser();
+	if (!Chooser->FallbackResult.IsValid())
+	{
+		MenuBuilder.AddSubMenu(
+			LOCTEXT("Add Fallback", "Add Fallback Result"),
+				LOCTEXT("Add Fallback Tooltip", "Add a Fallback row to the chooser, which will be used in the case where no other rows passed all filter columns"),
+				FNewMenuDelegate::CreateLambda([this](FMenuBuilder& MenuBuilder)
+				{
+					UChooserTable* Chooser = GetChooser();
+					MakeCreateResultMenu(MenuBuilder, Chooser->ResultType, FCreateStructDelegate::CreateLambda([this](UScriptStruct* Type)
+					{
+						UChooserTable* Chooser = GetChooser();
+						const FScopedTransaction Transaction(LOCTEXT("Add Fallback Row Transaction", "Add Fallback Row"));
+						Chooser->Modify(true);
+						Chooser->FallbackResult.InitializeAs(Type);
+						UpdateTableRows();
+					}));
+				})
+			);
+
+	}
+
+	MakeCreateResultMenu(MenuBuilder, Chooser->ResultType, FCreateStructDelegate::CreateLambda([this](UScriptStruct* Type)
+	{
+
+		UChooserTable* Chooser = GetChooser();
+		const FScopedTransaction Transaction(LOCTEXT("Add Row Transaction", "Add Row"));
+		Chooser->Modify(true);
+		Chooser->ResultsStructs.SetNum(Chooser->ResultsStructs.Num()+1);
+		Chooser->ResultsStructs.Last().InitializeAs(Type);
+		UpdateTableRows();
+	}));
+
+	return MenuBuilder.MakeWidget();
+}
+
+TSharedRef<SDockTab> FChooserTableEditor::SpawnTableTab( const FSpawnTabArgs& Args )
+{
+	check( Args.GetTabId() == TableTabId );
+
+	UChooserTable* Chooser = GetChooser();
+
+	// + button to create new columns
+	
+	CreateColumnComboButton = SNew(SPositiveActionButton)
+		.Text(LOCTEXT("Add Column", "Add Column"))
 		.OnGetMenuContent_Lambda([this]()
 		{
-			FStructViewerInitializationOptions Options;
-			Options.StructFilter = MakeShared<FStructFilter>(FObjectChooserBase::StaticStruct());
-			Options.bAllowViewOptions = false;
-			Options.bShowNoneOption = false;
-			Options.NameTypeToDisplay = EStructViewerNameTypeToDisplay::DisplayName;
-			
-			TSharedRef<SWidget> Widget = FModuleManager::LoadModuleChecked<FStructViewerModule>("StructViewer").CreateStructViewer(Options, FOnStructPicked::CreateLambda([this](const UScriptStruct* ChosenStruct)
-			{
-				CreateRowComboButton->SetIsOpen(false);
-				UChooserTable* Chooser = GetChooser();
-				const FScopedTransaction Transaction(LOCTEXT("Add Row Transaction", "Add Row"));
-				Chooser->Modify(true);
-
-				Chooser->ResultsStructs.SetNum(Chooser->ResultsStructs.Num()+1);
-				Chooser->ResultsStructs.Last().InitializeAs(ChosenStruct);
-				UpdateTableRows();
-			}));
-			
-			return Widget;
+			return MakeCreateColumnMenu();
 		});
 
+	CreateRowComboButton = SNew(SPositiveActionButton)
+		.Text(LOCTEXT("Add Row", "Add Row"))
+		.OnGetMenuContent(this, &FChooserTableEditor::MakeCreateRowMenu);
+			
 	HeaderRow = SNew(SHeaderRow);
 
 	UpdateTableRows();
@@ -1426,7 +1657,11 @@ void FChooserTableEditor::UpdateTableRows()
 	}
 
 	// Add one at the end, for the Fallback result
-	TableRows.Add(MakeShared<FChooserTableRow>(SChooserTableRow::SpecialIndex_Fallback));
+	if (Chooser->FallbackResult.IsValid())
+	{
+		TableRows.Add(MakeShared<FChooserTableRow>(SChooserTableRow::SpecialIndex_Fallback));
+	}
+	
 	// Add one at the end, for the "Add Row" control
 	TableRows.Add(MakeShared<FChooserTableRow>(SChooserTableRow::SpecialIndex_AddRow));
 
@@ -1546,7 +1781,14 @@ int32 FChooserTableEditor::DeleteSelectedRowsInternal(int32 RowIndexToRemember)
 	TArray<uint32> RowsToDelete;
 	for(auto& SelectedRow:SelectedRows)
 	{
-		RowsToDelete.Add(SelectedRow->Row);
+		if (SelectedRow->Row == ColumnWidget_SpecialIndex_Fallback)
+		{
+			Chooser->FallbackResult.Reset();
+		}
+		else
+		{
+			RowsToDelete.Add(SelectedRow->Row);
+		}
 	}
 
 	SelectedRows.SetNum(0);
@@ -1754,7 +1996,10 @@ void FChooserTableEditor::ToggleDisableSelection()
 		{
 			for (auto& Row : SelectedRows)
 			{
-				Chooser->DisabledRows[Row->Row] = !bDisabled;
+				if (Chooser->DisabledRows.IsValidIndex(Row->Row))
+				{
+					Chooser->DisabledRows[Row->Row] = !bDisabled;
+				}
 			}
 			RefreshRowSelectionDetails();
 		}
@@ -1797,17 +2042,32 @@ void FChooserTableEditor::DuplicateSelection()
 	}
 }
 
+bool FChooserTableEditor::HasFallbackSelected()
+{
+	for(UChooserRowDetails* SelectedRow : SelectedRows)
+	{
+		if (SelectedRow->Row == ColumnWidget_SpecialIndex_Fallback)
+		{
+			return true;
+		}
+	}
+
+	return false;
+}
 
 bool FChooserTableEditor::CanMoveRowsUp()
 {
 	if (HasRowsSelected())
 	{
 		UChooserTable* Chooser = GetChooser();
-		
+
 		int MinSelectedRow = Chooser->ResultsStructs.Num();
 		for(UChooserRowDetails* SelectedRow : SelectedRows)
 		{
-			MinSelectedRow = FMath::Min(SelectedRow->Row, MinSelectedRow);
+			if (SelectedRow->Row != ColumnWidget_SpecialIndex_Fallback)
+			{
+				MinSelectedRow = FMath::Min(SelectedRow->Row, MinSelectedRow);
+			}
 		}
 
 		return MinSelectedRow > 0;
@@ -1823,7 +2083,10 @@ void FChooserTableEditor::MoveRowsUp()
 		int MinSelectedRow = Chooser->ResultsStructs.Num();
 		for(UChooserRowDetails* SelectedRow : SelectedRows)
 		{
-			MinSelectedRow = FMath::Min(SelectedRow->Row, MinSelectedRow);
+			if (SelectedRow->Row != ColumnWidget_SpecialIndex_Fallback)
+			{
+				MinSelectedRow = FMath::Min(SelectedRow->Row, MinSelectedRow);
+			}
 		}
 		MoveRows(MinSelectedRow - 1);
 	}
@@ -1971,22 +2234,35 @@ UChooserTable* FChooserTableEditor::CopySelectionInternal()
 			Column.SetNumRows(0);
 			Column.SetNumRows(SelectedRowsCopy.Num());
 		}
-
+		
+		if (SelectedRowsCopy.Num() > 0 && SelectedRowsCopy[0]->Row == ColumnWidget_SpecialIndex_Fallback)
+		{
+			SelectedRowsCopy.RemoveAt(0);
+			
+			CopyData->FallbackResult = Chooser->FallbackResult;
+			if (FNestedChooser* CopiedNestedChooser = CopyData->FallbackResult.GetMutablePtr<FNestedChooser>())
+			{
+				// if the fallback result was a nested chooser, duplicate it
+				CopiedNestedChooser->Chooser = DuplicateNestedChooser(CopiedNestedChooser->Chooser, CopyData);
+			}
+		}
+		
 		CopyData->ResultsStructs.SetNum(SelectedRowsCopy.Num());
+		CopyData->DisabledRows.SetNum(SelectedRowsCopy.Num());
 
 		// add the selected results and column data
 		
 		for (int RowIndex = 0; RowIndex < SelectedRowsCopy.Num(); RowIndex++)
 		{
-			const FInstancedStruct& Result = Chooser->ResultsStructs[SelectedRowsCopy[RowIndex]->Row];
-
-			CopyData->ResultsStructs[RowIndex] = Result;
+			CopyData->ResultsStructs[RowIndex] = Chooser->ResultsStructs[SelectedRowsCopy[RowIndex]->Row];
 			if (FNestedChooser* CopiedNestedChooser = CopyData->ResultsStructs[RowIndex].GetMutablePtr<FNestedChooser>())
 			{
 				// if the result for this row was a nested chooser, duplicate it
 				CopiedNestedChooser->Chooser = DuplicateNestedChooser(CopiedNestedChooser->Chooser, CopyData);
 			}
-
+			
+			CopyData->DisabledRows[RowIndex] = Chooser->DisabledRows[SelectedRowsCopy[RowIndex]->Row];
+			
 			for (int ColumnIndex = 0; ColumnIndex < CopyData->ColumnsStructs.Num(); ColumnIndex++)
 			{
 				FChooserColumnBase& SourceColumn = Chooser->ColumnsStructs[ColumnIndex].GetMutable<FChooserColumnBase>();
@@ -2065,7 +2341,7 @@ void FChooserTableEditor::PasteInternal(UChooserTable* PastedContent, int PasteR
 	UChooserTable* Chooser = GetChooser();
 	Chooser->Modify();
 	
-	if (PastedContent->ResultsStructs.IsEmpty())
+	if (PastedContent->ResultsStructs.IsEmpty() && !PastedContent->FallbackResult.IsValid())
 	{
 		// pasting a column
 		int InsertColumnIndex = Chooser->ColumnsStructs.Num();
@@ -2106,11 +2382,40 @@ void FChooserTableEditor::PasteInternal(UChooserTable* PastedContent, int PasteR
 				{
 					InsertIndex = FMath::Max(InsertIndex, SelectedRows[SelectedRowIndex]->Row);
 				}
-				InsertIndex++;
+				if (InsertIndex == ColumnWidget_SpecialIndex_Fallback)
+				{
+					// if the only row selected was the fallback, reset insert index to the last row
+					InsertIndex = Chooser->ResultsStructs.Num();
+				}
+				else
+				{
+					InsertIndex++;
+				}
 			}
 		}
 
-		Chooser->ResultsStructs.Insert(PastedContent->ResultsStructs, InsertIndex);
+		if (PastedContent->ResultsStructs.Num() > 0)
+		{
+			Chooser->ResultsStructs.Insert(PastedContent->ResultsStructs, InsertIndex);
+			Chooser->DisabledRows.Insert(PastedContent->DisabledRows, InsertIndex);
+
+			// Make sure each column has the same number of row datas as there are results
+			for(FInstancedStruct& ColumnData : Chooser->ColumnsStructs)
+			{
+				FChooserColumnBase& Column = ColumnData.GetMutable<FChooserColumnBase>();
+				Column.InsertRows(InsertIndex, RowsToPaste);
+			}
+		}
+		if (PastedContent->FallbackResult.IsValid())
+		{
+			// paste fallback result if copy data has one
+			Chooser->FallbackResult = PastedContent->FallbackResult;
+			if (FNestedChooser* NestedChooser = Chooser->FallbackResult.GetMutablePtr<FNestedChooser>())
+			{
+				// duplicate the nested chooser if the fallback result refers to a nested chooser
+				NestedChooser->Chooser = DuplicateNestedChooser(NestedChooser->Chooser, Chooser);
+			}
+		}
 		
 		if (!PastedContent->NestedChoosers.IsEmpty())
 		{
@@ -2123,14 +2428,6 @@ void FChooserTableEditor::PasteInternal(UChooserTable* PastedContent, int PasteR
 				}
 			}
 		}
-		
-		// Make sure each column has the same number of row datas as there are results
-		for(FInstancedStruct& ColumnData : Chooser->ColumnsStructs)
-		{
-			FChooserColumnBase& Column = ColumnData.GetMutable<FChooserColumnBase>();
-			Column.InsertRows(InsertIndex, RowsToPaste);
-		}
-
 
 		// try to also paste column data from columns in the paste buffer which match the columns in the current chooser
 		// -- matching by column type and input value name
@@ -2170,6 +2467,12 @@ void FChooserTableEditor::PasteInternal(UChooserTable* PastedContent, int PasteR
 								{
 									Column.CopyRow(PastedColumn, i, InsertIndex + i);
 								}
+
+								if (PastedContent->FallbackResult.IsValid())
+								{
+									// if the fallback row was copied, paste the fallback data for columns
+									Column.CopyFallback(PastedColumn);
+								}
 								break;
 							}
 						}
@@ -2190,7 +2493,6 @@ void FChooserTableEditor::PasteInternal(UChooserTable* PastedContent, int PasteR
 			}
 		}
 
-		
 		for(int SourceColumnIndex = 0; SourceColumnIndex < PastedContent->ColumnsStructs.Num(); SourceColumnIndex++)
 		{
 			if (!MatchedSourceColumns[SourceColumnIndex])
@@ -2217,6 +2519,10 @@ void FChooserTableEditor::PasteInternal(UChooserTable* PastedContent, int PasteR
 		{
 			SelectRow(InsertIndex + i, false);
 		}
+		if(PastedContent->FallbackResult.IsValid())
+		{
+			SelectRow(ColumnWidget_SpecialIndex_Fallback, false);
+		}
 	}
 	
 	RefreshAll();
diff --git a/Engine/Plugins/Experimental/Chooser/Source/ChooserEditor/Private/ChooserTableEditor.h b/Engine/Plugins/Experimental/Chooser/Source/ChooserEditor/Private/ChooserTableEditor.h
--- a/Engine/Plugins/Experimental/Chooser/Source/ChooserEditor/Private/ChooserTableEditor.h
+++ b/Engine/Plugins/Experimental/Chooser/Source/ChooserEditor/Private/ChooserTableEditor.h
@@ -17,6 +17,7 @@
 #include "ChooserTableEditor.generated.h"
 
 class SComboButton;
+class SPositiveActionButton;
 class SEditableText;
 class IDetailsView;
 class UChooserRowDetails;
@@ -106,6 +107,7 @@ namespace UE::ChooserEditor
 		void ToggleDisableSelection();
 		void DeleteSelection();
 		void DuplicateSelection();
+		bool HasFallbackSelected();
 
 		bool CanMoveRowsUp();
 		void MoveRowsUp();
@@ -217,17 +219,19 @@ namespace UE::ChooserEditor
 		void UpdateTableColumns();
 		TArray<TSharedPtr<FChooserTableRow>> TableRows;
 	
-		TSharedPtr<SComboButton> CreateColumnComboButton;
-		TSharedPtr<SComboButton> CreateRowComboButton;
+		TSharedPtr<SPositiveActionButton> CreateColumnComboButton;
+		TSharedPtr<SWidget> CreateRowComboButton;
 
 		TSharedPtr<SHeaderRow> HeaderRow;
 		TSharedPtr<SListView<TSharedPtr<FChooserTableRow>>> TableView;
 
 		ESelectionType CurrentSelectionType = ESelectionType::Root;
 
+		TSharedRef<SWidget>	MakeCreateRowMenu();
+		TSharedRef<SWidget>	MakeCreateColumnMenu();
 	public:
 
-		TSharedPtr<SComboButton>& GetCreateRowComboButton() { return CreateRowComboButton; };
+		TSharedPtr<SWidget>& GetCreateRowComboButton() { return CreateRowComboButton; };
 		TSharedPtr<SNestedChooserTree> NestedChooserTree;
 
 		/** The name given to all instances of this type of editor */
diff --git a/Engine/Plugins/Experimental/Chooser/Source/ChooserEditor/Private/FloatDistanceColumnEditor.cpp b/Engine/Plugins/Experimental/Chooser/Source/ChooserEditor/Private/FloatDistanceColumnEditor.cpp
--- a/Engine/Plugins/Experimental/Chooser/Source/ChooserEditor/Private/FloatDistanceColumnEditor.cpp
+++ b/Engine/Plugins/Experimental/Chooser/Source/ChooserEditor/Private/FloatDistanceColumnEditor.cpp
@@ -31,8 +31,8 @@ TSharedRef<SWidget> CreateFloatDistanceColumnWidget(UChooserTable* Chooser, FCho
 		// create column header widget
 		
 		const FSlateBrush* ColumnIcon = FCoreStyle::Get().GetBrush("Icons.SortUp");
-		const FText ColumnTooltip = LOCTEXT("Float distance tooltip", "Float Distance Column: rows recieve a Cost based on how different the input float is from the row value");
-		const FText ColumnName = LOCTEXT("Float Distance","Float Distance");
+		const FText ColumnTooltip = LOCTEXT("Float difference tooltip", "Float Difference Column: rows recieve a Score based on how different the input float is from the row value");
+		const FText ColumnName = LOCTEXT("Float Difference","Float Difference");
 		
 		TSharedPtr<SWidget> DebugWidget = nullptr;
 		if (Chooser->GetEnableDebugTesting())
diff --git a/Engine/Plugins/Experimental/Chooser/Source/ChooserEditor/Private/OutputFloatColumnEditor.cpp b/Engine/Plugins/Experimental/Chooser/Source/ChooserEditor/Private/OutputFloatColumnEditor.cpp
--- a/Engine/Plugins/Experimental/Chooser/Source/ChooserEditor/Private/OutputFloatColumnEditor.cpp
+++ b/Engine/Plugins/Experimental/Chooser/Source/ChooserEditor/Private/OutputFloatColumnEditor.cpp
@@ -48,7 +48,7 @@ TSharedRef<SWidget> CreateOutputFloatColumnWidget(UChooserTable* Chooser, FChoos
         		})
         		.OnValueCommitted_Lambda([Chooser, OutputFloatColumn](double NewValue, ETextCommit::Type CommitType)
         		{
-					const FScopedTransaction Transaction(LOCTEXT("Edit Min Value", "Edit Min Value"));
+					const FScopedTransaction Transaction(LOCTEXT("Edit Float Value", "Edit Float Value"));
 					Chooser->Modify(true);
 					OutputFloatColumn->FallbackValue = NewValue;
         		});	
@@ -64,7 +64,7 @@ TSharedRef<SWidget> CreateOutputFloatColumnWidget(UChooserTable* Chooser, FChoos
     		{
     			if (Row < OutputFloatColumn->RowValues.Num())
     			{
-    				const FScopedTransaction Transaction(LOCTEXT("Edit Min Value", "Edit Min Value"));
+    				const FScopedTransaction Transaction(LOCTEXT("Edit Float Value", "Edit Float Value"));
     				Chooser->Modify(true);
     				OutputFloatColumn->RowValues[Row] = NewValue;
     			}
diff --git a/Engine/Plugins/Experimental/Chooser/Source/ChooserEditor/Private/SChooserTableRow.cpp b/Engine/Plugins/Experimental/Chooser/Source/ChooserEditor/Private/SChooserTableRow.cpp
--- a/Engine/Plugins/Experimental/Chooser/Source/ChooserEditor/Private/SChooserTableRow.cpp
+++ b/Engine/Plugins/Experimental/Chooser/Source/ChooserEditor/Private/SChooserTableRow.cpp
@@ -58,7 +58,7 @@ namespace UE::ChooserEditor
 					]
 				);
 		}
-		else if (RowIndex->RowIndex == SpecialIndex_Fallback)
+		else if (RowIndex->RowIndex == SpecialIndex_Fallback || RowIndex->RowIndex == SpecialIndex_AddRow )
 		{
 			SetContent(
 					SNew(SOverlay)
@@ -327,7 +327,11 @@ namespace UE::ChooserEditor
 			// on the row past the end, show an Add button in the result column
 			if (ColumnName == Result)
 			{
-				return Editor->GetCreateRowComboButton().ToSharedRef();
+				return SNew(SHorizontalBox)
+								+SHorizontalBox::Slot().AutoWidth()
+								[
+									Editor->GetCreateRowComboButton().ToSharedRef()
+								];
 			}
 		}
 
diff --git a/Engine/Plugins/Experimental/Chooser/Source/ProxyTable/Internal/LookupProxy.h b/Engine/Plugins/Experimental/Chooser/Source/ProxyTable/Internal/LookupProxy.h
--- a/Engine/Plugins/Experimental/Chooser/Source/ProxyTable/Internal/LookupProxy.h
+++ b/Engine/Plugins/Experimental/Chooser/Source/ProxyTable/Internal/LookupProxy.h
@@ -25,7 +25,7 @@ public:
 	CHOOSER_PARAMETER_BOILERPLATE();
 };
 
-USTRUCT()
+USTRUCT(DisplayName = "Lookup Proxy", Meta = (Category = "Proxy Table", Tooltip = "Find a Proxy Asset entry in a Proxy Table, and evaluate it's value if this row is selected."))
 struct PROXYTABLE_API FLookupProxy : public FObjectChooserBase
 {
 	GENERATED_BODY()
