Backout of CL 36135815 to fix crash in chooser output types

diff --git a/Engine/Plugins/Experimental/Chooser/Source/Chooser/Internal/BoolColumn.h b/Engine/Plugins/Experimental/Chooser/Source/Chooser/Internal/BoolColumn.h
--- a/Engine/Plugins/Experimental/Chooser/Source/Chooser/Internal/BoolColumn.h
+++ b/Engine/Plugins/Experimental/Chooser/Source/Chooser/Internal/BoolColumn.h
@@ -46,7 +46,7 @@ public:
 	
 };
 
-USTRUCT(DisplayName = "Bool", Meta = (Category = "Filter", Tooltip = "A column that filters rows based on the value of a Bool."))
+USTRUCT()
 struct CHOOSER_API FBoolColumn : public FChooserColumnBase
 {
 	GENERATED_BODY()
diff --git a/Engine/Plugins/Experimental/Chooser/Source/Chooser/Internal/Chooser.h b/Engine/Plugins/Experimental/Chooser/Source/Chooser/Internal/Chooser.h
--- a/Engine/Plugins/Experimental/Chooser/Source/Chooser/Internal/Chooser.h
+++ b/Engine/Plugins/Experimental/Chooser/Source/Chooser/Internal/Chooser.h
@@ -197,7 +197,7 @@ public:
 	static FObjectChooserBase::EIteratorStatus IterateChooser(const UChooserTable* Chooser, FObjectChooserBase::FObjectChooserIteratorCallback Callback);
 };
 
-USTRUCT(DisplayName = "Nested Chooser", Meta = (Category = "Chooser", Tooltip = "Reference another ChooserTable embedded in this asset, which will be evaluated at runtime if this row is selected."))
+USTRUCT(BlueprintType, DisplayName = "Nested Chooser")
 struct CHOOSER_API FNestedChooser : public FObjectChooserBase
 {
 	GENERATED_BODY()
@@ -215,7 +215,7 @@ struct CHOOSER_API FNestedChooser : public FObjectChooserBase
 	TObjectPtr<UChooserTable> Chooser;
 };
 
-USTRUCT(DisplayName = "Evaluate Chooser", Meta = (Category = "Chooser", Tooltip = "Reference another ChooserTable asset, which will be evaluated at runtime if this row is selected."))
+USTRUCT(BlueprintType, DisplayName = "Evaluate Chooser")
 struct CHOOSER_API FEvaluateChooser : public FObjectChooserBase
 {
 	GENERATED_BODY()
diff --git a/Engine/Plugins/Experimental/Chooser/Source/Chooser/Internal/EnumColumn.h b/Engine/Plugins/Experimental/Chooser/Source/Chooser/Internal/EnumColumn.h
--- a/Engine/Plugins/Experimental/Chooser/Source/Chooser/Internal/EnumColumn.h
+++ b/Engine/Plugins/Experimental/Chooser/Source/Chooser/Internal/EnumColumn.h
@@ -92,7 +92,7 @@ struct CHOOSER_API FEnumColumnBase : public FChooserColumnBase
 	virtual void EnumChanged(const UEnum* Enum) {}
 };
 
-USTRUCT(DisplayName = "Enum", Meta = (Category = "Filter", Tooltip = "A column that filters rows based on the value of an Enum, with Equal Not Equal, or Any as cell comparison options."))
+USTRUCT()
 struct CHOOSER_API FEnumColumn : public FEnumColumnBase
 {
 	GENERATED_BODY()
diff --git a/Engine/Plugins/Experimental/Chooser/Source/Chooser/Internal/FloatDistanceColumn.h b/Engine/Plugins/Experimental/Chooser/Source/Chooser/Internal/FloatDistanceColumn.h
--- a/Engine/Plugins/Experimental/Chooser/Source/Chooser/Internal/FloatDistanceColumn.h
+++ b/Engine/Plugins/Experimental/Chooser/Source/Chooser/Internal/FloatDistanceColumn.h
@@ -35,7 +35,7 @@ public:
 	virtual void NativeAutoPopulate(UObject* InObject, bool& OutSuccess, float& OutValue) { }
 };
 
-USTRUCT(DisplayName = "Float Difference", Meta = (Category = "Scoring", Tooltip = "A column which scores rows based on their difference from an Input float."))
+USTRUCT()
 struct CHOOSER_API FFloatDistanceColumn : public FChooserColumnBase
 {
 	GENERATED_BODY()
diff --git a/Engine/Plugins/Experimental/Chooser/Source/Chooser/Internal/FloatRangeColumn.h b/Engine/Plugins/Experimental/Chooser/Source/Chooser/Internal/FloatRangeColumn.h
--- a/Engine/Plugins/Experimental/Chooser/Source/Chooser/Internal/FloatRangeColumn.h
+++ b/Engine/Plugins/Experimental/Chooser/Source/Chooser/Internal/FloatRangeColumn.h
@@ -58,7 +58,7 @@ struct FChooserFloatRangeRowData
 };
 
 
-USTRUCT(DisplayName = "Float Range", Meta = (Category = "Filter", Tooltip = "A column which filters rows if an input value is not within the range specified for the row."))
+USTRUCT()
 struct CHOOSER_API FFloatRangeColumn : public FChooserColumnBase
 {
 	GENERATED_BODY()
diff --git a/Engine/Plugins/Experimental/Chooser/Source/Chooser/Internal/GameplayTagColumn.h b/Engine/Plugins/Experimental/Chooser/Source/Chooser/Internal/GameplayTagColumn.h
--- a/Engine/Plugins/Experimental/Chooser/Source/Chooser/Internal/GameplayTagColumn.h
+++ b/Engine/Plugins/Experimental/Chooser/Source/Chooser/Internal/GameplayTagColumn.h
@@ -47,7 +47,7 @@ enum class EGameplayTagMatchDirection : uint8
 	InputInRowValue,
 };
 
-USTRUCT(DisplayName = "Gameplay Tag", Meta = (Category = "Filter", Tooltip = "A column which filters rows by comparing Gameplay Tags."))
+USTRUCT()
 struct CHOOSER_API FGameplayTagColumn : public FChooserColumnBase
 {
 	GENERATED_BODY()
diff --git a/Engine/Plugins/Experimental/Chooser/Source/Chooser/Internal/MultiEnumColumn.h b/Engine/Plugins/Experimental/Chooser/Source/Chooser/Internal/MultiEnumColumn.h
--- a/Engine/Plugins/Experimental/Chooser/Source/Chooser/Internal/MultiEnumColumn.h
+++ b/Engine/Plugins/Experimental/Chooser/Source/Chooser/Internal/MultiEnumColumn.h
@@ -27,7 +27,7 @@ struct FChooserMultiEnumRowData
 };
 
 
-USTRUCT(DisplayName = "Enum (Or)", Meta = (Category = "Filter", Tooltip = "A column which filters rows using an Enum variable, where rows pass if the enum is one of any of the checked values."))
+USTRUCT()
 struct CHOOSER_API FMultiEnumColumn : public FChooserColumnBase
 {
 	GENERATED_BODY()
diff --git a/Engine/Plugins/Experimental/Chooser/Source/Chooser/Internal/ObjectChooser_Asset.h b/Engine/Plugins/Experimental/Chooser/Source/Chooser/Internal/ObjectChooser_Asset.h
--- a/Engine/Plugins/Experimental/Chooser/Source/Chooser/Internal/ObjectChooser_Asset.h
+++ b/Engine/Plugins/Experimental/Chooser/Source/Chooser/Internal/ObjectChooser_Asset.h
@@ -6,7 +6,7 @@
 #include "IObjectChooser.h"
 #include "ObjectChooser_Asset.generated.h"
 
-USTRUCT(DisplayName = "Asset", Meta = (ResultType = "Object", Category = "Basic", Tooltip = "A hard reference to a specific asset"))
+USTRUCT(DisplayName = "Asset")
 struct CHOOSER_API FAssetChooser : public FObjectChooserBase
 {
 	GENERATED_BODY()
@@ -23,7 +23,7 @@ struct CHOOSER_API FAssetChooser : public FObjectChooserBase
 	TObjectPtr<UObject> Asset;
 };
 
-USTRUCT(DisplayName = "Asset (Soft Reference)", Meta = (ResultType = "Object", Category = "Basic", Tooltip = "A soft object reference to a specific asset\nAssets will need to be preloaded manually to avoid a hitch if they are selected."))
+USTRUCT(DisplayName = "Soft Asset")
 struct CHOOSER_API FSoftAssetChooser : public FObjectChooserBase
 {
 	GENERATED_BODY()
diff --git a/Engine/Plugins/Experimental/Chooser/Source/Chooser/Internal/ObjectChooser_Class.h b/Engine/Plugins/Experimental/Chooser/Source/Chooser/Internal/ObjectChooser_Class.h
--- a/Engine/Plugins/Experimental/Chooser/Source/Chooser/Internal/ObjectChooser_Class.h
+++ b/Engine/Plugins/Experimental/Chooser/Source/Chooser/Internal/ObjectChooser_Class.h
@@ -6,7 +6,7 @@
 #include "IObjectChooser.h"
 #include "ObjectChooser_Class.generated.h"
 
-USTRUCT(DisplayName = "Class", Meta = (ResultType = "Class", Category = "Basic", Tooltip = "A reference to a Class.\nOnly for use in Choosers with ResultType set to Sub Class Of"))
+USTRUCT(DisplayName = "Class")
 struct CHOOSER_API FClassChooser : public FObjectChooserBase
 {
 	GENERATED_BODY()
diff --git a/Engine/Plugins/Experimental/Chooser/Source/Chooser/Internal/ObjectClassColumn.h b/Engine/Plugins/Experimental/Chooser/Source/Chooser/Internal/ObjectClassColumn.h
--- a/Engine/Plugins/Experimental/Chooser/Source/Chooser/Internal/ObjectClassColumn.h
+++ b/Engine/Plugins/Experimental/Chooser/Source/Chooser/Internal/ObjectClassColumn.h
@@ -34,7 +34,7 @@ struct FChooserObjectClassRowData
 	bool Evaluate(const UObject* LeftHandSide) const;
 };
 
-USTRUCT(DisplayName = "Object Class", Meta = (Category = "Filter", Tooltip = "A column which filters rows using an Object reference variable, by checking if that object is of a certain Class."))
+USTRUCT()
 struct CHOOSER_API FObjectClassColumn : public FChooserColumnBase
 {
 	GENERATED_BODY()
diff --git a/Engine/Plugins/Experimental/Chooser/Source/Chooser/Internal/ObjectColumn.h b/Engine/Plugins/Experimental/Chooser/Source/Chooser/Internal/ObjectColumn.h
--- a/Engine/Plugins/Experimental/Chooser/Source/Chooser/Internal/ObjectColumn.h
+++ b/Engine/Plugins/Experimental/Chooser/Source/Chooser/Internal/ObjectColumn.h
@@ -54,7 +54,7 @@ struct FChooserObjectRowData
 	bool Evaluate(const FSoftObjectPath& LeftHandSide) const;
 };
 
-USTRUCT(DisplayName = "Object", Meta = (Category = "Filter", Tooltip = "A column which filters rows by an input Object to specified Objects for each row."))
+USTRUCT()
 struct CHOOSER_API FObjectColumn : public FChooserColumnBase
 {
 	GENERATED_BODY()
diff --git a/Engine/Plugins/Experimental/Chooser/Source/Chooser/Internal/OutputBoolColumn.h b/Engine/Plugins/Experimental/Chooser/Source/Chooser/Internal/OutputBoolColumn.h
--- a/Engine/Plugins/Experimental/Chooser/Source/Chooser/Internal/OutputBoolColumn.h
+++ b/Engine/Plugins/Experimental/Chooser/Source/Chooser/Internal/OutputBoolColumn.h
@@ -8,7 +8,7 @@
 #include "StructUtils/InstancedStruct.h"
 #include "OutputBoolColumn.generated.h"
 
-USTRUCT(DisplayName = "Output Bool", Meta = (Category = "Output", Tooltip = "A column which writes a Bool value."))
+USTRUCT()
 struct CHOOSER_API FOutputBoolColumn : public FChooserColumnBase
 {
 	GENERATED_BODY()
@@ -26,18 +26,7 @@ struct CHOOSER_API FOutputBoolColumn : public FChooserColumnBase
 
 	virtual void AddToDetails(FInstancedPropertyBag& PropertyBag, int32 ColumnIndex, int32 RowIndex) override;
 	virtual void SetFromDetails(FInstancedPropertyBag& PropertyBag, int32 ColumnIndex, int32 RowIndex) override;
-	virtual void CopyFallback(FChooserColumnBase& SourceColumn) override { bFallbackValue = static_cast<FOutputBoolColumn&>(SourceColumn).bFallbackValue; }
 #endif
-		
-	bool& GetValueForIndex(int32 Index)
-	{
-		return Index == ChooserColumn_SpecialIndex_Fallback ? bFallbackValue : RowValues[Index];
-	}
-	
-	bool GetValueForIndex(int32 Index) const
-	{
-		return Index == ChooserColumn_SpecialIndex_Fallback ? bFallbackValue : RowValues[Index];
-	}
 	
 	// FallbackValue will be used as the output value if the all rows in the chooser fail, and the FallbackResult from the chooser is used.
 	UPROPERTY(EditAnywhere, Category=Data);
diff --git a/Engine/Plugins/Experimental/Chooser/Source/Chooser/Internal/OutputEnumColumn.h b/Engine/Plugins/Experimental/Chooser/Source/Chooser/Internal/OutputEnumColumn.h
--- a/Engine/Plugins/Experimental/Chooser/Source/Chooser/Internal/OutputEnumColumn.h
+++ b/Engine/Plugins/Experimental/Chooser/Source/Chooser/Internal/OutputEnumColumn.h
@@ -24,7 +24,7 @@ struct FChooserOutputEnumRowData
 };
 
 
-USTRUCT(DisplayName = "Output Enum", Meta = (Category = "Output", Tooltip = "A column which writes an Enum value."))
+USTRUCT()
 struct CHOOSER_API FOutputEnumColumn : public FEnumColumnBase
 {
 	GENERATED_BODY()
@@ -68,18 +68,9 @@ public:
 
 	virtual void AddToDetails (FInstancedPropertyBag& PropertyBag, int32 ColumnIndex, int32 RowIndex) override;
 	virtual void SetFromDetails(FInstancedPropertyBag& PropertyBag, int32 ColumnIndex, int32 RowIndex) override;
-	virtual void CopyFallback(FChooserColumnBase& SourceColumn) override { FallbackValue = static_cast<FOutputEnumColumn&>(SourceColumn).FallbackValue; }
 	
 	virtual void EnumChanged(const UEnum* Enum) override;
 #endif
-	FChooserOutputEnumRowData& GetValueForIndex(int32 Index)
-	{
-		return Index == ChooserColumn_SpecialIndex_Fallback ? FallbackValue : RowValues[Index];
-	}
-	const FChooserOutputEnumRowData& GetValueForIndex(int32 Index) const
-	{
-		return Index == ChooserColumn_SpecialIndex_Fallback ? FallbackValue : RowValues[Index];
-	}
 	
 	CHOOSER_COLUMN_BOILERPLATE(FChooserParameterEnumBase);
 };
\ No newline at end of file
diff --git a/Engine/Plugins/Experimental/Chooser/Source/Chooser/Internal/OutputFloatColumn.h b/Engine/Plugins/Experimental/Chooser/Source/Chooser/Internal/OutputFloatColumn.h
--- a/Engine/Plugins/Experimental/Chooser/Source/Chooser/Internal/OutputFloatColumn.h
+++ b/Engine/Plugins/Experimental/Chooser/Source/Chooser/Internal/OutputFloatColumn.h
@@ -8,7 +8,7 @@
 #include "StructUtils/InstancedStruct.h"
 #include "OutputFloatColumn.generated.h"
 
-USTRUCT(DisplayName = "Output Float", Meta = (Category = "Output", Tooltip = "A column which writes a Float value."))
+USTRUCT()
 struct CHOOSER_API FOutputFloatColumn : public FChooserColumnBase
 {
 	GENERATED_BODY()
@@ -26,18 +26,7 @@ struct CHOOSER_API FOutputFloatColumn : public FChooserColumnBase
 
 	virtual void AddToDetails(FInstancedPropertyBag& PropertyBag, int32 ColumnIndex, int32 RowIndex) override;
 	virtual void SetFromDetails(FInstancedPropertyBag& PropertyBag, int32 ColumnIndex, int32 RowIndex) override;
-	virtual void CopyFallback(FChooserColumnBase& SourceColumn) override { FallbackValue = static_cast<FOutputFloatColumn&>(SourceColumn).FallbackValue; }
 #endif
-	
-	double& GetValueForIndex(int32 Index)
-	{
-		return Index == ChooserColumn_SpecialIndex_Fallback ? FallbackValue : RowValues[Index];
-	}
-
-	double GetValueForIndex(int32 Index) const
-   	{
-   		return Index == ChooserColumn_SpecialIndex_Fallback ? FallbackValue : RowValues[Index];
-   	}
 
 	// FallbackValue will be used as the output value if the all rows in the chooser fail, and the FallbackResult from the chooser is used.
 	UPROPERTY(EditAnywhere, Category=Data);
diff --git a/Engine/Plugins/Experimental/Chooser/Source/Chooser/Internal/OutputObjectColumn.h b/Engine/Plugins/Experimental/Chooser/Source/Chooser/Internal/OutputObjectColumn.h
--- a/Engine/Plugins/Experimental/Chooser/Source/Chooser/Internal/OutputObjectColumn.h
+++ b/Engine/Plugins/Experimental/Chooser/Source/Chooser/Internal/OutputObjectColumn.h
@@ -11,7 +11,7 @@
 
 struct FBindingChainElement;
 
-USTRUCT(DisplayName = "Output Object", Meta = (Category = "Output", Tooltip = "A column which writes an Object reference."))
+USTRUCT()
 struct FChooserOutputObjectRowData
 {
 	GENERATED_BODY()
@@ -20,7 +20,7 @@ struct FChooserOutputObjectRowData
 	FInstancedStruct Value;
 };
 
-USTRUCT(DisplayName = "Output Object", Meta = (Category = "Output", Tooltip = "A column which an Object Reference."))
+USTRUCT()
 struct CHOOSER_API FOutputObjectColumn : public FChooserColumnBase
 {
 	GENERATED_BODY()
@@ -51,18 +51,7 @@ struct CHOOSER_API FOutputObjectColumn : public FChooserColumnBase
 #if WITH_EDITOR
 	virtual void AddToDetails(FInstancedPropertyBag& PropertyBag, int32 ColumnIndex, int32 RowIndex) override;
 	virtual void SetFromDetails(FInstancedPropertyBag& PropertyBag, int32 ColumnIndex, int32 RowIndex) override;
-	virtual void CopyFallback(FChooserColumnBase& SourceColumn) override { FallbackValue = static_cast<FOutputObjectColumn&>(SourceColumn).FallbackValue; }
 #endif
-	
-	FChooserOutputObjectRowData& GetValueForIndex(int32 Index)
-	{
-		return Index == ChooserColumn_SpecialIndex_Fallback ? FallbackValue : RowValues[Index];
-	}
-	
-	const FChooserOutputObjectRowData& GetValueForIndex(int32 Index) const
-	{
-		return Index == ChooserColumn_SpecialIndex_Fallback ? FallbackValue : RowValues[Index];
-	}
 
 	CHOOSER_COLUMN_BOILERPLATE(FChooserParameterObjectBase);
 
diff --git a/Engine/Plugins/Experimental/Chooser/Source/Chooser/Internal/OutputStructColumn.h b/Engine/Plugins/Experimental/Chooser/Source/Chooser/Internal/OutputStructColumn.h
--- a/Engine/Plugins/Experimental/Chooser/Source/Chooser/Internal/OutputStructColumn.h
+++ b/Engine/Plugins/Experimental/Chooser/Source/Chooser/Internal/OutputStructColumn.h
@@ -27,7 +27,7 @@ struct CHOOSER_API FStructContextProperty : public FChooserParameterStructBase
 #endif
 };
 
-USTRUCT(DisplayName = "Output Struct", Meta = (Category = "Output", Tooltip = "A column which writes all elements of a Struct.\n The data for each row must be set in the details panel, and the table cells will display values which differ from the struct default."))
+USTRUCT()
 struct CHOOSER_API FOutputStructColumn : public FChooserColumnBase
 {
 	GENERATED_BODY()
@@ -46,19 +46,8 @@ struct CHOOSER_API FOutputStructColumn : public FChooserColumnBase
 
 	virtual void AddToDetails(FInstancedPropertyBag& PropertyBag, int32 ColumnIndex, int32 RowIndex) override;
 	virtual void SetFromDetails(FInstancedPropertyBag& PropertyBag, int32 ColumnIndex, int32 RowIndex) override;
-	virtual void CopyFallback(FChooserColumnBase& SourceColumn) override { FallbackValue = static_cast<FOutputStructColumn&>(SourceColumn).FallbackValue; }
 #endif
 	
-	FInstancedStruct& GetValueForIndex(int32 Index)
-	{
-		return Index == ChooserColumn_SpecialIndex_Fallback ? FallbackValue : RowValues[Index];
-	}
-	
-	const FInstancedStruct& GetValueForIndex(int32 Index) const
-	{
-		return Index == ChooserColumn_SpecialIndex_Fallback ? FallbackValue : RowValues[Index];
-	}
-	
 	// FallbackValue will be used as the output value if the all rows in the chooser fail, and the FallbackResult from the chooser is used.
 	UPROPERTY(EditAnywhere, Meta = (StructTypeConst), Category=Data);
    	FInstancedStruct FallbackValue;
diff --git a/Engine/Plugins/Experimental/Chooser/Source/Chooser/Internal/RandomizeColumn.h b/Engine/Plugins/Experimental/Chooser/Source/Chooser/Internal/RandomizeColumn.h
--- a/Engine/Plugins/Experimental/Chooser/Source/Chooser/Internal/RandomizeColumn.h
+++ b/Engine/Plugins/Experimental/Chooser/Source/Chooser/Internal/RandomizeColumn.h
@@ -26,7 +26,7 @@ public:
 	CHOOSER_PARAMETER_BOILERPLATE();
 };
 
-USTRUCT(DisplayName = "Randomize", Meta = (Category = "Random", Tooltip = "The Randomize column will randomly select between whatever values have passed all filters.\n The value specified in each cell is a probability weighting for the row.\n  A row with a value twice as likely as another will be twice as likely to be selected.\n Using the optional RandomizationContext binding, it can track the most recent selection, and reduce the probability of randomly picking the same entry twice"))
+USTRUCT()
 struct CHOOSER_API FRandomizeColumn : public FChooserColumnBase
 {
 	GENERATED_BODY()
diff --git a/Engine/Plugins/Experimental/Chooser/Source/Chooser/Private/OutputBoolColumn.cpp b/Engine/Plugins/Experimental/Chooser/Source/Chooser/Private/OutputBoolColumn.cpp
--- a/Engine/Plugins/Experimental/Chooser/Source/Chooser/Private/OutputBoolColumn.cpp
+++ b/Engine/Plugins/Experimental/Chooser/Source/Chooser/Private/OutputBoolColumn.cpp
@@ -15,13 +15,19 @@ void FOutputBoolColumn::SetOutputs(FChooserEvaluationContext& Context, int RowIn
 {
 	if (InputValue.IsValid())
 	{
-		InputValue.Get<FChooserParameterBoolBase>().SetValue(Context, GetValueForIndex(RowIndex));
+		bool bOutputValue = bFallbackValue;
+		if (RowValues.IsValidIndex(RowIndex))
+		{
+			bOutputValue = RowValues[RowIndex];
+		}
+		
+		InputValue.Get<FChooserParameterBoolBase>().SetValue(Context, bOutputValue);
 	}
 	
 #if WITH_EDITOR
 	if (Context.DebuggingInfo.bCurrentDebugTarget)
 	{
-		TestValue = GetValueForIndex(RowIndex);
+		TestValue = RowValues[RowIndex];
 	}
 #endif
 }
@@ -35,7 +41,7 @@ void FOutputBoolColumn::SetOutputs(FChooserEvaluationContext& Context, int RowIn
 		FPropertyBagPropertyDesc PropertyDesc(PropertyName, EPropertyBagPropertyType::Bool);
 		PropertyDesc.MetaData.Add(FPropertyBagPropertyDescMetaData("DisplayName", DisplayName.ToString()));
 		PropertyBag.AddProperties({PropertyDesc});
-		PropertyBag.SetValueFloat(PropertyName, GetValueForIndex(RowIndex));
+		PropertyBag.SetValueFloat(PropertyName, RowValues[RowIndex]);
 	}
 
 	void FOutputBoolColumn::SetFromDetails(FInstancedPropertyBag& PropertyBag, int32 ColumnIndex, int32 RowIndex)
@@ -45,7 +51,7 @@ void FOutputBoolColumn::SetOutputs(FChooserEvaluationContext& Context, int RowIn
 		TValueOrError<bool, EPropertyBagResult> Result = PropertyBag.GetValueBool(PropertyName);
 		if (bool* Value = Result.TryGetValue())
 		{
-			GetValueForIndex(RowIndex) = *Value;
+			RowValues[RowIndex] = *Value;
 		}
 	} 
 
diff --git a/Engine/Plugins/Experimental/Chooser/Source/Chooser/Private/OutputEnumColumn.cpp b/Engine/Plugins/Experimental/Chooser/Source/Chooser/Private/OutputEnumColumn.cpp
--- a/Engine/Plugins/Experimental/Chooser/Source/Chooser/Private/OutputEnumColumn.cpp
+++ b/Engine/Plugins/Experimental/Chooser/Source/Chooser/Private/OutputEnumColumn.cpp
@@ -17,13 +17,18 @@ void FOutputEnumColumn::SetOutputs(FChooserEvaluationContext& Context, int RowIn
 {
 	if (InputValue.IsValid())
 	{
-		InputValue.Get<FChooserParameterEnumBase>().SetValue(Context, GetValueForIndex(RowIndex).Value);
+		uint8 OutputValue = FallbackValue.Value;
+		if (RowValues.IsValidIndex(RowIndex))
+		{
+			OutputValue = RowValues[RowIndex].Value;
+		}
+		InputValue.Get<FChooserParameterEnumBase>().SetValue(Context, OutputValue);
 	}
 	
 #if WITH_EDITOR
 	if (Context.DebuggingInfo.bCurrentDebugTarget)
 	{
-		TestValue = GetValueForIndex(RowIndex).Value;
+		TestValue = RowValues[RowIndex].Value;
 	}
 #endif
 }
@@ -70,7 +75,7 @@ void FOutputEnumColumn::PostLoad()
 
 #if WITH_EDITOR
 
-	void FOutputEnumColumn::AddToDetails(FInstancedPropertyBag& PropertyBag, int32 ColumnIndex, int32 RowIndex)
+	void FOutputEnumColumn::AddToDetails (FInstancedPropertyBag& PropertyBag, int32 ColumnIndex, int32 RowIndex)
 	{
 		FText DisplayName;
 		InputValue.Get<FChooserParameterBase>().GetDisplayName(DisplayName);
@@ -81,7 +86,7 @@ void FOutputEnumColumn::PostLoad()
 		FPropertyBagPropertyDesc PropertyDesc(PropertyName,  EPropertyBagPropertyType::Enum, InputValue.Get<FChooserParameterEnumBase>().GetEnum());
 		PropertyDesc.MetaData.Add(FPropertyBagPropertyDescMetaData("DisplayName", DisplayName.ToString()));
 		PropertyBag.AddProperties({PropertyDesc});
-		PropertyBag.SetValueEnum(PropertyName, GetValueForIndex(RowIndex).Value, InputValue.Get<FChooserParameterEnumBase>().GetEnum());
+		PropertyBag.SetValueEnum("Value", RowValues[RowIndex].Value, InputValue.Get<FChooserParameterEnumBase>().GetEnum());
 	}
 	
 	void FOutputEnumColumn::SetFromDetails(FInstancedPropertyBag& PropertyBag, int32 ColumnIndex, int32 RowIndex)
@@ -91,7 +96,7 @@ void FOutputEnumColumn::PostLoad()
 		TValueOrError<uint8, EPropertyBagResult> Result = PropertyBag.GetValueEnum(PropertyName, InputValue.Get<FChooserParameterEnumBase>().GetEnum());
 		if (uint8* Value = Result.TryGetValue())
 		{
-			GetValueForIndex(RowIndex).Value = *Value;
+			RowValues[RowIndex].Value = *Value;
 		}
 	}
 
diff --git a/Engine/Plugins/Experimental/Chooser/Source/Chooser/Private/OutputFloatColumn .cpp b/Engine/Plugins/Experimental/Chooser/Source/Chooser/Private/OutputFloatColumn .cpp
--- a/Engine/Plugins/Experimental/Chooser/Source/Chooser/Private/OutputFloatColumn .cpp	
+++ b/Engine/Plugins/Experimental/Chooser/Source/Chooser/Private/OutputFloatColumn .cpp	
@@ -15,13 +15,19 @@ void FOutputFloatColumn::SetOutputs(FChooserEvaluationContext& Context, int RowI
 {
 	if (InputValue.IsValid())
 	{
-		InputValue.Get<FChooserParameterFloatBase>().SetValue(Context, GetValueForIndex(RowIndex));
+		double OutputValue = FallbackValue;
+		if (RowValues.IsValidIndex(RowIndex))
+		{
+			OutputValue = RowValues[RowIndex];
+		}
+	
+		InputValue.Get<FChooserParameterFloatBase>().SetValue(Context, OutputValue);
 	}
 	
 #if WITH_EDITOR
 	if (Context.DebuggingInfo.bCurrentDebugTarget)
 	{
-		TestValue = GetValueForIndex(RowIndex);
+		TestValue = RowValues[RowIndex];
 	}
 #endif
 }
@@ -34,7 +40,7 @@ void FOutputFloatColumn::SetOutputs(FChooserEvaluationContext& Context, int RowI
 		FPropertyBagPropertyDesc PropertyDesc(PropertyName, EPropertyBagPropertyType::Float);
 		PropertyDesc.MetaData.Add(FPropertyBagPropertyDescMetaData("DisplayName", DisplayName.ToString()));
 		PropertyBag.AddProperties({PropertyDesc});
-		PropertyBag.SetValueFloat(PropertyName, GetValueForIndex(RowIndex));
+		PropertyBag.SetValueFloat(PropertyName, RowValues[RowIndex]);
 	}
 
 	void FOutputFloatColumn::SetFromDetails(FInstancedPropertyBag& PropertyBag, int32 ColumnIndex, int32 RowIndex)
@@ -44,7 +50,7 @@ void FOutputFloatColumn::SetOutputs(FChooserEvaluationContext& Context, int RowI
 		TValueOrError<float, EPropertyBagResult> Result = PropertyBag.GetValueFloat(PropertyName);
 		if (float* Value = Result.TryGetValue())
 		{
-			GetValueForIndex(RowIndex) = *Value;
+			RowValues[RowIndex] = *Value;
 		}
 	}
 #endif
\ No newline at end of file
diff --git a/Engine/Plugins/Experimental/Chooser/Source/Chooser/Private/OutputObjectColumn.cpp b/Engine/Plugins/Experimental/Chooser/Source/Chooser/Private/OutputObjectColumn.cpp
--- a/Engine/Plugins/Experimental/Chooser/Source/Chooser/Private/OutputObjectColumn.cpp
+++ b/Engine/Plugins/Experimental/Chooser/Source/Chooser/Private/OutputObjectColumn.cpp
@@ -32,10 +32,21 @@ void FOutputObjectColumn::Compile(IHasContextClass* Owner, bool bForce)
 
 void FOutputObjectColumn::SetOutputs(FChooserEvaluationContext& Context, int RowIndex) const
 {
-	if (const FObjectChooserBase* Value = GetValueForIndex(RowIndex).Value.GetPtr<FObjectChooserBase>())
+	if (RowValues.IsValidIndex(RowIndex))
 	{
-		UObject* Result = Value->ChooseObject(Context);
-		InputValue.Get<FChooserParameterObjectBase>().SetValue(Context, Result);
+		if (const FObjectChooserBase* Value = RowValues[RowIndex].Value.GetPtr<FObjectChooserBase>())
+		{
+			UObject* Result = Value->ChooseObject(Context);
+			InputValue.Get<FChooserParameterObjectBase>().SetValue(Context, Result);
+		}
+	}
+	else
+	{
+		if (const FObjectChooserBase* Value = FallbackValue.Value.GetPtr<FObjectChooserBase>())
+		{
+			UObject* Result = Value->ChooseObject(Context);
+			InputValue.Get<FChooserParameterObjectBase>().SetValue(Context, Result);
+		}
 	}
 }
 
@@ -48,8 +59,7 @@ void FOutputObjectColumn::SetOutputs(FChooserEvaluationContext& Context, int Row
 		FPropertyBagPropertyDesc PropertyDesc(PropertyName,  EPropertyBagPropertyType::Struct, FChooserOutputObjectRowData::StaticStruct());
 		PropertyDesc.MetaData.Add(FPropertyBagPropertyDescMetaData("DisplayName", DisplayName.ToString()));
 		PropertyBag.AddProperties({PropertyDesc});
-	
-		PropertyBag.SetValueStruct(PropertyName, GetValueForIndex(RowIndex));
+		PropertyBag.SetValueStruct(PropertyName, RowValues[RowIndex]);
 	}
 
 	void FOutputObjectColumn::SetFromDetails(FInstancedPropertyBag& PropertyBag, int32 ColumnIndex, int32 RowIndex)
@@ -59,7 +69,7 @@ void FOutputObjectColumn::SetOutputs(FChooserEvaluationContext& Context, int Row
 		TValueOrError<FStructView, EPropertyBagResult> Result = PropertyBag.GetValueStruct(PropertyName, FChooserOutputObjectRowData::StaticStruct());
 		if (FStructView* StructView = Result.TryGetValue())
 		{
-			GetValueForIndex(RowIndex) = StructView->Get<FChooserOutputObjectRowData>();
+			RowValues[RowIndex] = StructView->Get<FChooserOutputObjectRowData>();
 		}
 	}
 
diff --git a/Engine/Plugins/Experimental/Chooser/Source/Chooser/Private/OutputStructColumn.cpp b/Engine/Plugins/Experimental/Chooser/Source/Chooser/Private/OutputStructColumn.cpp
--- a/Engine/Plugins/Experimental/Chooser/Source/Chooser/Private/OutputStructColumn.cpp
+++ b/Engine/Plugins/Experimental/Chooser/Source/Chooser/Private/OutputStructColumn.cpp
@@ -67,12 +67,10 @@ void FOutputStructColumn::AddToDetails(FInstancedPropertyBag& PropertyBag, int32
 			FText DisplayName;
 			StructInput.GetDisplayName(DisplayName);
 			FName PropertyName("RowData",ColumnIndex);
-			
-			FInstancedStruct& Value = GetValueForIndex(RowIndex);
-			FPropertyBagPropertyDesc PropertyDesc(PropertyName,  EPropertyBagPropertyType::Struct, Value.GetScriptStruct());
+			FPropertyBagPropertyDesc PropertyDesc(PropertyName,  EPropertyBagPropertyType::Struct, RowValues[RowIndex].GetScriptStruct());
 			PropertyDesc.MetaData.Add(FPropertyBagPropertyDescMetaData("DisplayName", DisplayName.ToString()));
 			PropertyBag.AddProperties({PropertyDesc});
-			PropertyBag.SetValueStruct(PropertyName, FConstStructView(Value.GetScriptStruct(), Value.GetMemory()));
+			PropertyBag.SetValueStruct(PropertyName, FConstStructView(RowValues[RowIndex].GetScriptStruct(), RowValues[RowIndex].GetMemory()));
 		}
 	}
 }
@@ -86,12 +84,10 @@ void FOutputStructColumn::SetFromDetails(FInstancedPropertyBag& PropertyBag, int
 		{
 			FName PropertyName("RowData", ColumnIndex);
 
-			FInstancedStruct& Value = GetValueForIndex(RowIndex);
-
-			TValueOrError<FStructView, EPropertyBagResult> Result = PropertyBag.GetValueStruct(PropertyName, Value.GetScriptStruct());
+			TValueOrError<FStructView, EPropertyBagResult> Result = PropertyBag.GetValueStruct(PropertyName, RowValues[RowIndex].GetScriptStruct());
 			if (FStructView* StructView = Result.TryGetValue())
 			{
-				RowValues[RowIndex].GetScriptStruct()->CopyScriptStruct(Value.GetMutableMemory(), StructView->GetMemory());
+				RowValues[RowIndex].GetScriptStruct()->CopyScriptStruct(RowValues[RowIndex].GetMutableMemory(), StructView->GetMemory());
 			}
 		}
 	}
@@ -108,18 +104,22 @@ void FOutputStructColumn::SetOutputs(FChooserEvaluationContext& Context, int Row
 {
 	if (InputValue.IsValid())
 	{
-		const FInstancedStruct& OutputValue = GetValueForIndex(RowIndex);
+		const FInstancedStruct* OutputValue = &FallbackValue;
+		if (RowValues.IsValidIndex(RowIndex))
+		{
+			OutputValue = &RowValues[RowIndex];
+		}
 
-		if (OutputValue.IsValid())
+		if (OutputValue && OutputValue->IsValid())
 		{
-			InputValue.Get<FChooserParameterStructBase>().SetValue(Context, OutputValue);
+			InputValue.Get<FChooserParameterStructBase>().SetValue(Context, *OutputValue);
 		}
 	}
 	
 #if WITH_EDITOR
 	if (Context.DebuggingInfo.bCurrentDebugTarget)
 	{
-		TestValue = GetValueForIndex(RowIndex);
+		TestValue = RowValues[RowIndex];
 	}
 #endif
 }
diff --git a/Engine/Plugins/Experimental/Chooser/Source/Chooser/Private/RandomizeColumn.cpp b/Engine/Plugins/Experimental/Chooser/Source/Chooser/Private/RandomizeColumn.cpp
--- a/Engine/Plugins/Experimental/Chooser/Source/Chooser/Private/RandomizeColumn.cpp
+++ b/Engine/Plugins/Experimental/Chooser/Source/Chooser/Private/RandomizeColumn.cpp
@@ -186,14 +186,11 @@ void FRandomizeColumn::SetOutputs(FChooserEvaluationContext& Context, int RowInd
 
 	void FRandomizeColumn::AddToDetails(FInstancedPropertyBag& PropertyBag, int32 ColumnIndex, int32 RowIndex)
 	{
-		if (RowValues.IsValidIndex(RowIndex))
-		{
-			FName PropertyName("RowData",ColumnIndex);
-			FPropertyBagPropertyDesc PropertyDesc(PropertyName, EPropertyBagPropertyType::Float);
-			PropertyDesc.MetaData.Add(FPropertyBagPropertyDescMetaData("DisplayName", "Randomize"));
-			PropertyBag.AddProperties({PropertyDesc});
-			PropertyBag.SetValueFloat(PropertyName, RowValues[RowIndex]);
-		}
+		FName PropertyName("RowData",ColumnIndex);
+		FPropertyBagPropertyDesc PropertyDesc(PropertyName, EPropertyBagPropertyType::Float);
+		PropertyDesc.MetaData.Add(FPropertyBagPropertyDescMetaData("DisplayName", "Randomize"));
+		PropertyBag.AddProperties({PropertyDesc});
+		PropertyBag.SetValueFloat(PropertyName, RowValues[RowIndex]);
 	}
 
 	void FRandomizeColumn::SetFromDetails(FInstancedPropertyBag& PropertyBag, int32 ColumnIndex, int32 RowIndex)
diff --git a/Engine/Plugins/Experimental/Chooser/Source/Chooser/Public/IChooserColumn.h b/Engine/Plugins/Experimental/Chooser/Source/Chooser/Public/IChooserColumn.h
--- a/Engine/Plugins/Experimental/Chooser/Source/Chooser/Public/IChooserColumn.h
+++ b/Engine/Plugins/Experimental/Chooser/Source/Chooser/Public/IChooserColumn.h
@@ -69,7 +69,6 @@ public:
 	virtual void MoveRow(int SourceIndex, int TargetIndex) {}
 	virtual void InsertRows(int Index, int Count) {}
 	virtual void CopyRow(FChooserColumnBase& SourceColumn, int SourceIndex, int TargetIndex) {}
-	virtual void CopyFallback(FChooserColumnBase& SourceColumn) {}
 	
 	virtual UScriptStruct* GetInputBaseType() const { return nullptr; };
 	virtual const UScriptStruct* GetInputType() const { return nullptr; };
@@ -95,7 +94,6 @@ public:
 #endif
 };
 
-enum {  ChooserColumn_SpecialIndex_Fallback = -2 };
 
 #if WITH_EDITOR
 #define CHOOSER_COLUMN_BOILERPLATE2(ParameterType, RowValuesProperty) \
diff --git a/Engine/Plugins/Experimental/Chooser/Source/ChooserEditor/Private/BoolColumnEditor.cpp b/Engine/Plugins/Experimental/Chooser/Source/ChooserEditor/Private/BoolColumnEditor.cpp
--- a/Engine/Plugins/Experimental/Chooser/Source/ChooserEditor/Private/BoolColumnEditor.cpp
+++ b/Engine/Plugins/Experimental/Chooser/Source/ChooserEditor/Private/BoolColumnEditor.cpp
@@ -120,13 +120,8 @@ TSharedRef<SWidget> CreateOutputBoolColumnWidget(UChooserTable* Chooser, FChoose
 		+ SHorizontalBox::Slot().AutoWidth().VAlign(VAlign_Center)
 		[
 			SNew(SCheckBox)
-			.IsChecked_Lambda([ BoolColumn]() { return BoolColumn->bFallbackValue ? ECheckBoxState::Checked : ECheckBoxState::Unchecked; })
-			.OnCheckStateChanged_Lambda([Chooser, BoolColumn](ECheckBoxState State)
-			{
-				const FScopedTransaction Transaction(LOCTEXT("Change Bool Value", "Change Bool Value"));
-				Chooser->Modify(true);
-				BoolColumn->bFallbackValue = State == ECheckBoxState::Checked;
-			})
+			.IsChecked_Lambda([BoolColumn]() { return BoolColumn->bFallbackValue ? ECheckBoxState::Checked : ECheckBoxState::Unchecked; })
+			.OnCheckStateChanged_Lambda([BoolColumn](ECheckBoxState State) { BoolColumn->bFallbackValue = State == ECheckBoxState::Checked; })
 		]
 		+ SHorizontalBox::Slot().FillWidth(1);
 	}
diff --git a/Engine/Plugins/Experimental/Chooser/Source/ChooserEditor/Private/ChooserTableEditor.cpp b/Engine/Plugins/Experimental/Chooser/Source/ChooserEditor/Private/ChooserTableEditor.cpp
--- a/Engine/Plugins/Experimental/Chooser/Source/ChooserEditor/Private/ChooserTableEditor.cpp
+++ b/Engine/Plugins/Experimental/Chooser/Source/ChooserEditor/Private/ChooserTableEditor.cpp
@@ -50,7 +50,6 @@
 #include "Exporters/Exporter.h"
 #include "HAL/PlatformApplicationMisc.h"
 #include "Framework/Commands/GenericCommands.h"
-#include "SPositiveActionButton.h"
 
 #define LOCTEXT_NAMESPACE "ChooserEditor"
 
@@ -595,42 +594,23 @@ void FChooserTableEditor::OnObjectsTransacted(UObject* Object, const FTransactio
 		// if this is for the chooser we're editing
 		if (GetChooser() == RowDetails->Chooser)
 		{
-			if (RowDetails->Chooser->ResultsStructs.IsValidIndex(RowDetails->Row))
+			// copy all the values over
+			TValueOrError<FStructView, EPropertyBagResult> Result = RowDetails->Properties.GetValueStruct("Result", FInstancedStruct::StaticStruct());
+			if (Result.IsValid())
 			{
-				// copy all the values over
-				TValueOrError<FStructView, EPropertyBagResult> Result = RowDetails->Properties.GetValueStruct("Result", FInstancedStruct::StaticStruct());
-				if (Result.IsValid())
-				{
-					RowDetails->Chooser->ResultsStructs[RowDetails->Row] = Result.GetValue().Get<FInstancedStruct>();
-				}
-
-				int ColumnIndex = 0;
-				for (FInstancedStruct& ColumnData : RowDetails->Chooser->ColumnsStructs)
-				{
-					FChooserColumnBase& Column = ColumnData.GetMutable<FChooserColumnBase>();
-					Column.SetFromDetails(RowDetails->Properties, ColumnIndex, RowDetails->Row);
-					ColumnIndex++;
-				}
-			
-				TValueOrError<bool, EPropertyBagResult> DisabledResult = RowDetails->Properties.GetValueBool("Disabled");
-				RowDetails->Chooser->DisabledRows[RowDetails->Row] = DisabledResult.GetValue();
+				RowDetails->Chooser->ResultsStructs[RowDetails->Row] = Result.GetValue().Get<FInstancedStruct>();
 			}
-			else if (RowDetails->Row == ColumnWidget_SpecialIndex_Fallback)
+
+			int ColumnIndex = 0;
+			for (FInstancedStruct& ColumnData : RowDetails->Chooser->ColumnsStructs)
 			{
-				TValueOrError<FStructView, EPropertyBagResult> Result = RowDetails->Properties.GetValueStruct("Result", FInstancedStruct::StaticStruct());
-				if (Result.IsValid())
-				{
-					RowDetails->Chooser->FallbackResult = Result.GetValue().Get<FInstancedStruct>();
-				}
-				
-				int ColumnIndex = 0;
-				for (FInstancedStruct& ColumnData : RowDetails->Chooser->ColumnsStructs)
-				{
-					FChooserColumnBase& Column = ColumnData.GetMutable<FChooserColumnBase>();
-					Column.SetFromDetails(RowDetails->Properties, ColumnIndex, RowDetails->Row);
-					ColumnIndex++;
-				}
+				FChooserColumnBase& Column = ColumnData.GetMutable<FChooserColumnBase>();
+				Column.SetFromDetails(RowDetails->Properties, ColumnIndex, RowDetails->Row);
+				ColumnIndex++;
 			}
+			
+			TValueOrError<bool, EPropertyBagResult> DisabledResult = RowDetails->Properties.GetValueBool("Disabled");
+			RowDetails->Chooser->DisabledRows[RowDetails->Row] = DisabledResult.GetValue();
 
 			RefreshAll();
 		}
@@ -1018,18 +998,15 @@ int FChooserTableEditor::MoveRow(int SourceRowIndex, int TargetRowIndex)
 	
 void FChooserTableEditor::SelectRow(int32 RowIndex, bool bClear)
 {
-	if (TSharedPtr<FChooserTableRow>* Row = TableRows.FindByPredicate([RowIndex](const TSharedPtr<FChooserTableRow>& InRow)
-		{
-			return InRow->RowIndex == RowIndex;
-		}))
+	if (TableRows.IsValidIndex(RowIndex))
 	{
-		if (!TableView->IsItemSelected(*Row))
+		if (!TableView->IsItemSelected(TableRows[RowIndex]))
 		{
 			if (bClear)
 			{
 				TableView->ClearSelection();
 			}
-			TableView->SetItemSelection(*Row, true, ESelectInfo::OnMouseClick);
+			TableView->SetItemSelection(TableRows[RowIndex], true, ESelectInfo::OnMouseClick);
 		}
 	}
 }
@@ -1118,14 +1095,10 @@ void FChooserTableEditor::UpdateTableColumns()
 		.FillWidth(1.0)
 		.HeaderContent( )
 		[
-			SNew(SVerticalBox)
-			 + SVerticalBox::Slot().AutoHeight()
-			 [
-				SNew(SHorizontalBox)
-				+ SHorizontalBox::Slot().MaxWidth(150)
-				[
-					CreateColumnComboButton.ToSharedRef()
-				]
+			SNew(SHorizontalBox)
+			+ SHorizontalBox::Slot().MaxWidth(120)
+			[
+				CreateColumnComboButton.ToSharedRef()
 			]
 		]
 		);
@@ -1134,7 +1107,7 @@ void FChooserTableEditor::UpdateTableColumns()
 
 void FChooserTableEditor::AddColumn(const UScriptStruct* ColumnType)
 {
-	FSlateApplication::Get().DismissAllMenus();
+	CreateColumnComboButton->SetIsOpen(false);
 	UChooserTable* Chooser = GetChooser();
 	const FScopedTransaction Transaction(LOCTEXT("Add Column Transaction", "Add Column"));
 	Chooser->Modify(true);
@@ -1226,30 +1199,7 @@ void FChooserTableEditor::RefreshRowSelectionDetails()
 			}
 
 			SelectedRows.Add(Selection);
-		}
-		else if (SelectedItem->RowIndex == ColumnWidget_SpecialIndex_Fallback)
-		{
-			TObjectPtr<UChooserRowDetails> Selection = NewObject<UChooserRowDetails>();
-			Selection->Chooser = Chooser;
-			Selection->Row = SelectedItem->RowIndex;
-			Selection->SetFlags(RF_Standalone | RF_Transactional);
-
-			FInstancedStruct& Result = Chooser->FallbackResult;
-			Selection->Properties.AddProperties({ResultPropertyDesc});
-			Selection->Properties.SetValueStruct("Result", FConstStructView(FInstancedStruct::StaticStruct(), reinterpret_cast<uint8*>(&Result)));
-
-			int ColumnIndex = 0;
-			for (FInstancedStruct& ColumnData : Chooser->ColumnsStructs)
-			{
-				FChooserColumnBase& Column = ColumnData.GetMutable<FChooserColumnBase>();
-				if (Column.HasOutputs())
-				{
-					Column.AddToDetails(Selection->Properties, ColumnIndex, SelectedItem->RowIndex);
-				}
-				ColumnIndex++;
-			}
-
-			SelectedRows.Add(Selection);
+			
 		}
 	}
 	
@@ -1310,225 +1260,6 @@ TSharedRef<SDockTab> FChooserTableEditor::SpawnNestedTablesTreeTab( const FSpawn
 		];
 }
 
-DECLARE_DELEGATE_OneParam(FCreateStructDelegate, UScriptStruct*);
-
-struct FColumnTypeInfoStruct
-{
-	int SortOrder = 100;
-	FString Category;
-	UScriptStruct* Type;
-	bool operator < (const FColumnTypeInfoStruct& Other) const
-	{
-		if (Category == Other.Category)
-		{
-			return Type->GetDisplayNameText().ToString() < Other.Type->GetDisplayNameText().ToString();
-		}
-		else if (SortOrder == Other.SortOrder)
-		{
-			return Category < Other.Category;
-		}
-		else
-		{
-			return SortOrder < Other.SortOrder;
-		}
-	}
-};
-
-TSharedRef<SWidget>	FChooserTableEditor::MakeCreateColumnMenu()
-{
-	FMenuBuilder MenuBuilder(true, nullptr);
-	static TArray<FColumnTypeInfoStruct> ColumnTypes;
-
-	if (ColumnTypes.IsEmpty())
-	{
-		UScriptStruct* BaseType = FChooserColumnBase::StaticStruct();
-		for (TObjectIterator<UScriptStruct> StructIt; StructIt; ++StructIt)
-		{
-			if (*StructIt != BaseType && StructIt->IsChildOf(BaseType))
-			{
-				if (!StructIt->HasMetaData("Hidden"))
-				{
-					FColumnTypeInfoStruct Info;
-					Info.Type = *StructIt;
-					Info.Category = StructIt->HasMetaData("Category") ? StructIt->GetMetaData("Category") : "Other";
-
-					if (Info.Category == "Filter")
-					{
-						Info.SortOrder = 1;
-					}
-					else if (Info.Category == "Scoring")
-                   	{
-                   		Info.SortOrder = 2;
-                   	}
-					else if (Info.Category == "Output")
-                   	{
-                   		Info.SortOrder = 3;
-                   	}
-					else if (Info.Category == "Random")
-					{
-						Info.SortOrder = 4;
-					}
-
-					ColumnTypes.Add(Info);
-				}
-			}
-		}
-		ColumnTypes.Sort();
-	}
-
-	FString Section = "";
-	for(FColumnTypeInfoStruct& Type : ColumnTypes)
-	{
-		if (Section != Type.Category)
-		{
-			if (Section != "")
-			{
-				MenuBuilder.EndSection();
-			}
-			Section = Type.Category;
-			MenuBuilder.BeginSection(FName(Section), FText::FromString(Section));
-		}
-		
-		MenuBuilder.AddMenuEntry(Type.Type->GetDisplayNameText(), Type.Type->GetToolTipText(), FSlateIcon(),
-			FUIAction(
-				FExecuteAction::CreateLambda([this, Type]()
-				{
-					AddColumn(Type.Type);
-				}))
-			);
-							
-	}
-	return MenuBuilder.MakeWidget();
-}
-
-struct FResultTypeInfoStruct
-{
-	bool ObjectOnly = false;
-	bool ClassOnly = false;
-	FString Category;
-	UScriptStruct* Type = nullptr;
-	
-	bool operator < (const FResultTypeInfoStruct& Other) const
-	{
-		if (Category == Other.Category)
-		{
-			return Type->GetDisplayNameText().ToString() < Other.Type->GetDisplayNameText().ToString();
-		}
-		else
-		{
-			return Category < Other.Category;
-		}
-	}
-};
-
-void MakeCreateResultMenu(FMenuBuilder& MenuBuilder, EObjectChooserResultType ChooserResultType, FCreateStructDelegate CreateStruct)
-{
-	static TArray<FResultTypeInfoStruct> ResultTypes;
-
-	if (ResultTypes.IsEmpty())
-	{
-		UScriptStruct* BaseType = FObjectChooserBase::StaticStruct();
-		for (TObjectIterator<UScriptStruct> StructIt; StructIt; ++StructIt)
-		{
-			if (*StructIt != BaseType && StructIt->IsChildOf(BaseType))
-			{
-				if (!StructIt->HasMetaData("Hidden"))
-				{
-					FResultTypeInfoStruct Info;
-					Info.Type = *StructIt;
-					Info.Category = StructIt->HasMetaData("Category") ? StructIt->GetMetaData("Category") : "Other";
-
-					if (StructIt->HasMetaData("ResultType"))
-					{
-						FString ResultTypeString = StructIt->GetMetaData("ResultType");
-						Info.ClassOnly = ResultTypeString == "Class";
-						Info.ObjectOnly = ResultTypeString == "Object";
-					}
-					
-					ResultTypes.Add(Info);
-				}
-			}
-		}
-		ResultTypes.Sort();
-	}
-
-	FString Section = "";
-	for(FResultTypeInfoStruct& Type : ResultTypes)
-	{
-		if (Section != Type.Category)
-		{
-			if (Section != "")
-			{
-				MenuBuilder.EndSection();
-			}
-			Section = Type.Category;
-			MenuBuilder.BeginSection(FName(Section), FText::FromString(Section));
-		}
-		
-		MenuBuilder.AddMenuEntry(Type.Type->GetDisplayNameText(), Type.Type->GetToolTipText(), FSlateIcon(),
-			FUIAction(
-				FExecuteAction::CreateLambda([Type, CreateStruct]()
-				{
-					CreateStruct.Execute(Type.Type);
-				}),
-				FCanExecuteAction::CreateLambda([Type, ChooserResultType]()
-				{
-					if (Type.ClassOnly && ChooserResultType == EObjectChooserResultType::ObjectResult)
-					{
-						return false;
-					}
-					if (Type.ObjectOnly && ChooserResultType == EObjectChooserResultType::ClassResult)
-					{
-						return false;
-					}
-					return true;
-				})
-				)
-			);
-							
-	}
-}
-
-TSharedRef<SWidget>	FChooserTableEditor::MakeCreateRowMenu()
-{
-	FMenuBuilder MenuBuilder(true, nullptr);
-
-    UChooserTable* Chooser = GetChooser();
-	if (!Chooser->FallbackResult.IsValid())
-	{
-		MenuBuilder.AddSubMenu(
-			LOCTEXT("Add Fallback", "Add Fallback Result"),
-				LOCTEXT("Add Fallback Tooltip", "Add a Fallback row to the chooser, which will be used in the case where no other rows passed all filter columns"),
-				FNewMenuDelegate::CreateLambda([this](FMenuBuilder& MenuBuilder)
-				{
-					UChooserTable* Chooser = GetChooser();
-					MakeCreateResultMenu(MenuBuilder, Chooser->ResultType, FCreateStructDelegate::CreateLambda([this](UScriptStruct* Type)
-					{
-						UChooserTable* Chooser = GetChooser();
-						const FScopedTransaction Transaction(LOCTEXT("Add Fallback Row Transaction", "Add Fallback Row"));
-						Chooser->Modify(true);
-						Chooser->FallbackResult.InitializeAs(Type);
-						UpdateTableRows();
-					}));
-				})
-			);
-
-	}
-
-	MakeCreateResultMenu(MenuBuilder, Chooser->ResultType, FCreateStructDelegate::CreateLambda([this](UScriptStruct* Type)
-	{
-
-		UChooserTable* Chooser = GetChooser();
-		const FScopedTransaction Transaction(LOCTEXT("Add Row Transaction", "Add Row"));
-		Chooser->Modify(true);
-		Chooser->ResultsStructs.SetNum(Chooser->ResultsStructs.Num()+1);
-		Chooser->ResultsStructs.Last().InitializeAs(Type);
-		UpdateTableRows();
-	}));
-
-	return MenuBuilder.MakeWidget();
-}
-
 TSharedRef<SDockTab> FChooserTableEditor::SpawnTableTab( const FSpawnTabArgs& Args )
 {
 	check( Args.GetTabId() == TableTabId );
@@ -1537,17 +1268,55 @@ TSharedRef<SDockTab> FChooserTableEditor::SpawnTableTab( const FSpawnTabArgs& Ar
 
 	// + button to create new columns
 	
-	CreateColumnComboButton = SNew(SPositiveActionButton)
-		.Text(LOCTEXT("Add Column", "Add Column"))
+	CreateColumnComboButton = SNew(SComboButton).OnGetMenuContent_Lambda([this]()
+	{
+		FStructViewerInitializationOptions Options;
+		Options.StructFilter = MakeShared<FStructFilter>(FChooserColumnBase::StaticStruct());
+		Options.bAllowViewOptions = false;
+		Options.bShowNoneOption = false;
+		Options.NameTypeToDisplay = EStructViewerNameTypeToDisplay::DisplayName;
+		
+		// Add class filter for columns here
+		FStructViewerModule& StructViewerModule = FModuleManager::LoadModuleChecked<FStructViewerModule>("StructViewer");
+		TSharedRef<SWidget> Widget = StructViewerModule.CreateStructViewer(Options, FOnStructPicked::CreateRaw(this, &FChooserTableEditor::AddColumn));
+		return Widget;
+	})
+	.ComboButtonStyle(FAppStyle::Get(), "SimpleComboButton")
+	.ButtonContent()
+	[
+		SNew(STextBlock).Text(LOCTEXT("AddColumn", "+ Add Column"))
+	];
+
+
+	CreateRowComboButton = SNew(SComboButton)
+		.ComboButtonStyle(FAppStyle::Get(), "SimpleComboButton")
+		.ButtonContent()
+		[
+			SNew(STextBlock).Text(LOCTEXT("AddRow", "+ Add Row"))
+		]
 		.OnGetMenuContent_Lambda([this]()
 		{
-			return MakeCreateColumnMenu();
+			FStructViewerInitializationOptions Options;
+			Options.StructFilter = MakeShared<FStructFilter>(FObjectChooserBase::StaticStruct());
+			Options.bAllowViewOptions = false;
+			Options.bShowNoneOption = false;
+			Options.NameTypeToDisplay = EStructViewerNameTypeToDisplay::DisplayName;
+			
+			TSharedRef<SWidget> Widget = FModuleManager::LoadModuleChecked<FStructViewerModule>("StructViewer").CreateStructViewer(Options, FOnStructPicked::CreateLambda([this](const UScriptStruct* ChosenStruct)
+			{
+				CreateRowComboButton->SetIsOpen(false);
+				UChooserTable* Chooser = GetChooser();
+				const FScopedTransaction Transaction(LOCTEXT("Add Row Transaction", "Add Row"));
+				Chooser->Modify(true);
+
+				Chooser->ResultsStructs.SetNum(Chooser->ResultsStructs.Num()+1);
+				Chooser->ResultsStructs.Last().InitializeAs(ChosenStruct);
+				UpdateTableRows();
+			}));
+			
+			return Widget;
 		});
 
-	CreateRowComboButton = SNew(SPositiveActionButton)
-		.Text(LOCTEXT("Add Row", "Add Row"))
-		.OnGetMenuContent(this, &FChooserTableEditor::MakeCreateRowMenu);
-			
 	HeaderRow = SNew(SHeaderRow);
 
 	UpdateTableRows();
@@ -1657,11 +1426,7 @@ void FChooserTableEditor::UpdateTableRows()
 	}
 
 	// Add one at the end, for the Fallback result
-	if (Chooser->FallbackResult.IsValid())
-	{
-		TableRows.Add(MakeShared<FChooserTableRow>(SChooserTableRow::SpecialIndex_Fallback));
-	}
-	
+	TableRows.Add(MakeShared<FChooserTableRow>(SChooserTableRow::SpecialIndex_Fallback));
 	// Add one at the end, for the "Add Row" control
 	TableRows.Add(MakeShared<FChooserTableRow>(SChooserTableRow::SpecialIndex_AddRow));
 
@@ -1781,14 +1546,7 @@ int32 FChooserTableEditor::DeleteSelectedRowsInternal(int32 RowIndexToRemember)
 	TArray<uint32> RowsToDelete;
 	for(auto& SelectedRow:SelectedRows)
 	{
-		if (SelectedRow->Row == ColumnWidget_SpecialIndex_Fallback)
-		{
-			Chooser->FallbackResult.Reset();
-		}
-		else
-		{
-			RowsToDelete.Add(SelectedRow->Row);
-		}
+		RowsToDelete.Add(SelectedRow->Row);
 	}
 
 	SelectedRows.SetNum(0);
@@ -1996,10 +1754,7 @@ void FChooserTableEditor::ToggleDisableSelection()
 		{
 			for (auto& Row : SelectedRows)
 			{
-				if (Chooser->DisabledRows.IsValidIndex(Row->Row))
-				{
-					Chooser->DisabledRows[Row->Row] = !bDisabled;
-				}
+				Chooser->DisabledRows[Row->Row] = !bDisabled;
 			}
 			RefreshRowSelectionDetails();
 		}
@@ -2042,32 +1797,17 @@ void FChooserTableEditor::DuplicateSelection()
 	}
 }
 
-bool FChooserTableEditor::HasFallbackSelected()
-{
-	for(UChooserRowDetails* SelectedRow : SelectedRows)
-	{
-		if (SelectedRow->Row == ColumnWidget_SpecialIndex_Fallback)
-		{
-			return true;
-		}
-	}
-
-	return false;
-}
 
 bool FChooserTableEditor::CanMoveRowsUp()
 {
 	if (HasRowsSelected())
 	{
 		UChooserTable* Chooser = GetChooser();
-
+		
 		int MinSelectedRow = Chooser->ResultsStructs.Num();
 		for(UChooserRowDetails* SelectedRow : SelectedRows)
 		{
-			if (SelectedRow->Row != ColumnWidget_SpecialIndex_Fallback)
-			{
-				MinSelectedRow = FMath::Min(SelectedRow->Row, MinSelectedRow);
-			}
+			MinSelectedRow = FMath::Min(SelectedRow->Row, MinSelectedRow);
 		}
 
 		return MinSelectedRow > 0;
@@ -2083,10 +1823,7 @@ void FChooserTableEditor::MoveRowsUp()
 		int MinSelectedRow = Chooser->ResultsStructs.Num();
 		for(UChooserRowDetails* SelectedRow : SelectedRows)
 		{
-			if (SelectedRow->Row != ColumnWidget_SpecialIndex_Fallback)
-			{
-				MinSelectedRow = FMath::Min(SelectedRow->Row, MinSelectedRow);
-			}
+			MinSelectedRow = FMath::Min(SelectedRow->Row, MinSelectedRow);
 		}
 		MoveRows(MinSelectedRow - 1);
 	}
@@ -2234,35 +1971,22 @@ UChooserTable* FChooserTableEditor::CopySelectionInternal()
 			Column.SetNumRows(0);
 			Column.SetNumRows(SelectedRowsCopy.Num());
 		}
-		
-		if (SelectedRowsCopy.Num() > 0 && SelectedRowsCopy[0]->Row == ColumnWidget_SpecialIndex_Fallback)
-		{
-			SelectedRowsCopy.RemoveAt(0);
-			
-			CopyData->FallbackResult = Chooser->FallbackResult;
-			if (FNestedChooser* CopiedNestedChooser = CopyData->FallbackResult.GetMutablePtr<FNestedChooser>())
-			{
-				// if the fallback result was a nested chooser, duplicate it
-				CopiedNestedChooser->Chooser = DuplicateNestedChooser(CopiedNestedChooser->Chooser, CopyData);
-			}
-		}
-		
+
 		CopyData->ResultsStructs.SetNum(SelectedRowsCopy.Num());
-		CopyData->DisabledRows.SetNum(SelectedRowsCopy.Num());
 
 		// add the selected results and column data
 		
 		for (int RowIndex = 0; RowIndex < SelectedRowsCopy.Num(); RowIndex++)
 		{
-			CopyData->ResultsStructs[RowIndex] = Chooser->ResultsStructs[SelectedRowsCopy[RowIndex]->Row];
+			const FInstancedStruct& Result = Chooser->ResultsStructs[SelectedRowsCopy[RowIndex]->Row];
+
+			CopyData->ResultsStructs[RowIndex] = Result;
 			if (FNestedChooser* CopiedNestedChooser = CopyData->ResultsStructs[RowIndex].GetMutablePtr<FNestedChooser>())
 			{
 				// if the result for this row was a nested chooser, duplicate it
 				CopiedNestedChooser->Chooser = DuplicateNestedChooser(CopiedNestedChooser->Chooser, CopyData);
 			}
-			
-			CopyData->DisabledRows[RowIndex] = Chooser->DisabledRows[SelectedRowsCopy[RowIndex]->Row];
-			
+
 			for (int ColumnIndex = 0; ColumnIndex < CopyData->ColumnsStructs.Num(); ColumnIndex++)
 			{
 				FChooserColumnBase& SourceColumn = Chooser->ColumnsStructs[ColumnIndex].GetMutable<FChooserColumnBase>();
@@ -2341,7 +2065,7 @@ void FChooserTableEditor::PasteInternal(UChooserTable* PastedContent, int PasteR
 	UChooserTable* Chooser = GetChooser();
 	Chooser->Modify();
 	
-	if (PastedContent->ResultsStructs.IsEmpty() && !PastedContent->FallbackResult.IsValid())
+	if (PastedContent->ResultsStructs.IsEmpty())
 	{
 		// pasting a column
 		int InsertColumnIndex = Chooser->ColumnsStructs.Num();
@@ -2382,40 +2106,11 @@ void FChooserTableEditor::PasteInternal(UChooserTable* PastedContent, int PasteR
 				{
 					InsertIndex = FMath::Max(InsertIndex, SelectedRows[SelectedRowIndex]->Row);
 				}
-				if (InsertIndex == ColumnWidget_SpecialIndex_Fallback)
-				{
-					// if the only row selected was the fallback, reset insert index to the last row
-					InsertIndex = Chooser->ResultsStructs.Num();
-				}
-				else
-				{
-					InsertIndex++;
-				}
+				InsertIndex++;
 			}
 		}
 
-		if (PastedContent->ResultsStructs.Num() > 0)
-		{
-			Chooser->ResultsStructs.Insert(PastedContent->ResultsStructs, InsertIndex);
-			Chooser->DisabledRows.Insert(PastedContent->DisabledRows, InsertIndex);
-
-			// Make sure each column has the same number of row datas as there are results
-			for(FInstancedStruct& ColumnData : Chooser->ColumnsStructs)
-			{
-				FChooserColumnBase& Column = ColumnData.GetMutable<FChooserColumnBase>();
-				Column.InsertRows(InsertIndex, RowsToPaste);
-			}
-		}
-		if (PastedContent->FallbackResult.IsValid())
-		{
-			// paste fallback result if copy data has one
-			Chooser->FallbackResult = PastedContent->FallbackResult;
-			if (FNestedChooser* NestedChooser = Chooser->FallbackResult.GetMutablePtr<FNestedChooser>())
-			{
-				// duplicate the nested chooser if the fallback result refers to a nested chooser
-				NestedChooser->Chooser = DuplicateNestedChooser(NestedChooser->Chooser, Chooser);
-			}
-		}
+		Chooser->ResultsStructs.Insert(PastedContent->ResultsStructs, InsertIndex);
 		
 		if (!PastedContent->NestedChoosers.IsEmpty())
 		{
@@ -2428,6 +2123,14 @@ void FChooserTableEditor::PasteInternal(UChooserTable* PastedContent, int PasteR
 				}
 			}
 		}
+		
+		// Make sure each column has the same number of row datas as there are results
+		for(FInstancedStruct& ColumnData : Chooser->ColumnsStructs)
+		{
+			FChooserColumnBase& Column = ColumnData.GetMutable<FChooserColumnBase>();
+			Column.InsertRows(InsertIndex, RowsToPaste);
+		}
+
 
 		// try to also paste column data from columns in the paste buffer which match the columns in the current chooser
 		// -- matching by column type and input value name
@@ -2467,12 +2170,6 @@ void FChooserTableEditor::PasteInternal(UChooserTable* PastedContent, int PasteR
 								{
 									Column.CopyRow(PastedColumn, i, InsertIndex + i);
 								}
-
-								if (PastedContent->FallbackResult.IsValid())
-								{
-									// if the fallback row was copied, paste the fallback data for columns
-									Column.CopyFallback(PastedColumn);
-								}
 								break;
 							}
 						}
@@ -2493,6 +2190,7 @@ void FChooserTableEditor::PasteInternal(UChooserTable* PastedContent, int PasteR
 			}
 		}
 
+		
 		for(int SourceColumnIndex = 0; SourceColumnIndex < PastedContent->ColumnsStructs.Num(); SourceColumnIndex++)
 		{
 			if (!MatchedSourceColumns[SourceColumnIndex])
@@ -2519,10 +2217,6 @@ void FChooserTableEditor::PasteInternal(UChooserTable* PastedContent, int PasteR
 		{
 			SelectRow(InsertIndex + i, false);
 		}
-		if(PastedContent->FallbackResult.IsValid())
-		{
-			SelectRow(ColumnWidget_SpecialIndex_Fallback, false);
-		}
 	}
 	
 	RefreshAll();
diff --git a/Engine/Plugins/Experimental/Chooser/Source/ChooserEditor/Private/ChooserTableEditor.h b/Engine/Plugins/Experimental/Chooser/Source/ChooserEditor/Private/ChooserTableEditor.h
--- a/Engine/Plugins/Experimental/Chooser/Source/ChooserEditor/Private/ChooserTableEditor.h
+++ b/Engine/Plugins/Experimental/Chooser/Source/ChooserEditor/Private/ChooserTableEditor.h
@@ -17,7 +17,6 @@
 #include "ChooserTableEditor.generated.h"
 
 class SComboButton;
-class SPositiveActionButton;
 class SEditableText;
 class IDetailsView;
 class UChooserRowDetails;
@@ -107,7 +106,6 @@ namespace UE::ChooserEditor
 		void ToggleDisableSelection();
 		void DeleteSelection();
 		void DuplicateSelection();
-		bool HasFallbackSelected();
 
 		bool CanMoveRowsUp();
 		void MoveRowsUp();
@@ -219,19 +217,17 @@ namespace UE::ChooserEditor
 		void UpdateTableColumns();
 		TArray<TSharedPtr<FChooserTableRow>> TableRows;
 	
-		TSharedPtr<SPositiveActionButton> CreateColumnComboButton;
-		TSharedPtr<SWidget> CreateRowComboButton;
+		TSharedPtr<SComboButton> CreateColumnComboButton;
+		TSharedPtr<SComboButton> CreateRowComboButton;
 
 		TSharedPtr<SHeaderRow> HeaderRow;
 		TSharedPtr<SListView<TSharedPtr<FChooserTableRow>>> TableView;
 
 		ESelectionType CurrentSelectionType = ESelectionType::Root;
 
-		TSharedRef<SWidget>	MakeCreateRowMenu();
-		TSharedRef<SWidget>	MakeCreateColumnMenu();
 	public:
 
-		TSharedPtr<SWidget>& GetCreateRowComboButton() { return CreateRowComboButton; };
+		TSharedPtr<SComboButton>& GetCreateRowComboButton() { return CreateRowComboButton; };
 		TSharedPtr<SNestedChooserTree> NestedChooserTree;
 
 		/** The name given to all instances of this type of editor */
diff --git a/Engine/Plugins/Experimental/Chooser/Source/ChooserEditor/Private/FloatDistanceColumnEditor.cpp b/Engine/Plugins/Experimental/Chooser/Source/ChooserEditor/Private/FloatDistanceColumnEditor.cpp
--- a/Engine/Plugins/Experimental/Chooser/Source/ChooserEditor/Private/FloatDistanceColumnEditor.cpp
+++ b/Engine/Plugins/Experimental/Chooser/Source/ChooserEditor/Private/FloatDistanceColumnEditor.cpp
@@ -31,8 +31,8 @@ TSharedRef<SWidget> CreateFloatDistanceColumnWidget(UChooserTable* Chooser, FCho
 		// create column header widget
 		
 		const FSlateBrush* ColumnIcon = FCoreStyle::Get().GetBrush("Icons.SortUp");
-		const FText ColumnTooltip = LOCTEXT("Float difference tooltip", "Float Difference Column: rows recieve a Score based on how different the input float is from the row value");
-		const FText ColumnName = LOCTEXT("Float Difference","Float Difference");
+		const FText ColumnTooltip = LOCTEXT("Float distance tooltip", "Float Distance Column: rows recieve a Cost based on how different the input float is from the row value");
+		const FText ColumnName = LOCTEXT("Float Distance","Float Distance");
 		
 		TSharedPtr<SWidget> DebugWidget = nullptr;
 		if (Chooser->GetEnableDebugTesting())
diff --git a/Engine/Plugins/Experimental/Chooser/Source/ChooserEditor/Private/OutputFloatColumnEditor.cpp b/Engine/Plugins/Experimental/Chooser/Source/ChooserEditor/Private/OutputFloatColumnEditor.cpp
--- a/Engine/Plugins/Experimental/Chooser/Source/ChooserEditor/Private/OutputFloatColumnEditor.cpp
+++ b/Engine/Plugins/Experimental/Chooser/Source/ChooserEditor/Private/OutputFloatColumnEditor.cpp
@@ -48,7 +48,7 @@ TSharedRef<SWidget> CreateOutputFloatColumnWidget(UChooserTable* Chooser, FChoos
         		})
         		.OnValueCommitted_Lambda([Chooser, OutputFloatColumn](double NewValue, ETextCommit::Type CommitType)
         		{
-					const FScopedTransaction Transaction(LOCTEXT("Edit Float Value", "Edit Float Value"));
+					const FScopedTransaction Transaction(LOCTEXT("Edit Min Value", "Edit Min Value"));
 					Chooser->Modify(true);
 					OutputFloatColumn->FallbackValue = NewValue;
         		});	
@@ -64,7 +64,7 @@ TSharedRef<SWidget> CreateOutputFloatColumnWidget(UChooserTable* Chooser, FChoos
     		{
     			if (Row < OutputFloatColumn->RowValues.Num())
     			{
-    				const FScopedTransaction Transaction(LOCTEXT("Edit Float Value", "Edit Float Value"));
+    				const FScopedTransaction Transaction(LOCTEXT("Edit Min Value", "Edit Min Value"));
     				Chooser->Modify(true);
     				OutputFloatColumn->RowValues[Row] = NewValue;
     			}
diff --git a/Engine/Plugins/Experimental/Chooser/Source/ChooserEditor/Private/SChooserTableRow.cpp b/Engine/Plugins/Experimental/Chooser/Source/ChooserEditor/Private/SChooserTableRow.cpp
--- a/Engine/Plugins/Experimental/Chooser/Source/ChooserEditor/Private/SChooserTableRow.cpp
+++ b/Engine/Plugins/Experimental/Chooser/Source/ChooserEditor/Private/SChooserTableRow.cpp
@@ -58,7 +58,7 @@ namespace UE::ChooserEditor
 					]
 				);
 		}
-		else if (RowIndex->RowIndex == SpecialIndex_Fallback || RowIndex->RowIndex == SpecialIndex_AddRow )
+		else if (RowIndex->RowIndex == SpecialIndex_Fallback)
 		{
 			SetContent(
 					SNew(SOverlay)
@@ -327,11 +327,7 @@ namespace UE::ChooserEditor
 			// on the row past the end, show an Add button in the result column
 			if (ColumnName == Result)
 			{
-				return SNew(SHorizontalBox)
-								+SHorizontalBox::Slot().AutoWidth()
-								[
-									Editor->GetCreateRowComboButton().ToSharedRef()
-								];
+				return Editor->GetCreateRowComboButton().ToSharedRef();
 			}
 		}
 
diff --git a/Engine/Plugins/Experimental/Chooser/Source/ProxyTable/Internal/LookupProxy.h b/Engine/Plugins/Experimental/Chooser/Source/ProxyTable/Internal/LookupProxy.h
--- a/Engine/Plugins/Experimental/Chooser/Source/ProxyTable/Internal/LookupProxy.h
+++ b/Engine/Plugins/Experimental/Chooser/Source/ProxyTable/Internal/LookupProxy.h
@@ -25,7 +25,7 @@ public:
 	CHOOSER_PARAMETER_BOILERPLATE();
 };
 
-USTRUCT(DisplayName = "Lookup Proxy", Meta = (Category = "Proxy Table", Tooltip = "Find a Proxy Asset entry in a Proxy Table, and evaluate it's value if this row is selected."))
+USTRUCT()
 struct PROXYTABLE_API FLookupProxy : public FObjectChooserBase
 {
 	GENERATED_BODY()
