Chooser editor tooltips etc

diff --git a/Engine/Plugins/Chooser/Source/Chooser/Internal/BoolColumn.h b/Engine/Plugins/Chooser/Source/Chooser/Internal/BoolColumn.h
--- a/Engine/Plugins/Chooser/Source/Chooser/Internal/BoolColumn.h
+++ b/Engine/Plugins/Chooser/Source/Chooser/Internal/BoolColumn.h
@@ -53,14 +53,14 @@ struct CHOOSER_API FBoolColumn : public FChooserColumnBase
 	public:
 	FBoolColumn();
 	
-	UPROPERTY(EditAnywhere, NoClear, Meta = (ExcludeBaseStruct, BaseStruct = "/Script/Chooser.ChooserParameterBoolBase"), Category = "Data")
+	UPROPERTY(EditAnywhere, NoClear, Meta = (ExcludeBaseStruct, BaseStruct = "/Script/Chooser.ChooserParameterBoolBase", ToolTip="The bool property this column will filter based on" ), Category = "Data")
 	FInstancedStruct InputValue;
 
 #if WITH_EDITORONLY_DATA
 	UPROPERTY()
 	TArray<bool> RowValues_DEPRECATED;
 
-	UPROPERTY(EditAnywhere, Category= "Data", DisplayName="DefaultRowValue");
+	UPROPERTY(EditAnywhere, Category=Data, meta=(ToolTip="DefaultRowValue will be assigned to cells when new rows are created"));
 	EBoolColumnCellValue DefaultRowValue = EBoolColumnCellValue::MatchAny;
 #endif
 	
diff --git a/Engine/Plugins/Chooser/Source/Chooser/Internal/Chooser.h b/Engine/Plugins/Chooser/Source/Chooser/Internal/Chooser.h
--- a/Engine/Plugins/Chooser/Source/Chooser/Internal/Chooser.h
+++ b/Engine/Plugins/Chooser/Source/Chooser/Internal/Chooser.h
@@ -144,12 +144,12 @@ public:
 	TObjectPtr<UChooserTable> RootChooser;
 	
 	// FallbackResult will be used as the Result if there are no rows in the chooser which pass all filters.  If FallbackResult is not assigned, then the Chooser will return null in that case.
-	UPROPERTY(EditAnywhere, Meta = (ExcludeBaseStruct, BaseStruct = "/Script/Chooser.ObjectChooserBase"), Category = "Fallback")
+	UPROPERTY()
 	FInstancedStruct FallbackResult;
 
 #if WITH_EDITORONLY_DATA
 	// Each possible Result (Rows of chooser table)
-	UPROPERTY(EditAnywhere, NoClear, DisplayName = "Results", Meta = (ExcludeBaseStruct, BaseStruct = "/Script/Chooser.ObjectChooserBase"), Category = "Hidden")
+	UPROPERTY()
 	TArray<FInstancedStruct> ResultsStructs;
 
 	UPROPERTY()
@@ -177,14 +177,18 @@ public:
 	TArray<FInstancedStruct> CookedResults;
 
 	// Columns which filter Results
-	UPROPERTY(EditAnywhere, NoClear, DisplayName = "Columns", Category = Hidden, meta = (ExcludeBaseStruct, BaseStruct = "/Script/Chooser.ChooserColumnBase"))
+	UPROPERTY(EditAnywhere, NoClear, DisplayName = "Results", Meta = (ExcludeBaseStruct, BaseStruct = "/Script/Chooser.ObjectChooserBase"), Category = "Hidden")
 	TArray<FInstancedStruct> ColumnsStructs;
+	
+	// The kind of output this chooser has (Object or Class)
+    UPROPERTY(EditAnywhere, DisplayName = "Result Type", Category="Result")
+    EObjectChooserResultType ResultType = EObjectChooserResultType::ObjectResult;
 
+	// The Class of Object this Chooser returns when ResultType is set to ObjectOfType, or the Parent Class of the Classes returned by this chooser when ResultType is set to ClassOfType
 	UPROPERTY(EditAnywhere, DisplayName= "Result Class", Category="Result", Meta = (AllowAbstract=true))
 	TObjectPtr<UClass> OutputObjectType;
-	UPROPERTY(EditAnywhere, DisplayName = "Result Type", Category="Result")
-	EObjectChooserResultType ResultType = EObjectChooserResultType::ObjectResult;
-
+	
+	// Parameter Objects or Structs from which the chooser can read or write properties 
 	UPROPERTY(EditAnywhere, NoClear, DisplayName = "Parameters", Meta = (ExcludeBaseStruct, BaseStruct = "/Script/Chooser.ContextObjectTypeBase"), Category = "Parameters")
 	TArray<FInstancedStruct> ContextData;
 
diff --git a/Engine/Plugins/Chooser/Source/Chooser/Internal/EnumColumn.h b/Engine/Plugins/Chooser/Source/Chooser/Internal/EnumColumn.h
--- a/Engine/Plugins/Chooser/Source/Chooser/Internal/EnumColumn.h
+++ b/Engine/Plugins/Chooser/Source/Chooser/Internal/EnumColumn.h
@@ -99,15 +99,15 @@ struct CHOOSER_API FEnumColumn : public FEnumColumnBase
 public:
 	FEnumColumn();
 
-	UPROPERTY(EditAnywhere, NoClear, Meta = (ExcludeBaseStruct, BaseStruct = "/Script/Chooser.ChooserParameterEnumBase"), Category = "Data")
+	UPROPERTY(EditAnywhere, NoClear, Meta = (ExcludeBaseStruct, BaseStruct = "/Script/Chooser.ChooserParameterEnumBase", ToolTip="The Enum property this column will filter based on"), Category = "Data")
 	FInstancedStruct InputValue;
 
 #if WITH_EDITORONLY_DATA
-	UPROPERTY(EditAnywhere, Category = "Data")
+	UPROPERTY(EditAnywhere, Category=Data, meta=(ToolTip="DefaultRowValue will be assigned to cells when new rows are created"));
 	FChooserEnumRowData DefaultRowValue;
 #endif
 	
-	UPROPERTY(EditAnywhere, Category = "Data")
+	UPROPERTY()
 	// array of results (cells for this column for each row in the table)
 	// should match the length of the Results array
 	TArray<FChooserEnumRowData> RowValues;
diff --git a/Engine/Plugins/Chooser/Source/Chooser/Internal/FloatDistanceColumn.h b/Engine/Plugins/Chooser/Source/Chooser/Internal/FloatDistanceColumn.h
--- a/Engine/Plugins/Chooser/Source/Chooser/Internal/FloatDistanceColumn.h
+++ b/Engine/Plugins/Chooser/Source/Chooser/Internal/FloatDistanceColumn.h
@@ -42,7 +42,7 @@ struct CHOOSER_API FFloatDistanceColumn : public FChooserColumnBase
 	public:
 	FFloatDistanceColumn();
 		
-	UPROPERTY(EditAnywhere, NoClear, Meta = (ExcludeBaseStruct, BaseStruct = "/Script/Chooser.ChooserParameterFloatBase"), Category = "Data")
+	UPROPERTY(EditAnywhere, NoClear, Meta = (ExcludeBaseStruct, BaseStruct = "/Script/Chooser.ChooserParameterFloatBase", ToolTip="The Float column this fillter will Score based on"), Category = "Data")
 	FInstancedStruct InputValue;
 
 	// Maximum Distance used for normalizing scoring (greater distances will be considered equal to the max)
@@ -69,7 +69,7 @@ struct CHOOSER_API FFloatDistanceColumn : public FChooserColumnBase
 	double MaxValue = 180;
 
 #if WITH_EDITORONLY_DATA
-	UPROPERTY(EditAnywhere, Category="Data")
+	UPROPERTY(EditAnywhere, Category=Data, meta=(ToolTip="DefaultRowValue will be assigned to cells when new rows are created"));
 	FChooserFloatDistanceRowData DefaultRowValue;
 
 	// Optional class to auto populate column data based on the result asset
@@ -77,7 +77,7 @@ struct CHOOSER_API FFloatDistanceColumn : public FChooserColumnBase
 	TSubclassOf<UFloatAutoPopulator> AutoPopulator;
 #endif
 	
-	UPROPERTY(EditAnywhere, Category="Data")
+	UPROPERTY()
 	// array of results (cells for this column for each row in the table)
 	// should match the length of the Results array 
 	TArray<FChooserFloatDistanceRowData> RowValues;
diff --git a/Engine/Plugins/Chooser/Source/Chooser/Internal/FloatRangeColumn.h b/Engine/Plugins/Chooser/Source/Chooser/Internal/FloatRangeColumn.h
--- a/Engine/Plugins/Chooser/Source/Chooser/Internal/FloatRangeColumn.h
+++ b/Engine/Plugins/Chooser/Source/Chooser/Internal/FloatRangeColumn.h
@@ -65,7 +65,7 @@ struct CHOOSER_API FFloatRangeColumn : public FChooserColumnBase
 	public:
 	FFloatRangeColumn();
 		
-	UPROPERTY(EditAnywhere, NoClear, Meta = (ExcludeBaseStruct, BaseStruct = "/Script/Chooser.ChooserParameterFloatBase"), Category = "Data")
+	UPROPERTY(EditAnywhere, NoClear, Meta = (ExcludeBaseStruct, BaseStruct = "/Script/Chooser.ChooserParameterFloatBase", ToolTip="The Float property this column will filter based on"), Category = "Data")
 	FInstancedStruct InputValue;
 
 	// Wrap input, and comparisons for numbers such as angles which 
@@ -79,11 +79,11 @@ struct CHOOSER_API FFloatRangeColumn : public FChooserColumnBase
 	double MaxValue = 180;
 
 #if WITH_EDITORONLY_DATA
-	UPROPERTY(EditAnywhere, Category="Data")
+	UPROPERTY(EditAnywhere, Category=Data, meta=(ToolTip="DefaultRowValue will be assigned to cells when new rows are created"));
 	FChooserFloatRangeRowData DefaultRowValue;
 #endif
 	
-	UPROPERTY(EditAnywhere, Category="Data")
+	UPROPERTY()
 	// array of results (cells for this column for each row in the table)
 	// should match the length of the Results array 
 	TArray<FChooserFloatRangeRowData> RowValues;
diff --git a/Engine/Plugins/Chooser/Source/Chooser/Internal/GameplayTagColumn.h b/Engine/Plugins/Chooser/Source/Chooser/Internal/GameplayTagColumn.h
--- a/Engine/Plugins/Chooser/Source/Chooser/Internal/GameplayTagColumn.h
+++ b/Engine/Plugins/Chooser/Source/Chooser/Internal/GameplayTagColumn.h
@@ -54,7 +54,7 @@ struct CHOOSER_API FGameplayTagColumn : public FChooserColumnBase
 	public:
 	FGameplayTagColumn();
 	
-	UPROPERTY(EditAnywhere, Meta = (ExcludeBaseStruct, BaseStruct = "/Script/Chooser.ChooserParameterGameplayTagBase"), Category = "Data")
+	UPROPERTY(EditAnywhere, Meta = (ExcludeBaseStruct, BaseStruct = "/Script/Chooser.ChooserParameterGameplayTagBase", ToolTip="The GameplayTagCollection property this column will filter based on"), Category = "Data")
 	FInstancedStruct InputValue;
 
 	UPROPERTY(EditAnywhere, Category="Data")
@@ -73,11 +73,11 @@ struct CHOOSER_API FGameplayTagColumn : public FChooserColumnBase
 	bool bInvertMatchingLogic = false;
 
 #if WITH_EDITORONLY_DATA
-	UPROPERTY(EditAnywhere, Category=Runtime)
+	UPROPERTY(EditAnywhere, Category=Data, meta=(ToolTip="DefaultRowValue will be assigned to cells when new rows are created"));
 	FGameplayTagContainer DefaultRowValue;
 #endif
 	
-	UPROPERTY(EditAnywhere, Category=Runtime)
+	UPROPERTY()
 	// array of results (cells for this column for each row in the table)
 	// should match the length of the Results array 
 	TArray<FGameplayTagContainer> RowValues;
diff --git a/Engine/Plugins/Chooser/Source/Chooser/Internal/MultiEnumColumn.h b/Engine/Plugins/Chooser/Source/Chooser/Internal/MultiEnumColumn.h
--- a/Engine/Plugins/Chooser/Source/Chooser/Internal/MultiEnumColumn.h
+++ b/Engine/Plugins/Chooser/Source/Chooser/Internal/MultiEnumColumn.h
@@ -34,15 +34,15 @@ struct CHOOSER_API FMultiEnumColumn : public FChooserColumnBase
 public:
 	FMultiEnumColumn();
 
-	UPROPERTY(EditAnywhere, NoClear, Meta = (ExcludeBaseStruct, BaseStruct = "/Script/Chooser.ChooserParameterEnumBase"), Category = "Data")
+	UPROPERTY(EditAnywhere, NoClear, Meta = (ExcludeBaseStruct, BaseStruct = "/Script/Chooser.ChooserParameterEnumBase", ToolTip="The Enum property this column will filter based on"), Category = "Data")
 	FInstancedStruct InputValue;
 
 #if WITH_EDITORONLY_DATA
-	UPROPERTY(EditAnywhere, Category = "Data")
+	UPROPERTY(EditAnywhere, Category=Data, meta=(ToolTip="DefaultRowValue will be assigned to cells when new rows are created"));
 	FChooserMultiEnumRowData DefaultRowValue;
 #endif
 	
-	UPROPERTY(EditAnywhere, Category = "Data")
+	UPROPERTY()
 	// array of results (cells for this column for each row in the table)
 	// should match the length of the Results array
 	TArray<FChooserMultiEnumRowData> RowValues;
diff --git a/Engine/Plugins/Chooser/Source/Chooser/Internal/ObjectClassColumn.h b/Engine/Plugins/Chooser/Source/Chooser/Internal/ObjectClassColumn.h
--- a/Engine/Plugins/Chooser/Source/Chooser/Internal/ObjectClassColumn.h
+++ b/Engine/Plugins/Chooser/Source/Chooser/Internal/ObjectClassColumn.h
@@ -41,15 +41,15 @@ struct CHOOSER_API FObjectClassColumn : public FChooserColumnBase
 
 	FObjectClassColumn();
 
-	UPROPERTY(EditAnywhere, NoClear, Meta = (ExcludeBaseStruct, BaseStruct = "/Script/Chooser.ChooserParameterObjectBase"), Category = "Data")
+	UPROPERTY(EditAnywhere, NoClear, Meta = (ExcludeBaseStruct, BaseStruct = "/Script/Chooser.ChooserParameterObjectBase", ToolTip="The Object reference property this column will filter based on"), Category = "Data")
 	FInstancedStruct InputValue;
 
 #if WITH_EDITORONLY_DATA
-	UPROPERTY(EditAnywhere, Category = "Data")
+	UPROPERTY(EditAnywhere, Category = "Data", Meta =(ToolTip="DefaultRowValue will be assigned to cells when new rows are created"))
 	FChooserObjectClassRowData DefaultRowValue;
 #endif
 
-	UPROPERTY(EditAnywhere, Category = "Data")
+	UPROPERTY()
 	// array of results (cells for this column for each row in the table)
 	// should match the length of the Results array
 	TArray<FChooserObjectClassRowData> RowValues;
diff --git a/Engine/Plugins/Chooser/Source/Chooser/Internal/ObjectColumn.h b/Engine/Plugins/Chooser/Source/Chooser/Internal/ObjectColumn.h
--- a/Engine/Plugins/Chooser/Source/Chooser/Internal/ObjectColumn.h
+++ b/Engine/Plugins/Chooser/Source/Chooser/Internal/ObjectColumn.h
@@ -61,15 +61,15 @@ struct CHOOSER_API FObjectColumn : public FChooserColumnBase
 
 	FObjectColumn();
 
-	UPROPERTY(EditAnywhere, NoClear, Meta = (ExcludeBaseStruct, BaseStruct = "/Script/Chooser.ChooserParameterObjectBase"), Category = "Data")
+	UPROPERTY(EditAnywhere, NoClear, Meta = (ExcludeBaseStruct, BaseStruct = "/Script/Chooser.ChooserParameterObjectBase", ToolTip="The Object reference property this column will filter based on"), Category = "Data")
 	FInstancedStruct InputValue;
 
 #if WITH_EDITORONLY_DATA
-	UPROPERTY(EditAnywhere, Category = "Data")
+	UPROPERTY(EditAnywhere, Category=Data, meta=(ToolTip="DefaultRowValue will be assigned to cells when new rows are created"));
 	FChooserObjectRowData DefaultRowValue;
 #endif
 
-	UPROPERTY(EditAnywhere, Category = "Data")
+	UPROPERTY()
 	// array of results (cells for this column for each row in the table)
 	// should match the length of the Results array
 	TArray<FChooserObjectRowData> RowValues;
diff --git a/Engine/Plugins/Chooser/Source/Chooser/Internal/OutputBoolColumn.h b/Engine/Plugins/Chooser/Source/Chooser/Internal/OutputBoolColumn.h
--- a/Engine/Plugins/Chooser/Source/Chooser/Internal/OutputBoolColumn.h
+++ b/Engine/Plugins/Chooser/Source/Chooser/Internal/OutputBoolColumn.h
@@ -18,7 +18,7 @@ struct CHOOSER_API FOutputBoolColumn : public FChooserColumnBase
 	virtual bool HasOutputs() const override { return true; }
 	virtual void SetOutputs(FChooserEvaluationContext& Context, int RowIndex) const override;
 	
-	UPROPERTY(EditAnywhere, NoClear, Meta = (ExcludeBaseStruct, BaseStruct = "/Script/Chooser.ChooserParameterBoolBase"), Category = "Hidden")
+	UPROPERTY(EditAnywhere, NoClear, Meta = (ExcludeBaseStruct, BaseStruct = "/Script/Chooser.ChooserParameterBoolBase", ToolTip="The Bool property this column will write to"), Category = "Hidden")
 	FInstancedStruct InputValue;
 
 #if WITH_EDITOR
@@ -40,15 +40,15 @@ struct CHOOSER_API FOutputBoolColumn : public FChooserColumnBase
 	}
 	
 	// FallbackValue will be used as the output value if the all rows in the chooser fail, and the FallbackResult from the chooser is used.
-	UPROPERTY(EditAnywhere, Category=Data);
+	UPROPERTY()
 	bool bFallbackValue = false;
 	
 #if WITH_EDITORONLY_DATA
-	UPROPERTY(EditAnywhere, Category=Data);
+	UPROPERTY(EditAnywhere, Category=Data, meta=(ToolTip="DefaultRowValue will be assigned to cells when new rows are created"));
 	bool DefaultRowValue = false;
 #endif
 	
-	UPROPERTY(EditAnywhere, Category=Data);
+	UPROPERTY()
 	TArray<bool> RowValues; 
 
 	CHOOSER_COLUMN_BOILERPLATE(FChooserParameterBoolBase);
diff --git a/Engine/Plugins/Chooser/Source/Chooser/Internal/OutputEnumColumn.h b/Engine/Plugins/Chooser/Source/Chooser/Internal/OutputEnumColumn.h
--- a/Engine/Plugins/Chooser/Source/Chooser/Internal/OutputEnumColumn.h
+++ b/Engine/Plugins/Chooser/Source/Chooser/Internal/OutputEnumColumn.h
@@ -34,21 +34,21 @@ public:
 	virtual bool HasOutputs() const override { return true; }
 	virtual void SetOutputs(FChooserEvaluationContext& Context, int RowIndex) const override;
 
-	UPROPERTY(EditAnywhere, NoClear, Meta = (ExcludeBaseStruct, BaseStruct = "/Script/Chooser.ChooserParameterEnumBase"), Category = "Data")
+	UPROPERTY(EditAnywhere, NoClear, Meta = (ExcludeBaseStruct, BaseStruct = "/Script/Chooser.ChooserParameterEnumBase", ToolTip="The Enum property this column will write to"), Category = "Data")
 	FInstancedStruct InputValue;
 
 	// FallbackValue will be used as the output value if the all rows in the chooser fail, and the FallbackResult from the chooser is used.
-	UPROPERTY(EditAnywhere, Category = "Data")
+	UPROPERTY()
 	FChooserOutputEnumRowData FallbackValue;
 
 #if WITH_EDITORONLY_DATA
-	UPROPERTY(EditAnywhere, Category = "Data")
+	UPROPERTY(EditAnywhere, Category=Data, meta=(ToolTip="DefaultRowValue will be assigned to cells when new rows are created"));
 	FChooserOutputEnumRowData DefaultRowValue;
 
 	virtual void PostLoad() override;
 #endif
 	
-	UPROPERTY(EditAnywhere, Category = "Data")
+	UPROPERTY()
 	// array of results (cells for this column for each row in the table)
 	// should match the length of the Results array
 	TArray<FChooserOutputEnumRowData> RowValues;
diff --git a/Engine/Plugins/Chooser/Source/Chooser/Internal/OutputFloatColumn.h b/Engine/Plugins/Chooser/Source/Chooser/Internal/OutputFloatColumn.h
--- a/Engine/Plugins/Chooser/Source/Chooser/Internal/OutputFloatColumn.h
+++ b/Engine/Plugins/Chooser/Source/Chooser/Internal/OutputFloatColumn.h
@@ -18,7 +18,7 @@ struct CHOOSER_API FOutputFloatColumn : public FChooserColumnBase
 	virtual bool HasOutputs() const override { return true; }
 	virtual void SetOutputs(FChooserEvaluationContext& Context, int RowIndex) const override;
 	
-	UPROPERTY(EditAnywhere, Meta = (ExcludeBaseStruct, NoClear, BaseStruct = "/Script/Chooser.ChooserParameterFloatBase"), Category = "Hidden")
+	UPROPERTY(EditAnywhere, Meta = (ExcludeBaseStruct, NoClear, BaseStruct = "/Script/Chooser.ChooserParameterFloatBase", ToolTip="The Float property this column will write to"), Category = "Hidden")
 	FInstancedStruct InputValue;
 
 #if WITH_EDITOR
@@ -40,15 +40,15 @@ struct CHOOSER_API FOutputFloatColumn : public FChooserColumnBase
    	}
 
 	// FallbackValue will be used as the output value if the all rows in the chooser fail, and the FallbackResult from the chooser is used.
-	UPROPERTY(EditAnywhere, Category=Data);
+	UPROPERTY()
 	double FallbackValue = 0.0;
 	
 #if WITH_EDITORONLY_DATA
-	UPROPERTY(EditAnywhere, Category=Data);
+	UPROPERTY(EditAnywhere, Category=Data, meta=(ToolTip="DefaultRowValue will be assigned to cells when new rows are created"));
 	double DefaultRowValue = 0.0;
 #endif
 	
-	UPROPERTY(EditAnywhere, Category=Data);
+	UPROPERTY()
 	TArray<double> RowValues; 
 
 	CHOOSER_COLUMN_BOILERPLATE(FChooserParameterFloatBase);
diff --git a/Engine/Plugins/Chooser/Source/Chooser/Internal/OutputObjectColumn.h b/Engine/Plugins/Chooser/Source/Chooser/Internal/OutputObjectColumn.h
--- a/Engine/Plugins/Chooser/Source/Chooser/Internal/OutputObjectColumn.h
+++ b/Engine/Plugins/Chooser/Source/Chooser/Internal/OutputObjectColumn.h
@@ -31,17 +31,17 @@ struct CHOOSER_API FOutputObjectColumn : public FChooserColumnBase
 	FInstancedStruct InputValue;
 
 #if WITH_EDITORONLY_DATA
-	UPROPERTY(EditAnywhere, Category = "Data")
+	UPROPERTY(EditAnywhere, Category=Data, meta=(ToolTip="DefaultRowValue will be assigned to cells when new rows are created"));
 	FChooserOutputObjectRowData DefaultRowValue;
 #endif
 
-	UPROPERTY(EditAnywhere, Category = "Data")
+	UPROPERTY()
 	// array of results (cells for this column for each row in the table)
 	// should match the length of the Results array
 	TArray<FChooserOutputObjectRowData> RowValues;
 
 	// FallbackValue will be used as the output value if the all rows in the chooser fail, and the FallbackResult from the chooser is used.
-    UPROPERTY(EditAnywhere, Category=Data);
+    UPROPERTY()
 	FChooserOutputObjectRowData FallbackValue;
 
 	virtual bool HasFilters() const override { return false; }
diff --git a/Engine/Plugins/Chooser/Source/Chooser/Internal/OutputStructColumn.h b/Engine/Plugins/Chooser/Source/Chooser/Internal/OutputStructColumn.h
--- a/Engine/Plugins/Chooser/Source/Chooser/Internal/OutputStructColumn.h
+++ b/Engine/Plugins/Chooser/Source/Chooser/Internal/OutputStructColumn.h
@@ -37,7 +37,7 @@ struct CHOOSER_API FOutputStructColumn : public FChooserColumnBase
 	virtual bool HasOutputs() const override { return true; }
 	virtual void SetOutputs(FChooserEvaluationContext& Context, int RowIndex) const override;
 	
-	UPROPERTY(EditAnywhere, NoClear, Meta = (ExcludeBaseStruct, BaseStruct = "/Script/Chooser.ChooserParameterStructBase"), Category = "Hidden")
+	UPROPERTY(EditAnywhere, NoClear, Meta = (ExcludeBaseStruct, BaseStruct = "/Script/Chooser.ChooserParameterStructBase", ToolTip="The Struct property this column wil write to"), Category = "Hidden")
 	FInstancedStruct InputValue;
 
 #if WITH_EDITOR
@@ -60,11 +60,11 @@ struct CHOOSER_API FOutputStructColumn : public FChooserColumnBase
 	}
 	
 	// FallbackValue will be used as the output value if the all rows in the chooser fail, and the FallbackResult from the chooser is used.
-	UPROPERTY(EditAnywhere, Meta = (StructTypeConst), Category=Data);
+	UPROPERTY();
    	FInstancedStruct FallbackValue;
 	
 #if WITH_EDITORONLY_DATA
-	UPROPERTY(EditAnywhere, Meta = (StructTypeConst), Category=Data);
+	UPROPERTY(EditAnywhere, Meta = (StructTypeConst, ToolTip="DefaultRowValue will be assigned to cells when new rows are created"), Category=Data);
 	FInstancedStruct DefaultRowValue;
 #endif
 	
diff --git a/Engine/Plugins/Chooser/Source/Chooser/Internal/RandomizeColumn.h b/Engine/Plugins/Chooser/Source/Chooser/Internal/RandomizeColumn.h
--- a/Engine/Plugins/Chooser/Source/Chooser/Internal/RandomizeColumn.h
+++ b/Engine/Plugins/Chooser/Source/Chooser/Internal/RandomizeColumn.h
@@ -33,7 +33,7 @@ struct CHOOSER_API FRandomizeColumn : public FChooserColumnBase
 	public:
 	FRandomizeColumn();
 	
-	UPROPERTY(EditAnywhere, NoClear, Meta = (ExcludeBaseStruct, BaseStruct = "/Script/Chooser.RandomizeContextProperty"), Category = "Data")
+	UPROPERTY(EditAnywhere, NoClear, Meta = (ExcludeBaseStruct, BaseStruct = "/Script/Chooser.RandomizeContextProperty", ToolTip="Optional reference to a ChooserRandomizationContext struct. If bound, this is used to store the most recent selection (for each Choosers referencing it - you only need to create one variable per Character or context), for use with RepateProbabilityMultiplier to reduce the chance of selecting the same entry twice."), Category = "Data")
 	FInstancedStruct InputValue;
 	
 	
@@ -44,11 +44,11 @@ struct CHOOSER_API FRandomizeColumn : public FChooserColumnBase
 	float EqualCostThreshold = 0.001f;
 
 #if WITH_EDITORONLY_DATA
-	UPROPERTY(EditAnywhere, Category= "Data", DisplayName="DefaultRowValue");
+	UPROPERTY(EditAnywhere, Category=Data, meta=(ToolTip="DefaultRowValue will be assigned to cells when new rows are created"));
 	float DefaultRowValue = 1.0f;
 #endif
 	
-	UPROPERTY(EditAnywhere, Category= "Data", DisplayName="RowValues");
+	UPROPERTY()
 	TArray<float> RowValues; 
 	
 	virtual void Filter(FChooserEvaluationContext& Context, const FChooserIndexArray& IndexListIn, FChooserIndexArray& IndexListOut) const override;
diff --git a/Engine/Plugins/Chooser/Source/Chooser/Private/OutputStructColumn.cpp b/Engine/Plugins/Chooser/Source/Chooser/Private/OutputStructColumn.cpp
--- a/Engine/Plugins/Chooser/Source/Chooser/Private/OutputStructColumn.cpp
+++ b/Engine/Plugins/Chooser/Source/Chooser/Private/OutputStructColumn.cpp
@@ -1,7 +1,6 @@
 // Copyright Epic Games, Inc. All Rights Reserved.
 #include "OutputStructColumn.h"
 #include "ChooserPropertyAccess.h"
-#include "OutputStructColumn.h"
 
 #if WITH_EDITOR
 #include "IPropertyAccessEditor.h"
@@ -91,7 +90,7 @@ void FOutputStructColumn::SetFromDetails(FInstancedPropertyBag& PropertyBag, int
 			TValueOrError<FStructView, EPropertyBagResult> Result = PropertyBag.GetValueStruct(PropertyName, Value.GetScriptStruct());
 			if (FStructView* StructView = Result.TryGetValue())
 			{
-				RowValues[RowIndex].GetScriptStruct()->CopyScriptStruct(Value.GetMutableMemory(), StructView->GetMemory());
+				GetValueForIndex(RowIndex).GetScriptStruct()->CopyScriptStruct(Value.GetMutableMemory(), StructView->GetMemory());
 			}
 		}
 	}
diff --git a/Engine/Plugins/Chooser/Source/ChooserEditor/Private/ChooserEditorStyle.cpp b/Engine/Plugins/Chooser/Source/ChooserEditor/Private/ChooserEditorStyle.cpp
--- a/Engine/Plugins/Chooser/Source/ChooserEditor/Private/ChooserEditorStyle.cpp
+++ b/Engine/Plugins/Chooser/Source/ChooserEditor/Private/ChooserEditorStyle.cpp
@@ -18,7 +18,7 @@ FChooserEditorStyle::FChooserEditorStyle() :
 	const FVector2D Icon16x16(16.0f, 16.0f);
 	const FVector2D Icon24x24(24.0f, 24.0f);
 
-	SetContentRoot(FPaths::EnginePluginsDir() / TEXT("Experimental/Chooser/Content"));
+	SetContentRoot(FPaths::EnginePluginsDir() / TEXT("Chooser/Content"));
 	SetCoreContentRoot(FPaths::EngineContentDir() / TEXT(""));
 
 	// content browser icon
diff --git a/Engine/Plugins/Chooser/Source/ChooserEditor/Public/ChooserInitializer.h b/Engine/Plugins/Chooser/Source/ChooserEditor/Public/ChooserInitializer.h
--- a/Engine/Plugins/Chooser/Source/ChooserEditor/Public/ChooserInitializer.h
+++ b/Engine/Plugins/Chooser/Source/ChooserEditor/Public/ChooserInitializer.h
@@ -23,12 +23,16 @@ struct FGenericChooserInitializer : public FChooserInitializer
 
 	virtual void Initialize(UChooserTable* Chooser) const override;
 	
-	UPROPERTY(EditAnywhere, DisplayName= "Result Class", Category="Result", Meta = (AllowAbstract=true))
-	TObjectPtr<UClass> OutputObjectType;
-	
+	// The kind of output this chooser has (Object or Class)
 	UPROPERTY(EditAnywhere, DisplayName = "Result Type", Category="Result")
 	EObjectChooserResultType ResultType = EObjectChooserResultType::ObjectResult;
 	
+	
+	// The Class of Object this Chooser returns when ResultType is set to ObjectOfType, or the Parent Class of the Classes returned by this chooser when ResultType is set to ClassOfType
+	UPROPERTY(EditAnywhere, DisplayName= "Result Class", Category="Result", Meta = (AllowAbstract=true))
+	TObjectPtr<UClass> OutputObjectType;
+	
+    // Parameter Objects or Structs from which the chooser can read or write properties 
 	UPROPERTY(EditAnywhere, DisplayName = "Parameters", NoClear, Meta = (ExpandByDefault, ExcludeBaseStruct, BaseStruct = "/Script/Chooser.ContextObjectTypeBase"), Category = "Parameters")
 	TArray<FInstancedStruct> ContextData;
 };
