Add Compile time validation in ChooserPlayer animgraph node.
Mainly for the case of a LookupProxy which is missing a proxy table binding

diff --git a/Engine/Plugins/Chooser/Source/Chooser/Public/ChooserPropertyAccess.h b/Engine/Plugins/Chooser/Source/Chooser/Public/ChooserPropertyAccess.h
--- a/Engine/Plugins/Chooser/Source/Chooser/Public/ChooserPropertyAccess.h
+++ b/Engine/Plugins/Chooser/Source/Chooser/Public/ChooserPropertyAccess.h
@@ -413,7 +413,11 @@ bool FChooserPropertyBinding::SetValue(FChooserEvaluationContext& Context, const
 	{\
 		Binding.Compile(Owner, bForce);\
 	};\
-	\
+	virtual bool HasCompileErrors(FText& Message) override\
+	{\
+		Message = Binding.CompileMessage; \
+		return !Message.IsEmpty(); \
+	}\
 	virtual void AddSearchNames(FStringBuilderBase& Builder) const override\
 	{\
 		for (const FName& Entry : Binding.PropertyBindingChain)\
@@ -458,11 +462,26 @@ bool FChooserPropertyBinding::SetValue(FChooserEvaluationContext& Context, const
 			}\
 		}\
 	}
+#elif WITH_EDITORONLY_DATA
+#define CHOOSER_PARAMETER_BOILERPLATE() \
+	virtual void Compile(IHasContextClass* Owner, bool bForce) override\
+	{\
+		Binding.Compile(Owner, bForce);\
+	}\
+	virtual bool HasCompileErrors(FText& Message) override\
+	{\
+		Message = Binding.CompileMessage; \
+		return !Message.IsEmpty(); \
+	}
 #else
 #define CHOOSER_PARAMETER_BOILERPLATE() \
 	virtual void Compile(IHasContextClass* Owner, bool bForce) override\
 	{\
 		Binding.Compile(Owner, bForce);\
+	}\
+	virtual bool HasCompileErrors(FText& Message) override\
+	{\
+		return false;\
 	}
 #endif
 
diff --git a/Engine/Plugins/Chooser/Source/Chooser/Public/IChooserParameterBase.h b/Engine/Plugins/Chooser/Source/Chooser/Public/IChooserParameterBase.h
--- a/Engine/Plugins/Chooser/Source/Chooser/Public/IChooserParameterBase.h
+++ b/Engine/Plugins/Chooser/Source/Chooser/Public/IChooserParameterBase.h
@@ -27,6 +27,7 @@ struct FChooserParameterBase
 
 	virtual void PostLoad() {};
 	virtual void Compile(IHasContextClass* Owner, bool bForce) {};
+	virtual bool HasCompileErrors(FText& OutMessage) { return false; }
 
 	virtual ~FChooserParameterBase() {}
 };
\ No newline at end of file
diff --git a/Engine/Plugins/Chooser/Source/Chooser/Public/IObjectChooser.h b/Engine/Plugins/Chooser/Source/Chooser/Public/IObjectChooser.h
--- a/Engine/Plugins/Chooser/Source/Chooser/Public/IObjectChooser.h
+++ b/Engine/Plugins/Chooser/Source/Chooser/Public/IObjectChooser.h
@@ -103,6 +103,7 @@ public:
 	DECLARE_DELEGATE_RetVal_OneParam( EIteratorStatus, FObjectChooserIteratorCallback, UObject*);
 
 	virtual void Compile(class IHasContextClass* HasContext, bool bForce) {};
+	virtual bool HasCompileErrors(FText& Message) { return false; }
 
 	virtual UObject* ChooseObject(FChooserEvaluationContext& Context) const { return nullptr; };
 	virtual EIteratorStatus ChooseMulti(FChooserEvaluationContext& ContextData, FObjectChooserIteratorCallback Callback) const
diff --git a/Engine/Plugins/Chooser/Source/ChooserUncooked/Private/AnimGraphNode_ChooserPlayer.cpp b/Engine/Plugins/Chooser/Source/ChooserUncooked/Private/AnimGraphNode_ChooserPlayer.cpp
--- a/Engine/Plugins/Chooser/Source/ChooserUncooked/Private/AnimGraphNode_ChooserPlayer.cpp
+++ b/Engine/Plugins/Chooser/Source/ChooserUncooked/Private/AnimGraphNode_ChooserPlayer.cpp
@@ -53,6 +53,27 @@ FText UAnimGraphNode_ChooserPlayer::GetMenuCategory() const
 	return LOCTEXT("NodeCategory", "Animation|Sequences");
 }
 
+void UAnimGraphNode_ChooserPlayer::ValidateAnimNodeDuringCompilation(USkeleton* ForSkeleton, FCompilerResultsLog& MessageLog)
+{
+	Super::ValidateAnimNodeDuringCompilation(ForSkeleton, MessageLog);
+
+	if (FObjectChooserBase* ChooserBase = Node.Chooser.GetMutablePtr<FObjectChooserBase>())
+	{
+		ChooserBase->Compile(this, true);
+		FText Message;
+		if (ChooserBase->HasCompileErrors(Message))
+		{
+			FText NodeMessage = FText::Format(LOCTEXT("error in node", "{0} in @@"), Message);
+			MessageLog.Error(*NodeMessage.ToString(), this);
+		}
+	}
+	else
+	{
+		MessageLog.Error(TEXT("No Chooser set in @@"), this);
+	}
+	
+}
+
 void UAnimGraphNode_ChooserPlayer::CustomizeDetails(IDetailLayoutBuilder& DetailBuilder)
 {
 	Super::CustomizeDetails(DetailBuilder);
diff --git a/Engine/Plugins/Chooser/Source/ChooserUncooked/Public/AnimGraphNode_ChooserPlayer.h b/Engine/Plugins/Chooser/Source/ChooserUncooked/Public/AnimGraphNode_ChooserPlayer.h
--- a/Engine/Plugins/Chooser/Source/ChooserUncooked/Public/AnimGraphNode_ChooserPlayer.h
+++ b/Engine/Plugins/Chooser/Source/ChooserUncooked/Public/AnimGraphNode_ChooserPlayer.h
@@ -30,6 +30,7 @@ public:
 	virtual FText GetTooltipText() const override;
 	virtual FText GetNodeTitle(ENodeTitleType::Type TitleType) const override;
 	virtual FText GetMenuCategory() const override;
+	virtual void ValidateAnimNodeDuringCompilation(USkeleton* ForSkeleton, FCompilerResultsLog& MessageLog) override;
 	// End of UEdGraphNode interface
 
 	virtual FAnimNode_BlendStack_Standalone* GetBlendStackNode() const override { return (FAnimNode_BlendStack_Standalone*)(&Node); }
diff --git a/Engine/Plugins/Chooser/Source/ProxyTable/Internal/LookupProxy.h b/Engine/Plugins/Chooser/Source/ProxyTable/Internal/LookupProxy.h
--- a/Engine/Plugins/Chooser/Source/ProxyTable/Internal/LookupProxy.h
+++ b/Engine/Plugins/Chooser/Source/ProxyTable/Internal/LookupProxy.h
@@ -35,6 +35,7 @@ struct PROXYTABLE_API FLookupProxy : public FObjectChooserBase
 	FLookupProxy();
 	
 	virtual void Compile(IHasContextClass* HasContext, bool bForce) override;
+	virtual bool HasCompileErrors(FText& Message) override;
 
 	virtual void GetDebugName(FString& OutName) const override;
 	
diff --git a/Engine/Plugins/Chooser/Source/ProxyTable/Private/LookupProxy.cpp b/Engine/Plugins/Chooser/Source/ProxyTable/Private/LookupProxy.cpp
--- a/Engine/Plugins/Chooser/Source/ProxyTable/Private/LookupProxy.cpp
+++ b/Engine/Plugins/Chooser/Source/ProxyTable/Private/LookupProxy.cpp
@@ -98,6 +98,37 @@ void FLookupProxy::Compile(IHasContextClass* HasContext, bool bForce)
 	}
 }
 
+bool FLookupProxy::HasCompileErrors(FText& Message)
+{
+	if (Proxy)
+	{
+		if (FChooserParameterBase* ProxyTableParam = ProxyTable.GetMutablePtr<FChooserParameterBase>())
+		{
+			FText BindingMessage;
+			if (ProxyTableParam->HasCompileErrors(BindingMessage))
+			{
+				Message = FText::Format(NSLOCTEXT("LookupProxy", "LookupProxyBinding", "LookupProxy: ProxyTable Reference - {0}"), BindingMessage);
+				return true;
+			}
+			else
+			{
+				return false;
+			}
+		}
+		else
+		{
+			Message = NSLOCTEXT("LookupProxy", "NoParameter", "LookupProxy: No Proxy Table parameter bound");
+        	return true;
+		}
+	}
+	else
+	{
+		Message = NSLOCTEXT("LookupProxy", "NoProxy", "LookupProxy: Proxy Asset is not set");
+		return true;
+	}
+}
+
+
 void FLookupProxy::GetDebugName(FString& OutName) const
 {
 	if (Proxy)
