Updaet MultiEnum Column details panel editing

diff --git a/Engine/Plugins/Chooser/Source/Chooser/Private/EnumColumn.cpp b/Engine/Plugins/Chooser/Source/Chooser/Private/EnumColumn.cpp
--- a/Engine/Plugins/Chooser/Source/Chooser/Private/EnumColumn.cpp
+++ b/Engine/Plugins/Chooser/Source/Chooser/Private/EnumColumn.cpp
@@ -143,7 +143,7 @@ void FEnumColumn::Filter(FChooserEvaluationContext& Context, const FChooserIndex
 		FName PropertyName("RowData",ColumnIndex);
 
 		FInstancedPropertyBag Struct;
-		Struct.AddProperty("Value", EPropertyBagPropertyType::Enum, const_cast<UEnum*>(InputValue.Get<FChooserParameterEnumBase>().GetEnum()));
+		Struct.AddProperty("Value", EPropertyBagPropertyType::Enum, InputValue.Get<FChooserParameterEnumBase>().GetEnum());
 		Struct.AddProperty("Comparison", EPropertyBagPropertyType::Enum, StaticEnum<EEnumColumnCellValueComparison>());
 		
 		TValueOrError<FStructView, EPropertyBagResult> Result = PropertyBag.GetValueStruct(PropertyName, const_cast<UPropertyBag*>(Struct.GetPropertyBagStruct()));
diff --git a/Engine/Plugins/Chooser/Source/Chooser/Private/MultiEnumColumn.cpp b/Engine/Plugins/Chooser/Source/Chooser/Private/MultiEnumColumn.cpp
--- a/Engine/Plugins/Chooser/Source/Chooser/Private/MultiEnumColumn.cpp
+++ b/Engine/Plugins/Chooser/Source/Chooser/Private/MultiEnumColumn.cpp
@@ -61,25 +61,45 @@ void FMultiEnumColumn::Filter(FChooserEvaluationContext& Context, const FChooser
 
 #if WITH_EDITOR
 	void FMultiEnumColumn::AddToDetails (FInstancedPropertyBag& PropertyBag, int32 ColumnIndex, int32 RowIndex)
+{
+	FText DisplayName;
+	InputValue.Get<FChooserParameterBase>().GetDisplayName(DisplayName);
+	FName PropertyName("RowData",ColumnIndex);
+	if (const UEnum* Enum = InputValue.Get<FChooserParameterEnumBase>().GetEnum())
 	{
-		FText DisplayName;
-		InputValue.Get<FChooserParameterBase>().GetDisplayName(DisplayName);
-		FName PropertyName("RowData",ColumnIndex);
-
-		FPropertyBagPropertyDesc PropertyDesc(PropertyName,  EPropertyBagPropertyType::Struct, FChooserMultiEnumRowData::StaticStruct());
+		FPropertyBagPropertyDesc PropertyDesc(PropertyName,  EPropertyBagContainerType::Array, EPropertyBagPropertyType::Enum, Enum);
 		PropertyDesc.MetaData.Add(FPropertyBagPropertyDescMetaData("DisplayName", DisplayName.ToString()));
 		PropertyBag.AddProperties({PropertyDesc});
-		PropertyBag.SetValueStruct(PropertyName, RowValues[RowIndex]);
+	
+		FPropertyBagArrayRef ArrayRef = PropertyBag.GetMutableArrayRef(PropertyName).GetValue();
+
+		int32 NumEnums = Enum->NumEnums();
+		for(int32 i=0; i<NumEnums; i++)
+		{
+			if (RowValues[RowIndex].Value & (1 << Enum->GetValueByIndex(i)))
+			{
+				ArrayRef.AddValue();
+				ArrayRef.SetValueEnum(ArrayRef.Num()-1,Enum->GetValueByIndex(i), Enum);
+			}
+		}
 	}
+}
 	
-	void FMultiEnumColumn::SetFromDetails(FInstancedPropertyBag& PropertyBag, int32 ColumnIndex, int32 RowIndex)
-	{
-		FName PropertyName("RowData",ColumnIndex);
+void FMultiEnumColumn::SetFromDetails(FInstancedPropertyBag& PropertyBag, int32 ColumnIndex, int32 RowIndex)
+{
+	FName PropertyName("RowData",ColumnIndex);
 
-		TValueOrError<FStructView, EPropertyBagResult> Result = PropertyBag.GetValueStruct(PropertyName, FChooserMultiEnumRowData::StaticStruct());
-		if (const FStructView* StructView = Result.TryGetValue())
+	if (const UEnum* Enum = InputValue.Get<FChooserParameterEnumBase>().GetEnum())
+	{
+		TValueOrError<const FPropertyBagArrayRef, EPropertyBagResult> Result = PropertyBag.GetArrayRef(PropertyName);
+		if (const FPropertyBagArrayRef* ArrayRef = Result.TryGetValue())
 		{
-			RowValues[RowIndex] = StructView->Get<FChooserMultiEnumRowData>();
+			RowValues[RowIndex].Value = 0;
+			for(int32 i = 0; i<ArrayRef->Num(); i++)
+			{
+				RowValues[RowIndex].Value  |= (1 << ArrayRef->GetValueEnum(i, Enum).GetValue());
+			}
 		}
 	}
+}
 #endif
\ No newline at end of file
