Fix broken icons on ProxyTable asset types

diff --git a/Engine/Plugins/Chooser/Source/ProxyTableEditor/Private/ProxyTableEditorStyle.cpp b/Engine/Plugins/Chooser/Source/ProxyTableEditor/Private/ProxyTableEditorStyle.cpp
--- a/Engine/Plugins/Chooser/Source/ProxyTableEditor/Private/ProxyTableEditorStyle.cpp
+++ b/Engine/Plugins/Chooser/Source/ProxyTableEditor/Private/ProxyTableEditorStyle.cpp
@@ -18,7 +18,7 @@ FProxyTableEditorStyle::FProxyTableEditorStyle() :
 	const FVector2D Icon16x16(16.0f, 16.0f);
 	const FVector2D Icon24x24(24.0f, 24.0f);
 
-	SetContentRoot(FPaths::EnginePluginsDir() / TEXT("Experimental/Chooser/Content"));
+	SetContentRoot(FPaths::EnginePluginsDir() / TEXT("Chooser/Content"));
 	SetCoreContentRoot(FPaths::EngineContentDir() / TEXT(""));
 
 	// proxy asset icon
