Fix move row/column command key bindings not working in Chooser Table Editor

diff --git a/Engine/Plugins/Chooser/Source/ChooserEditor/Private/ChooserTableEditor.cpp b/Engine/Plugins/Chooser/Source/ChooserEditor/Private/ChooserTableEditor.cpp
--- a/Engine/Plugins/Chooser/Source/ChooserEditor/Private/ChooserTableEditor.cpp
+++ b/Engine/Plugins/Chooser/Source/ChooserEditor/Private/ChooserTableEditor.cpp
@@ -1554,6 +1554,14 @@ TSharedRef<SDockTab> FChooserTableEditor::SpawnTableTab( const FSpawnTabArgs& Ar
 	UpdateTableColumns();
 
 	TableView = SNew(SListView<TSharedPtr<FChooserTableRow>>)
+				.OnKeyDownHandler_Lambda([this](const FGeometry&, const FKeyEvent& Event)
+				{
+					if (GetToolkitCommands()->ProcessCommandBindings(Event))
+					{
+						return FReply::Handled();
+					}
+					return FReply::Unhandled();
+				})
     			.ListItemsSource(&TableRows)
 				.OnContextMenuOpening_Raw(this, &FChooserTableEditor::GenerateRowContextMenu)
 				.OnSelectionChanged_Lambda([this](TSharedPtr<FChooserTableRow>,  ESelectInfo::Type SelectInfo)
