apply application scale factor to chooser factory dialog

diff --git a/Engine/Plugins/Chooser/Source/ChooserEditor/Private/ChooserFactory.cpp b/Engine/Plugins/Chooser/Source/ChooserEditor/Private/ChooserFactory.cpp
--- a/Engine/Plugins/Chooser/Source/ChooserEditor/Private/ChooserFactory.cpp
+++ b/Engine/Plugins/Chooser/Source/ChooserEditor/Private/ChooserFactory.cpp
@@ -4,6 +4,7 @@
 #include "Chooser.h"
 #include "ChooserInitializer.h"
 #include "ChooserEditorSettings.h"
+#include "Framework/Application/SlateApplication.h"
 #include "Widgets/Input/SButton.h"
 #include "Widgets/Layout/SUniformGridPanel.h"
 #include "Widgets/SCompoundWidget.h"
@@ -72,6 +73,7 @@ public:
 					.HAlign(HAlign_Center)
 					.ContentPadding( FAppStyle::GetMargin("StandardDialog.ContentPadding") )
 					.OnClicked(this, &SChooserCreateDialog::OkClicked)
+					.IsEnabled_Lambda([this]() { return ChooserFactory->ChooserInitializer.IsValid(); } )
 					.Text(LOCTEXT("Create", "Create"))
 				]
 				+SUniformGridPanel::Slot(1,0)
@@ -100,11 +102,13 @@ public:
 		}
 		
 		DetailsView->SetObject(ChooserFactory.Get());
+		
+		const float AppScale = FSlateApplication::Get().GetApplicationScale();
 
 		Window = SNew(SWindow)
 		.Title( LOCTEXT("Create Chooser Options", "Create Chooser Table") )
 		.SizingRule(ESizingRule::FixedSize)
-		.ClientSize(FVector2f(550, 300))
+		.ClientSize(FVector2f(AppScale*550, AppScale*300))
 		.SupportsMinimize(false)
 		.SupportsMaximize(false)
 		[
