Fix chooser table editor AutoPopulate command when rows are selected

diff --git a/Engine/Plugins/Chooser/Source/ChooserEditor/Private/ChooserTableEditor.cpp b/Engine/Plugins/Chooser/Source/ChooserEditor/Private/ChooserTableEditor.cpp
--- a/Engine/Plugins/Chooser/Source/ChooserEditor/Private/ChooserTableEditor.cpp
+++ b/Engine/Plugins/Chooser/Source/ChooserEditor/Private/ChooserTableEditor.cpp
@@ -1911,14 +1911,14 @@ void FChooserTableEditor::AutoPopulateSelection()
 	{
 		const FScopedTransaction Transaction(LOCTEXT("Auto Populate Selection", "Auto Populate Selection"));
 		Chooser->Modify();
-		if (SelectedColumn)
+		if (HasColumnSelected())
 		{
 			if (Chooser->ColumnsStructs.IsValidIndex(SelectedColumn->Column))
 			{
 				AutoPopulateColumn(Chooser->ColumnsStructs[SelectedColumn->Column].GetMutable<FChooserColumnBase>());
 			}
 		}
-		else
+		else if (HasRowsSelected())
 		{
 			for(UChooserRowDetails* RowDetails : SelectedRows)
 			{
