Fix double click in rewind debugger to attach Chooser Table Editor debugging.  When the Chooser was already open, it would fail to enable the display of debug values in the column header

diff --git a/Engine/Plugins/Chooser/Source/ChooserEditor/Private/ChooserTableEditor.cpp b/Engine/Plugins/Chooser/Source/ChooserEditor/Private/ChooserTableEditor.cpp
--- a/Engine/Plugins/Chooser/Source/ChooserEditor/Private/ChooserTableEditor.cpp
+++ b/Engine/Plugins/Chooser/Source/ChooserEditor/Private/ChooserTableEditor.cpp
@@ -200,11 +200,13 @@ FChooserTableEditor::~FChooserTableEditor()
 	if (SelectedColumn)
 	{
 		SelectedColumn->ClearFlags(RF_Standalone);
+		SelectedColumn = nullptr;
 	}
 	for (UObject* SelectedRow : SelectedRows)
 	{
 		SelectedRow->ClearFlags(RF_Standalone);
 	}
+	SelectedRows.Empty();
 	
 	FCoreUObjectDelegates::OnObjectsReplaced.RemoveAll(this);
 	FCoreUObjectDelegates::OnObjectTransacted.RemoveAll(this);
@@ -734,6 +736,8 @@ void FChooserTableEditor::FocusWindow(UObject* ObjectToFocusOn)
 	{
 		SetChooserTableToEdit(Chooser);
 	}
+	// refresh, even if we set the same chooser we were already editing. (Rewind Debugger double click enables debug testing, which requires recreating the header widgets)
+	RefreshAll();
 	FAssetEditorToolkit::FocusWindow(ObjectToFocusOn);
 }
 
@@ -2584,3 +2588,7 @@ void FChooserTableEditor::RegisterWidgets()
 
 #undef LOCTEXT_NAMESPACE
 
+void UChooserColumnDetails::BeginDestroy()
+{
+	UObject::BeginDestroy();
+}
