Fix for using ChooserTrace from outside the chooser module

diff --git a/Engine/Plugins/Chooser/Source/Chooser/Internal/ChooserTrace.h b/Engine/Plugins/Chooser/Source/Chooser/Internal/ChooserTrace.h
--- a/Engine/Plugins/Chooser/Source/Chooser/Internal/ChooserTrace.h
+++ b/Engine/Plugins/Chooser/Source/Chooser/Internal/ChooserTrace.h
@@ -15,14 +15,14 @@ class UObject;
 
 struct FChooserEvaluationContext;
 
-struct FChooserTrace
+struct CHOOSER_API FChooserTrace
 {
-	CHOOSER_API static void OutputChooserEvaluation(const UObject* ChooserAsset, const FChooserEvaluationContext& Context, uint32 SelectedIndex);
+	static void OutputChooserEvaluation(const UObject* ChooserAsset, const FChooserEvaluationContext& Context, uint32 SelectedIndex);
 	
-	CHOOSER_API static void OutputChooserValueArchive(const FChooserEvaluationContext& Context, const TCHAR* Key, const FBufferArchive& ValueArchive);
+	static void OutputChooserValueArchive(const FChooserEvaluationContext& Context, const TCHAR* Key, const FBufferArchive& ValueArchive);
 	
 	template<typename T>
-	CHOOSER_API static void OutputChooserValue(const FChooserEvaluationContext& Context, const TCHAR* Key, const T& Value)
+	static void OutputChooserValue(const FChooserEvaluationContext& Context, const TCHAR* Key, const T& Value)
 	{
 		FBufferArchive Archive;
 		Archive << const_cast<T&>(Value);
