Fix build.

diff --git a/Engine/Plugins/Chooser/Source/ChooserEditor/Private/ChooserTableEditor.cpp b/Engine/Plugins/Chooser/Source/ChooserEditor/Private/ChooserTableEditor.cpp
--- a/Engine/Plugins/Chooser/Source/ChooserEditor/Private/ChooserTableEditor.cpp
+++ b/Engine/Plugins/Chooser/Source/ChooserEditor/Private/ChooserTableEditor.cpp
@@ -2588,7 +2588,3 @@ void FChooserTableEditor::RegisterWidgets()
 
 #undef LOCTEXT_NAMESPACE
 
-void UChooserColumnDetails::BeginDestroy()
-{
-	UObject::BeginDestroy();
-}
