Fix Crash ChooserTableEditor due to 'SelectedColumn' not added to GC.
- Add SelectedColumn to Root like 'SLevelEditor' does for 'SelectedElements'

diff --git a/Engine/Plugins/Chooser/Source/ChooserEditor/Private/ChooserTableEditor.cpp b/Engine/Plugins/Chooser/Source/ChooserEditor/Private/ChooserTableEditor.cpp
--- a/Engine/Plugins/Chooser/Source/ChooserEditor/Private/ChooserTableEditor.cpp
+++ b/Engine/Plugins/Chooser/Source/ChooserEditor/Private/ChooserTableEditor.cpp
@@ -199,7 +199,7 @@ FChooserTableEditor::~FChooserTableEditor()
 {
 	if (SelectedColumn)
 	{
-		SelectedColumn->ClearFlags(RF_Standalone);
+		SelectedColumn->RemoveFromRoot();
 		SelectedColumn = nullptr;
 	}
 	for (UObject* SelectedRow : SelectedRows)
@@ -1737,7 +1737,7 @@ void FChooserTableEditor::SelectColumn(UChooserTable* ChooserEditor, int Index)
    		if (SelectedColumn == nullptr)
    		{
    			SelectedColumn = NewObject<UChooserColumnDetails>();
-   			SelectedColumn->SetFlags(RF_Standalone);
+			SelectedColumn->AddToRoot();
    		}
    
    		SelectedColumn->Chooser = Chooser;
