Proxy Table Editor Crash fixes

diff --git a/Engine/Plugins/Chooser/Source/ProxyTableEditor/Private/ProxyTableEditor.cpp b/Engine/Plugins/Chooser/Source/ProxyTableEditor/Private/ProxyTableEditor.cpp
--- a/Engine/Plugins/Chooser/Source/ProxyTableEditor/Private/ProxyTableEditor.cpp
+++ b/Engine/Plugins/Chooser/Source/ProxyTableEditor/Private/ProxyTableEditor.cpp
@@ -117,6 +117,8 @@ const FName FProxyTableEditor::ProxyEditorAppIdentifier( TEXT( "ProxyEditorApp"
 
 FProxyTableEditor::~FProxyTableEditor()
 {
+	ClearSelectedRows();
+	
 	GEditor->GetEditorSubsystem<UImportSubsystem>()->OnAssetPostImport.RemoveAll(this);
 	FCoreUObjectDelegates::OnObjectsReplaced.RemoveAll(this);
 	FCoreUObjectDelegates::OnObjectTransacted.RemoveAll(this);
@@ -712,7 +714,10 @@ void FProxyTableEditor::DeleteSelectedRows()
 	TArray<uint32> RowsToDelete;
 	for(auto& SelectedRow:SelectedRows)
 	{
-		RowsToDelete.Add(SelectedRow->Row);
+		if (ProxyTable->Entries.IsValidIndex(SelectedRow->Row))
+		{
+			RowsToDelete.Add(SelectedRow->Row);
+		}
 	}
 	// sort indices in reverse
 	RowsToDelete.Sort([](int32 A, int32 B){ return A>B; });
@@ -728,6 +733,10 @@ void FProxyTableEditor::DeleteSelectedRows()
 	
 void FProxyTableEditor::ClearSelectedRows() 
 {
+	for(UObject* SelectedRow : SelectedRows)
+	{
+		SelectedRow->RemoveFromRoot();
+	}
 	SelectedRows.SetNum(0);
 	TableView->ClearSelection();
 	SelectRootProperties();
@@ -877,7 +886,7 @@ TSharedRef<SDockTab> FProxyTableEditor::SpawnTableTab( const FSpawnTabArgs& Args
 					{
 						for (UObject* SelectedRow : SelectedRows)
                      	{
-                     		SelectedRow->ClearFlags(RF_Standalone);
+                     		SelectedRow->RemoveFromRoot();
                      	}
 						SelectedRows.SetNum(0);
 						UProxyTable* ProxyTable = Cast<UProxyTable>(EditingObjects[0]);
@@ -885,7 +894,7 @@ TSharedRef<SDockTab> FProxyTableEditor::SpawnTableTab( const FSpawnTabArgs& Args
 						TObjectPtr<UProxyRowDetails> Selection = NewObject<UProxyRowDetails>();
 						Selection->ProxyTable = ProxyTable;
 						Selection->Row = SelectedItem->RowIndex;
-						Selection->SetFlags(RF_Standalone);
+						Selection->AddToRoot();
 						SelectedRows.Add(Selection);
 						
 						TArray<UObject*> DetailsObjects;
