Add "(Preview)" to ChooserTable debugging targets that are in preview worlds

diff --git a/Engine/Plugins/Chooser/Source/Chooser/Private/Chooser.cpp b/Engine/Plugins/Chooser/Source/Chooser/Private/Chooser.cpp
--- a/Engine/Plugins/Chooser/Source/Chooser/Private/Chooser.cpp
+++ b/Engine/Plugins/Chooser/Source/Chooser/Private/Chooser.cpp
@@ -362,7 +362,11 @@ void UChooserTable::UpdateDebugging(FChooserEvaluationContext& Context) const
 
 				if (UWorld* World = ContextObject->GetWorld())
 				{
-					if (World->GetNetMode() == ENetMode::NM_DedicatedServer)
+					if (World->IsPreviewWorld())
+					{
+						DebugName += " (Preview)";
+					}
+					else if (World->GetNetMode() == ENetMode::NM_DedicatedServer)
 					{
 						DebugName += " (Server)";
 					}
diff --git a/Engine/Plugins/Chooser/Source/ChooserEditor/Private/RewindDebuggerChooser.cpp b/Engine/Plugins/Chooser/Source/ChooserEditor/Private/RewindDebuggerChooser.cpp
--- a/Engine/Plugins/Chooser/Source/ChooserEditor/Private/RewindDebuggerChooser.cpp
+++ b/Engine/Plugins/Chooser/Source/ChooserEditor/Private/RewindDebuggerChooser.cpp
@@ -45,8 +45,12 @@ void FRewindDebuggerChooser::Update(float DeltaTime, IRewindDebugger* RewindDebu
 					if (const FWorldInfo* WorldInfo = GameplayProvider->FindWorldInfoFromObject(OwnerId))
 					{
 						const FObjectInfo& WorldObjectInfo = GameplayProvider->GetObjectInfo(WorldInfo->Id);
-
-						if (WorldInfo->NetMode == FWorldInfo::ENetMode::DedicatedServer)
+						
+						if (WorldInfo->Type == FWorldInfo::EType::EditorPreview)
+						{
+							DebugName += " (Preview)";
+						}
+						else if (WorldInfo->NetMode == FWorldInfo::ENetMode::DedicatedServer)
 						{
 							DebugName += " (Server)";
 						}
