Fix crash copy/pasting None nested chooser row

diff --git a/Engine/Plugins/Chooser/Source/ChooserEditor/Private/ChooserTableEditor.cpp b/Engine/Plugins/Chooser/Source/ChooserEditor/Private/ChooserTableEditor.cpp
--- a/Engine/Plugins/Chooser/Source/ChooserEditor/Private/ChooserTableEditor.cpp
+++ b/Engine/Plugins/Chooser/Source/ChooserEditor/Private/ChooserTableEditor.cpp
@@ -2269,8 +2269,11 @@ UChooserTable* FChooserTableEditor::CopySelectionInternal()
 			CopyData->ResultsStructs[RowIndex] = Chooser->ResultsStructs[SelectedRowsCopy[RowIndex]->Row];
 			if (FNestedChooser* CopiedNestedChooser = CopyData->ResultsStructs[RowIndex].GetMutablePtr<FNestedChooser>())
 			{
-				// if the result for this row was a nested chooser, duplicate it
-				CopiedNestedChooser->Chooser = DuplicateNestedChooser(CopiedNestedChooser->Chooser, CopyData);
+				if (CopiedNestedChooser->Chooser)
+				{
+					// if the result for this row was a nested chooser (with a valid chooser assigned), duplicate it
+					CopiedNestedChooser->Chooser = DuplicateNestedChooser(CopiedNestedChooser->Chooser, CopyData);
+				}
 			}
 			
 			CopyData->DisabledRows[RowIndex] = Chooser->DisabledRows[SelectedRowsCopy[RowIndex]->Row];
