Chooser Tables: Add support for gameplay tag query columns

[RN] minor framework
#tests Both chooser table manual debug and inside a level blueprint in PIE, and it returned the proper rows based on input gameplay tag container & the column's query for all combinations of the input (some rows matching, all rows matching, no rows matching)

diff --git a/Engine/Plugins/Chooser/Source/Chooser/Internal/GameplayTagQueryColumn.h b/Engine/Plugins/Chooser/Source/Chooser/Internal/GameplayTagQueryColumn.h
new file mode 100644
--- /dev/null
+++ b/Engine/Plugins/Chooser/Source/Chooser/Internal/GameplayTagQueryColumn.h
@@ -0,0 +1,64 @@
+// Copyright Epic Games, Inc. All Rights Reserved.
+#pragma once
+
+#include "IChooserColumn.h"
+#include "IChooserParameterGameplayTag.h"
+#include "ChooserPropertyAccess.h"
+#include "GameplayTagContainer.h"
+#include "StructUtils/InstancedStruct.h"
+#include "Serialization/MemoryReader.h"
+#include "GameplayTagQueryColumn.generated.h"
+
+USTRUCT(DisplayName = "Gameplay Tag Query", Meta = (Category = "Filter", Tooltip = "A column which filters rows by matching the input against a Gameplay Tag Query."))
+struct CHOOSER_API FGameplayTagQueryColumn : public FChooserColumnBase
+{
+	GENERATED_BODY()
+	public:
+	FGameplayTagQueryColumn();
+	
+	UPROPERTY(EditAnywhere, Meta = (ExcludeBaseStruct, BaseStruct = "/Script/Chooser.ChooserParameterGameplayTagBase", ToolTip="The GameplayTagContainer property this column will filter by apply the query on"), Category = "Data")
+	FInstancedStruct InputValue;
+
+#if WITH_EDITORONLY_DATA
+	UPROPERTY(EditAnywhere, Category=Data, meta=(ToolTip="DefaultRowValue will be assigned to cells when new rows are created"));
+	FGameplayTagQuery DefaultRowValue;
+#endif
+	
+	UPROPERTY()
+	// array of results (cells for this column for each row in the table)
+	// should match the length of the Results array 
+	TArray<FGameplayTagQuery> RowValues;
+
+	bool TestRow(int32 RowIndex, const FGameplayTagContainer& Value) const;
+	
+	virtual void Filter(FChooserEvaluationContext& Context, const FChooserIndexArray& IndexListIn, FChooserIndexArray& IndexListOut) const override;
+
+#if WITH_EDITOR
+	mutable FGameplayTagContainer TestValue;
+	virtual bool EditorTestFilter(int32 RowIndex) const override
+	{
+		return TestRow(RowIndex, TestValue);
+	}
+
+	virtual void SetTestValue(TArrayView<const uint8> Value) override
+    {
+    	FMemoryReaderView Reader(Value);
+		FString Tags;
+    	Reader << Tags;
+		TestValue.FromExportString(Tags);
+    }
+	
+	virtual void AddToDetails(FInstancedPropertyBag& PropertyBag, int32 ColumnIndex, int32 RowIndex) override;
+	virtual void SetFromDetails(FInstancedPropertyBag& PropertyBag, int32 ColumnIndex, int32 RowIndex) override;
+#endif
+
+	virtual void PostLoad() override
+	{
+		if (InputValue.IsValid())
+		{
+			InputValue.GetMutable<FChooserParameterBase>().PostLoad();
+		}
+	}
+
+	CHOOSER_COLUMN_BOILERPLATE(FChooserParameterGameplayTagBase);
+};
diff --git a/Engine/Plugins/Chooser/Source/Chooser/Private/GameplayTagQueryColumn.cpp b/Engine/Plugins/Chooser/Source/Chooser/Private/GameplayTagQueryColumn.cpp
new file mode 100644
--- /dev/null
+++ b/Engine/Plugins/Chooser/Source/Chooser/Private/GameplayTagQueryColumn.cpp
@@ -0,0 +1,77 @@
+// Copyright Epic Games, Inc. All Rights Reserved.
+#include "GameplayTagQueryColumn.h"
+
+#include "ChooserIndexArray.h"
+#include "ChooserTrace.h"
+#include "GameplayTagColumn.h"
+#if WITH_EDITOR
+#include "StructUtils/PropertyBag.h"
+#endif
+
+
+FGameplayTagQueryColumn::FGameplayTagQueryColumn()
+{
+	InputValue.InitializeAs(FGameplayTagContextProperty::StaticStruct());
+}
+
+bool FGameplayTagQueryColumn::TestRow(int32 RowIndex, const FGameplayTagContainer& Value) const
+{
+	if (RowValues.IsValidIndex(RowIndex))
+	{
+		return RowValues[RowIndex].Matches(Value);
+	}
+	return false;
+}
+
+void FGameplayTagQueryColumn::Filter(FChooserEvaluationContext& Context, const FChooserIndexArray& IndexListIn, FChooserIndexArray& IndexListOut) const
+{
+	const FGameplayTagContainer* Result = nullptr;
+	if (InputValue.IsValid() && InputValue.Get<FChooserParameterGameplayTagBase>().GetValue(Context,Result))
+	{
+		TRACE_CHOOSER_VALUE(Context, ToCStr(InputValue.Get<FChooserParameterBase>().GetDebugName()), Result->ToString());
+
+#if WITH_EDITOR
+		if (Context.DebuggingInfo.bCurrentDebugTarget)
+		{
+			TestValue = *Result;
+		}
+#endif
+
+		for (const FChooserIndexArray::FIndexData& IndexData : IndexListIn)
+		{
+			if (TestRow(IndexData.Index, *Result))
+			{
+				IndexListOut.Push(IndexData);
+			}
+		}
+	}
+	else
+	{
+		// passthrough fallback (behaves better during live editing)
+		IndexListOut = IndexListIn;
+	}
+}
+
+#if WITH_EDITOR
+	void FGameplayTagQueryColumn::AddToDetails(FInstancedPropertyBag& PropertyBag, int32 ColumnIndex, int32 RowIndex)
+	{
+		FText DisplayName;
+		InputValue.Get<FChooserParameterBase>().GetDisplayName(DisplayName);
+		FName PropertyName("RowData", ColumnIndex);
+		FPropertyBagPropertyDesc PropertyDesc(PropertyName, EPropertyBagPropertyType::Struct, FGameplayTagQuery::StaticStruct());
+		PropertyDesc.MetaData.Add(FPropertyBagPropertyDescMetaData("DisplayName", DisplayName.ToString()));
+		PropertyBag.AddProperties({PropertyDesc});
+		PropertyBag.SetValueStruct(PropertyName, RowValues[RowIndex]);
+	}
+
+	void FGameplayTagQueryColumn::SetFromDetails(FInstancedPropertyBag& PropertyBag, int32 ColumnIndex, int32 RowIndex)
+	{
+		FName PropertyName("RowData", ColumnIndex);
+		
+		TValueOrError<FStructView, EPropertyBagResult> Result = PropertyBag.GetValueStruct(PropertyName, FGameplayTagQuery::StaticStruct());
+		if (FStructView* StructView = Result.TryGetValue())
+		{
+			RowValues[RowIndex] = StructView->Get<FGameplayTagQuery>();
+		}
+	}	
+#endif
diff --git a/Engine/Plugins/Chooser/Source/ChooserEditor/Private/ChooserEditorModule.cpp b/Engine/Plugins/Chooser/Source/ChooserEditor/Private/ChooserEditorModule.cpp
--- a/Engine/Plugins/Chooser/Source/ChooserEditor/Private/ChooserEditorModule.cpp
+++ b/Engine/Plugins/Chooser/Source/ChooserEditor/Private/ChooserEditorModule.cpp
@@ -15,6 +15,7 @@
 #include "FloatRangeColumnEditor.h"
 #include "FrameTimeCustomization.h"
 #include "GameplayTagColumnEditor.h"
+#include "GameplayTagQueryColumnEditor.h"
 #include "IAssetTools.h"
 #include "MultiEnumColumnEditor.h"
 #include "ObjectColumnEditor.h"
@@ -59,6 +60,7 @@ void FModule::StartupModule()
 	
 	FChooserTableEditor::RegisterWidgets();
 	RegisterGameplayTagWidgets();
+	RegisterGameplayTagQueryWidgets();
 	RegisterFloatDistanceWidgets();
 	RegisterFloatRangeWidgets();
 	RegisterOutputFloatWidgets();
diff --git a/Engine/Plugins/Chooser/Source/ChooserEditor/Private/GameplayTagQueryColumnEditor.cpp b/Engine/Plugins/Chooser/Source/ChooserEditor/Private/GameplayTagQueryColumnEditor.cpp
new file mode 100644
--- /dev/null
+++ b/Engine/Plugins/Chooser/Source/ChooserEditor/Private/GameplayTagQueryColumnEditor.cpp
@@ -0,0 +1,90 @@
+// Copyright Epic Games, Inc. All Rights Reserved.
+
+#include "GameplayTagQueryColumnEditor.h"
+
+#include <ChooserColumnHeader.h>
+
+#include "GameplayTagQueryColumn.h"
+#include "ObjectChooserWidgetFactories.h"
+#include "SGameplayTagQueryEntryBox.h"
+#include "SGameplayTagWidget.h"
+#include "SSimpleComboButton.h"
+#include "Widgets/Images/SImage.h"
+
+#define LOCTEXT_NAMESPACE "FGameplayTagQueryColumnEditor"
+
+namespace UE::ChooserEditor
+{
+
+TSharedRef<SWidget> CreateGameplayTagQueryColumnWidget(UChooserTable* Chooser, FChooserColumnBase* Column, int Row)
+{
+	FGameplayTagQueryColumn* GameplayTagQueryColumn = static_cast<struct FGameplayTagQueryColumn*>(Column);
+	
+	if (Row == ColumnWidget_SpecialIndex_Fallback)
+	{
+		return SNullWidget::NullWidget;
+	}
+	else if (Row == ColumnWidget_SpecialIndex_Header)
+	{
+		// create column header widget
+		const FSlateBrush* ColumnIcon = FCoreStyle::Get().GetBrush("Icons.Filter");
+		const FText ColumnTooltip = LOCTEXT("Gameplay Tag Query Tooltip", "Gameplay Tag Query: cells pass if the input gameplay tag collection matches the query specified in the column properties. Note that empty queries never pass.");
+		const FText ColumnName = LOCTEXT("Gameplay Tag Query","Gameplay Tag Query");
+		
+		TSharedPtr<SWidget> DebugWidget = nullptr;
+		if (Chooser->GetEnableDebugTesting())
+		{
+			DebugWidget = SNew(SSimpleComboButton)
+				.IsEnabled_Lambda([Chooser]()
+				{
+					 return !Chooser->HasDebugTarget();
+				})
+				.Text_Lambda([GameplayTagQueryColumn]()
+				{
+					FText Text = FText::FromString(GameplayTagQueryColumn->TestValue.ToStringSimple(false));
+					if (Text.IsEmpty())
+					{
+						Text = LOCTEXT("None", "None");
+					}
+					return Text;
+				})	
+				.OnGetMenuContent_Lambda([Chooser, GameplayTagQueryColumn]()
+				{
+					TArray<SGameplayTagWidget::FEditableGameplayTagContainerDatum> EditableContainers;
+					EditableContainers.Emplace(Chooser, &(GameplayTagQueryColumn->TestValue));
+					return TSharedRef<SWidget>(SNew(SGameplayTagWidget, EditableContainers));
+				});
+		}
+
+		return MakeColumnHeaderWidget(Chooser, Column, ColumnName, ColumnTooltip, ColumnIcon, DebugWidget);
+	}
+
+	// create cell widget
+	return SNew(SGameplayTagQueryEntryBox)
+		.TagQuery_Lambda([GameplayTagQueryColumn, Row]()
+		{
+			if (Row < GameplayTagQueryColumn->RowValues.Num())
+			{
+				return GameplayTagQueryColumn->RowValues[Row];
+			}
+			return FGameplayTagQuery::EmptyQuery;
+		})
+		.ReadOnly(false)
+		.OnTagQueryChanged_Lambda([GameplayTagQueryColumn, Row](const FGameplayTagQuery& UpdatedQuery)
+		{
+			if (Row < GameplayTagQueryColumn->RowValues.Num())
+			{
+				GameplayTagQueryColumn->RowValues[Row] = UpdatedQuery;
+			}
+		});
+}
+	
+void RegisterGameplayTagQueryWidgets()
+{
+	FObjectChooserWidgetFactories::RegisterColumnWidgetCreator(FGameplayTagQueryColumn::StaticStruct(), CreateGameplayTagQueryColumnWidget);
+	// No need to make and register a creator for gameplay tag containers - it's already registered in GameplayTagColumnEditor
+}
+	
+}
+
+#undef LOCTEXT_NAMESPACE
diff --git a/Engine/Plugins/Chooser/Source/ChooserEditor/Private/GameplayTagQueryColumnEditor.h b/Engine/Plugins/Chooser/Source/ChooserEditor/Private/GameplayTagQueryColumnEditor.h
new file mode 100644
--- /dev/null
+++ b/Engine/Plugins/Chooser/Source/ChooserEditor/Private/GameplayTagQueryColumnEditor.h
@@ -0,0 +1,8 @@
+// Copyright Epic Games, Inc. All Rights Reserved.
+
+#pragma once
+
+namespace UE::ChooserEditor
+{
+	void RegisterGameplayTagQueryWidgets();
+}
