Chooser tables: Add NoPrimaryResult chooser type. This is mostly UX fluff for when you are just interested into writing output values (especially integral ones), and not into returning an asset or class per se.

# Adapt EvaluateChooser BP node to hide the Result pin and Mode setting in case we are working with NoPrimaryResult
# Fix EvaluateChooser BP node compiling itself out along what followed its Then pin if no value pin was connected (it was setup as a Pure node)
# Fix EvaluateChooser BP node showing input struct pins as by-value when they are by-reference (the node modifies Read/Write parameters in-place)

[RN] minor framework
#tests Inside a level blueprint in PIE, used the EvaluateChooser node in all its forms with each kind of supported tables (old and this new one), was behaving properly in the BP editor. At runtime, I could verify it wrote back the proper output values based on the supplied input.

diff --git a/Engine/Plugins/Chooser/Source/Chooser/Internal/Chooser.h b/Engine/Plugins/Chooser/Source/Chooser/Internal/Chooser.h
--- a/Engine/Plugins/Chooser/Source/Chooser/Internal/Chooser.h
+++ b/Engine/Plugins/Chooser/Source/Chooser/Internal/Chooser.h
@@ -180,12 +180,12 @@ public:
 	UPROPERTY(EditAnywhere, NoClear, DisplayName = "Results", Meta = (ExcludeBaseStruct, BaseStruct = "/Script/Chooser.ObjectChooserBase"), Category = "Hidden")
 	TArray<FInstancedStruct> ColumnsStructs;
 	
-	// The kind of output this chooser has (Object or Class)
-    UPROPERTY(EditAnywhere, DisplayName = "Result Type", Category="Result")
+	// The kind of output this chooser has (Object or Class or No primary result)
+    UPROPERTY(EditAnywhere, DisplayName = "Result Type", Category="Result", Meta = (EditConditionHides, EditCondition="ResultType != EObjectChooserResultType::NoPrimaryResult"))
     EObjectChooserResultType ResultType = EObjectChooserResultType::ObjectResult;
 
 	// The Class of Object this Chooser returns when ResultType is set to ObjectOfType, or the Parent Class of the Classes returned by this chooser when ResultType is set to ClassOfType
-	UPROPERTY(EditAnywhere, DisplayName= "Result Class", Category="Result", Meta = (AllowAbstract=true))
+	UPROPERTY(EditAnywhere, DisplayName= "Result Class", Category="Result", Meta = (AllowAbstract=true, EditConditionHides, EditCondition="ResultType != EObjectChooserResultType::NoPrimaryResult"))
 	TObjectPtr<UClass> OutputObjectType;
 	
 	// Parameter Objects or Structs from which the chooser can read or write properties 
diff --git a/Engine/Plugins/Chooser/Source/Chooser/Public/IHasContext.h b/Engine/Plugins/Chooser/Source/Chooser/Public/IHasContext.h
--- a/Engine/Plugins/Chooser/Source/Chooser/Public/IHasContext.h
+++ b/Engine/Plugins/Chooser/Source/Chooser/Public/IHasContext.h
@@ -44,4 +44,6 @@ enum class EObjectChooserResultType
 	ObjectResult UMETA(DisplayName = "Object Of Type"),
 	// The Chooser returns a Class that is a SubClass of the specified Result Class (eg a chooser could return a type of Character to spawn, or a type of AnimInstance to link)
 	ClassResult UMETA(DisplayName = "SubClass Of"),
+	// The Chooser returns nothing, but can write to one or more outputs (useful if you are only interested in returning integral types like a float or string)
+	NoPrimaryResult UMETA(Hidden),
 };
diff --git a/Engine/Plugins/Chooser/Source/ChooserEditor/Private/ChooserInitializer.cpp b/Engine/Plugins/Chooser/Source/ChooserEditor/Private/ChooserInitializer.cpp
--- a/Engine/Plugins/Chooser/Source/ChooserEditor/Private/ChooserInitializer.cpp
+++ b/Engine/Plugins/Chooser/Source/ChooserEditor/Private/ChooserInitializer.cpp
@@ -5,7 +5,6 @@
 #include "AnimNode_ChooserPlayer.h"
 #include "Chooser.h"
 #include "ChooserPropertyAccess.h"
-#include "ObjectChooserWidgetFactories.h"
 #include "Animation/AnimationAsset.h"
 #include "Animation/AnimInstance.h"
 
@@ -42,5 +41,14 @@ void FChooserPlayerInitializer::Initialize(UChooserTable* Chooser) const
 	Chooser->OutputObjectType = UAnimationAsset::StaticClass();
 }
 
+void FNoPrimaryResultChooserInitializer::Initialize(UChooserTable* Chooser) const
+{
+	Chooser->ContextData = ContextData;
+	Chooser->ResultType = EObjectChooserResultType::NoPrimaryResult;
+
+	// Dummy result type
+	Chooser->OutputObjectType = UClass::StaticClass();
+}
+
 
 #undef LOCTEXT_NAMESPACE
diff --git a/Engine/Plugins/Chooser/Source/ChooserEditor/Private/ChooserTableEditor.cpp b/Engine/Plugins/Chooser/Source/ChooserEditor/Private/ChooserTableEditor.cpp
--- a/Engine/Plugins/Chooser/Source/ChooserEditor/Private/ChooserTableEditor.cpp
+++ b/Engine/Plugins/Chooser/Source/ChooserEditor/Private/ChooserTableEditor.cpp
@@ -1072,19 +1072,22 @@ void FChooserTableEditor::UpdateTableColumns()
 					.DefaultLabel(FText())
 					.ManualWidth(30));
 
-	HeaderRow->AddColumn(SHeaderRow::Column("Result")
-					.ManualWidth(300)
-					.HeaderContent()
-					[
-						SNew(SVerticalBox)
-						+ SVerticalBox::Slot()
-						.VAlign(VAlign_Top)
+	if (Chooser->ResultType != EObjectChooserResultType::NoPrimaryResult)
+	{
+		HeaderRow->AddColumn(SHeaderRow::Column("Result")
+						.ManualWidth(300)
+						.HeaderContent()
 						[
-							SNew(STextBlock)
-								.Text(LOCTEXT("Result", "Result"))
-								.ToolTipText(LOCTEXT("ResultTooltip", "The Result is the asset which will be returned if a row is selected (or other Chooser to evaluate to get the asset to return"))
-						]
-					]);
+							SNew(SVerticalBox)
+							+ SVerticalBox::Slot()
+							.VAlign(VAlign_Top)
+							[
+								SNew(STextBlock)
+									.Text(LOCTEXT("Result", "Result"))
+									.ToolTipText(LOCTEXT("ResultTooltip", "The Result is the asset which will be returned if a row is selected (or other Chooser to evaluate to get the asset to return"))
+							]
+						]);
+	}
 
 	FName ColumnId("ChooserColumn", 1);
 	int NumColumns = Chooser->ColumnsStructs.Num();	
@@ -1210,10 +1213,13 @@ void FChooserTableEditor::RefreshRowSelectionDetails()
 			Selection->Chooser = Chooser;
 			Selection->Row = SelectedItem->RowIndex;
 			Selection->SetFlags(RF_Standalone | RF_Transactional);
-
-			FInstancedStruct& Result = Chooser->ResultsStructs[SelectedItem->RowIndex];
-			Selection->Properties.AddProperties({ResultPropertyDesc});
-			Selection->Properties.SetValueStruct("Result", FConstStructView(FInstancedStruct::StaticStruct(), reinterpret_cast<uint8*>(&Result)));
+			
+			if (Chooser->ResultType != EObjectChooserResultType::NoPrimaryResult)
+			{
+				FInstancedStruct& Result = Chooser->ResultsStructs[SelectedItem->RowIndex];
+				Selection->Properties.AddProperties({ResultPropertyDesc});
+				Selection->Properties.SetValueStruct("Result", FConstStructView(FInstancedStruct::StaticStruct(), reinterpret_cast<uint8*>(&Result)));
+			}
 
 			int ColumnIndex = 0;
 			for (FInstancedStruct& ColumnData : Chooser->ColumnsStructs)
@@ -1238,9 +1244,12 @@ void FChooserTableEditor::RefreshRowSelectionDetails()
 			Selection->Row = SelectedItem->RowIndex;
 			Selection->SetFlags(RF_Standalone | RF_Transactional);
 
-			FInstancedStruct& Result = Chooser->FallbackResult;
-			Selection->Properties.AddProperties({ResultPropertyDesc});
-			Selection->Properties.SetValueStruct("Result", FConstStructView(FInstancedStruct::StaticStruct(), reinterpret_cast<uint8*>(&Result)));
+			if (Chooser->ResultType != EObjectChooserResultType::NoPrimaryResult)
+			{
+				FInstancedStruct& Result = Chooser->FallbackResult;
+				Selection->Properties.AddProperties({ResultPropertyDesc});
+				Selection->Properties.SetValueStruct("Result", FConstStructView(FInstancedStruct::StaticStruct(), reinterpret_cast<uint8*>(&Result)));
+			}
 
 			int ColumnIndex = 0;
 			for (FInstancedStruct& ColumnData : Chooser->ColumnsStructs)
@@ -1500,35 +1509,86 @@ TSharedRef<SWidget>	FChooserTableEditor::MakeCreateRowMenu()
     UChooserTable* Chooser = GetChooser();
 	if (!Chooser->FallbackResult.IsValid())
 	{
-		MenuBuilder.AddSubMenu(
-			LOCTEXT("Add Fallback", "Add Fallback Result"),
-				LOCTEXT("Add Fallback Tooltip", "Add a Fallback row to the chooser, which will be used in the case where no other rows passed all filter columns"),
-				FNewMenuDelegate::CreateLambda([this](FMenuBuilder& MenuBuilder)
+		if (Chooser->ResultType == EObjectChooserResultType::NoPrimaryResult)
+		{
+			MenuBuilder.AddMenuEntry(
+				LOCTEXT("Add Fallback Output", "Add Fallback Output"),
+				LOCTEXT("Add Fallback Output Tooltip", "Add a Fallback row to the chooser, which will be used in the case where no other rows passed all filter columns"),
+				FSlateIcon(),
+				FUIAction(FExecuteAction::CreateLambda([this]()
 				{
 					UChooserTable* Chooser = GetChooser();
-					MakeCreateResultMenu(MenuBuilder, Chooser->ResultType, FCreateStructDelegate::CreateLambda([this](UScriptStruct* Type)
+					const FScopedTransaction Transaction(LOCTEXT("Add Fallback Row Transaction", "Add Fallback Row"));
+					Chooser->Modify(true);
+					
+					// Just construct a dummy result to make sure all rows always have "valid results"
+					// You can't just leave a null result otherwise rows don't apply their output.
+					Chooser->FallbackResult.InitializeAs(FClassChooser::StaticStruct());
+					Chooser->FallbackResult.GetMutable<FClassChooser>().Class = UClass::StaticClass();
+					
+					UpdateTableRows();
+				}))
+			);
+		}
+		else
+		{
+			MenuBuilder.AddSubMenu(
+				LOCTEXT("Add Fallback", "Add Fallback Result"),
+					LOCTEXT("Add Fallback Tooltip", "Add a Fallback row to the chooser, which will be used in the case where no other rows passed all filter columns"),
+					FNewMenuDelegate::CreateLambda([this](FMenuBuilder& MenuBuilder)
 					{
 						UChooserTable* Chooser = GetChooser();
-						const FScopedTransaction Transaction(LOCTEXT("Add Fallback Row Transaction", "Add Fallback Row"));
-						Chooser->Modify(true);
-						Chooser->FallbackResult.InitializeAs(Type);
-						UpdateTableRows();
-					}));
-				})
-			);
-
+						MakeCreateResultMenu(MenuBuilder, Chooser->ResultType, FCreateStructDelegate::CreateLambda([this](UScriptStruct* Type)
+						{
+							UChooserTable* Chooser = GetChooser();
+							const FScopedTransaction Transaction(LOCTEXT("Add Fallback Row Transaction", "Add Fallback Row"));
+							Chooser->Modify(true);
+							
+							Chooser->FallbackResult.InitializeAs(Type);
+							
+							UpdateTableRows();
+						}));
+					})
+				);
+		}
 	}
-
-	MakeCreateResultMenu(MenuBuilder, Chooser->ResultType, FCreateStructDelegate::CreateLambda([this](UScriptStruct* Type)
+	
+	if (Chooser->ResultType == EObjectChooserResultType::NoPrimaryResult)
 	{
-
-		UChooserTable* Chooser = GetChooser();
-		const FScopedTransaction Transaction(LOCTEXT("Add Row Transaction", "Add Row"));
-		Chooser->Modify(true);
-		Chooser->ResultsStructs.SetNum(Chooser->ResultsStructs.Num()+1);
-		Chooser->ResultsStructs.Last().InitializeAs(Type);
-		UpdateTableRows();
-	}));
+		MenuBuilder.AddMenuEntry(
+			LOCTEXT("Add Output Row", "Add Output Row"),
+			LOCTEXT("Add Output Row", "Add a regular row to the chooser"),
+			FSlateIcon(),
+			FUIAction(FExecuteAction::CreateLambda([this]()
+			{
+				UChooserTable* Chooser = GetChooser();
+				const FScopedTransaction Transaction(LOCTEXT("Add Row Transaction", "Add Row"));
+				Chooser->Modify(true);
+
+				// Just construct a dummy result to make sure all rows always have "valid results"
+				// You can't just leave a null result otherwise rows don't apply their output.
+				FInstancedStruct& NewResult = Chooser->ResultsStructs.AddDefaulted_GetRef();
+				NewResult.InitializeAs(FClassChooser::StaticStruct());
+				NewResult.GetMutable<FClassChooser>().Class = UClass::StaticClass();
+				
+				UpdateTableRows();
+			}))
+		);
+	}
+	else
+	{
+		MakeCreateResultMenu(MenuBuilder, Chooser->ResultType, FCreateStructDelegate::CreateLambda([this](UScriptStruct* Type)
+		{
+			UChooserTable* Chooser = GetChooser();
+			const FScopedTransaction Transaction(LOCTEXT("Add Row Transaction", "Add Row"));
+			Chooser->Modify(true);
+			
+			FInstancedStruct& NewResult = Chooser->ResultsStructs.AddDefaulted_GetRef();
+			NewResult.InitializeAs(Type);
+			
+			UpdateTableRows();
+		}));
+	}
 
 	return MenuBuilder.MakeWidget();
 }
diff --git a/Engine/Plugins/Chooser/Source/ChooserEditor/Private/SChooserTableRow.cpp b/Engine/Plugins/Chooser/Source/ChooserEditor/Private/SChooserTableRow.cpp
--- a/Engine/Plugins/Chooser/Source/ChooserEditor/Private/SChooserTableRow.cpp
+++ b/Engine/Plugins/Chooser/Source/ChooserEditor/Private/SChooserTableRow.cpp
@@ -340,8 +340,22 @@ namespace UE::ChooserEditor
 		}
 		else if (RowIndex->RowIndex == SpecialIndex_AddRow)
 		{
-			// on the row past the end, show an Add button in the result column
-			if (ColumnName == Result)
+			// on the row past the end, show an Add button in the first column available
+			bool bIsLeftmostColumn;
+			if (Chooser->ResultType != EObjectChooserResultType::NoPrimaryResult)
+			{
+				bIsLeftmostColumn = ColumnName == Result;
+			}
+			else if (Chooser->ColumnsStructs.IsEmpty())
+			{
+				bIsLeftmostColumn = ColumnName == AddColumn;
+			}
+			else
+			{
+				bIsLeftmostColumn = ColumnName.GetNumber() == 1;
+			}
+			
+			if (bIsLeftmostColumn)
 			{
 				return SNew(SHorizontalBox)
 								+SHorizontalBox::Slot().AutoWidth()
diff --git a/Engine/Plugins/Chooser/Source/ChooserEditor/Public/ChooserInitializer.h b/Engine/Plugins/Chooser/Source/ChooserEditor/Public/ChooserInitializer.h
--- a/Engine/Plugins/Chooser/Source/ChooserEditor/Public/ChooserInitializer.h
+++ b/Engine/Plugins/Chooser/Source/ChooserEditor/Public/ChooserInitializer.h
@@ -45,4 +45,16 @@ struct FChooserPlayerInitializer : public FChooserInitializer
 
 	UPROPERTY(EditAnywhere, Category="AnimClass")
 	TObjectPtr<UClass> AnimClass;
-};
\ No newline at end of file
+};
+
+USTRUCT(DisplayName="No Primary Result Chooser", Meta = (ToolTip="A ChooserTable for use in Blueprint, which returns no primary result but writes to outputs (useful if you are interested in returning only integral types like a float or string). Note: this table can't evaluate or nest other tables, and will set its outputs to the first row that matches (no 'multi' mode).") )
+struct FNoPrimaryResultChooserInitializer : public FChooserInitializer
+{
+	GENERATED_BODY()
+
+	virtual void Initialize(UChooserTable* Chooser) const override;
+	
+	// Parameter Objects or Structs from which the chooser can read or write properties 
+	UPROPERTY(EditAnywhere, DisplayName = "Parameters", NoClear, Meta = (ExpandByDefault, ExcludeBaseStruct, BaseStruct = "/Script/Chooser.ContextObjectTypeBase"), Category = "Parameters")
+	TArray<FInstancedStruct> ContextData;
+};
diff --git a/Engine/Plugins/Chooser/Source/ChooserUncooked/Private/EvaluateChooserNode.cpp b/Engine/Plugins/Chooser/Source/ChooserUncooked/Private/EvaluateChooserNode.cpp
--- a/Engine/Plugins/Chooser/Source/ChooserUncooked/Private/EvaluateChooserNode.cpp
+++ b/Engine/Plugins/Chooser/Source/ChooserUncooked/Private/EvaluateChooserNode.cpp
@@ -11,11 +11,10 @@
 #include "EdGraph/EdGraph.h"
 #include "EdGraph/EdGraphSchema.h"
 #include "EdGraphSchema_K2.h"
-#include "EdGraphSchema_K2_Actions.h"
 #include "EditorCategoryUtils.h"
+#include "IHasContext.h"
 #include "Internationalization/Internationalization.h"
 #include "K2Node_CallFunction.h"
-#include "K2Node_MakeArray.h"
 #include "K2Node_MakeStruct.h"
 #include "Kismet2/BlueprintEditorUtils.h"
 #include "Kismet2/CompilerResultsLog.h"
@@ -355,10 +354,22 @@ void UK2Node_EvaluateChooser2::ChooserChanged()
 	{
 		UnregisterChooserCallback();
 	
+		bShowModeSetting = false;
+		
 		if (Chooser)
 		{
 			Chooser->OnOutputObjectTypeChanged.AddUObject(this, &UK2Node_EvaluateChooser2::ResultTypeChanged);
 			Chooser->OnContextClassChanged.AddUObject(this, &UK2Node::ReconstructNode);
+			
+			if (Chooser->ResultType != EObjectChooserResultType::NoPrimaryResult)
+			{
+				bShowModeSetting = true;
+			}
+			else
+			{
+				// NoPrimaryResult tables only set the output to their first result anyway
+				Mode = EEvaluateChooserMode::FirstResult;
+			}
 		}
 	
 		CurrentCallbackChooser = Chooser;
@@ -388,8 +399,6 @@ void UK2Node_EvaluateChooser2::PreloadRequiredAssets()
 void UK2Node_EvaluateChooser2::AllocateDefaultPins()
 {
 	Super::AllocateDefaultPins();
-	
-	UClass* ResultType = UObject::StaticClass();
 
 	if (!FindPin(UEdGraphSchema_K2::PN_Execute, EGPD_Input))
 	{
@@ -405,6 +414,12 @@ void UK2Node_EvaluateChooser2::AllocateDefaultPins()
 
 	if (Chooser)
 	{
+		// We show input struct contexts pins as "references", because the supplied struct might get modified if it's an In/Out param.
+		// We also show output struct contexts pins as references, to show they are there just for convenience and actually pointing to the very same
+		// struct that the one on the input side.
+		FCreatePinParams StructContextPinCreationParams;
+		StructContextPinCreationParams.bIsReference = true;
+		
 		for (FInstancedStruct& ContextDataEntry : Chooser->ContextData)
 		{
 			if (ContextDataEntry.IsValid())
@@ -434,7 +449,7 @@ void UK2Node_EvaluateChooser2::AllocateDefaultPins()
 						}
 						else
 						{
-							CreatePin(EGPD_Input, UEdGraphSchema_K2::PC_Struct, StructContext.Struct, StructContext.Struct.GetFName());
+							CreatePin(EGPD_Input, UEdGraphSchema_K2::PC_Struct, StructContext.Struct, StructContext.Struct.GetFName(), StructContextPinCreationParams);
 						}
 					}
 
@@ -446,7 +461,7 @@ void UK2Node_EvaluateChooser2::AllocateDefaultPins()
 						}
 						else
 						{
-							CreatePin(EGPD_Output, UEdGraphSchema_K2::PC_Struct, StructContext.Struct, StructContext.Struct.GetFName());
+							CreatePin(EGPD_Output, UEdGraphSchema_K2::PC_Struct, StructContext.Struct, StructContext.Struct.GetFName(), StructContextPinCreationParams);
 						}
 					}
 
@@ -470,28 +485,41 @@ void UK2Node_EvaluateChooser2::AllocateDefaultPins()
 				}
 			}
 		}
-	}
-
-	bool bResultIsClass = false;
-   	if (Chooser && Chooser->OutputObjectType)
-   	{
-   		ResultType = Chooser->OutputObjectType;
-   		bResultIsClass = Chooser->ResultType == EObjectChooserResultType::ClassResult;
-   	}
+		
+		UEdGraphPin* ResultPin = FindPin(TEXT("Result"), EGPD_Output);
+		if (Chooser->ResultType == EObjectChooserResultType::NoPrimaryResult)
+		{
+			// No need to keep existing result pin if this chooser doesn't work with primary results
+			if (ResultPin)
+			{
+				RemovePin(ResultPin);
+			}
+		}
+		else
+		{
+			UClass* ResultType = UObject::StaticClass();
+			bool bResultIsClass = false;
+			if (Chooser->OutputObjectType)
+			{
+				ResultType = Chooser->OutputObjectType;
+				bResultIsClass = Chooser->ResultType == EObjectChooserResultType::ClassResult;
+			}
 
-	if (UEdGraphPin* ResultPin = FindPin(TEXT("Result"), EGPD_Output))
-	{
-		ResultPin->PinType.PinCategory = bResultIsClass ? UEdGraphSchema_K2::PC_Class : UEdGraphSchema_K2::PC_Object;
-		ResultPin->PinType.PinSubCategoryObject = ResultType;
-		ResultPin->PinType.ContainerType = (Mode == EEvaluateChooserMode::AllResults) ? EPinContainerType::Array : EPinContainerType::None;
-	}
-	else
-	{
-		UEdGraphNode::FCreatePinParams PinParams;
-		PinParams.ContainerType = (Mode == EEvaluateChooserMode::AllResults) ? EPinContainerType::Array : EPinContainerType::None;
-		PinParams.ValueTerminalType.TerminalCategory =  bResultIsClass ? UEdGraphSchema_K2::PC_Class : UEdGraphSchema_K2::PC_Object;
+			if (ResultPin)
+			{
+				ResultPin->PinType.PinCategory = bResultIsClass ? UEdGraphSchema_K2::PC_Class : UEdGraphSchema_K2::PC_Object;
+				ResultPin->PinType.PinSubCategoryObject = ResultType;
+				ResultPin->PinType.ContainerType = (Mode == EEvaluateChooserMode::AllResults) ? EPinContainerType::Array : EPinContainerType::None;
+			}
+			else
+			{
+				UEdGraphNode::FCreatePinParams PinParams;
+				PinParams.ContainerType = (Mode == EEvaluateChooserMode::AllResults) ? EPinContainerType::Array : EPinContainerType::None;
+				PinParams.ValueTerminalType.TerminalCategory =  bResultIsClass ? UEdGraphSchema_K2::PC_Class : UEdGraphSchema_K2::PC_Object;
 				
-		CreatePin(EGPD_Output, PinParams.ValueTerminalType.TerminalCategory, ResultType, TEXT("Result"), PinParams);
+				CreatePin(EGPD_Output, PinParams.ValueTerminalType.TerminalCategory, ResultType, TEXT("Result"), PinParams);
+			}
+		}
 	}
 }
 
@@ -575,7 +603,8 @@ void UK2Node_EvaluateChooser2::ExpandNode(class FKismetCompilerContext& Compiler
 
 	UEdGraphPin* ExecInput = GetExecPin();
 	UEdGraphPin* ExecOutput = GetPassThroughPin(ExecInput);
-	UEdGraphPin* ResultPin = FindPinChecked(TEXT("Result"));
+	UEdGraphPin* ResultPin = FindPin(TEXT("Result"));
+	check(ResultPin || (Chooser && Chooser->ResultType == EObjectChooserResultType::NoPrimaryResult));
 
 	UEdGraphPin* PreviousNodeExecOutput = nullptr;
 	
@@ -801,9 +830,11 @@ void UK2Node_EvaluateChooser2::ExpandNode(class FKismetCompilerContext& Compiler
 				}
 			}
 		}
-				
 
-		CompilerContext.MovePinLinksToIntermediate(*ResultPin, *OutputPin);
+		if (ResultPin)
+		{
+			CompilerContext.MovePinLinksToIntermediate(*ResultPin, *OutputPin);
+		}
 	}
 	
 	BreakAllNodeLinks();
diff --git a/Engine/Plugins/Chooser/Source/ChooserUncooked/Private/EvaluateChooserNode.h b/Engine/Plugins/Chooser/Source/ChooserUncooked/Private/EvaluateChooserNode.h
--- a/Engine/Plugins/Chooser/Source/ChooserUncooked/Private/EvaluateChooserNode.h
+++ b/Engine/Plugins/Chooser/Source/ChooserUncooked/Private/EvaluateChooserNode.h
@@ -106,7 +106,7 @@ class UK2Node_EvaluateChooser2 : public UK2Node
 	//~ End UEdGraphNode Interface.
 
 	//~ Begin UK2Node Interface.
-	virtual bool IsNodePure() const override { return true; }
+	virtual bool IsNodePure() const override { return false; }
 	virtual bool NodeCausesStructuralBlueprintChange() const override { return true; }
 	virtual void PostReconstructNode() override;
 	virtual void ExpandNode(class FKismetCompilerContext& CompilerContext, UEdGraph* SourceGraph) override;
@@ -130,9 +130,12 @@ private:
 	UPROPERTY(EditAnywhere, Category = "Chooser")
 	TObjectPtr<UChooserTable> Chooser;
 
-	UPROPERTY(EditAnywhere, Category = "Chooser")
+	UPROPERTY(EditAnywhere, Category = "Chooser", Meta = (EditConditionHides, EditCondition="bShowModeSetting"))
 	EEvaluateChooserMode Mode;
 
+	UPROPERTY(Transient)
+	bool bShowModeSetting = false;
+
 	/** Tooltip text for this node. */
 	FText NodeTooltip;
-};
\ No newline at end of file
+};
