Fix nested chooser duplication

diff --git a/Engine/Plugins/Chooser/Source/ChooserEditor/Private/ChooserTableEditor.cpp b/Engine/Plugins/Chooser/Source/ChooserEditor/Private/ChooserTableEditor.cpp
--- a/Engine/Plugins/Chooser/Source/ChooserEditor/Private/ChooserTableEditor.cpp
+++ b/Engine/Plugins/Chooser/Source/ChooserEditor/Private/ChooserTableEditor.cpp
@@ -2249,7 +2249,8 @@ void FChooserTableEditor::MoveColumnRight()
 
 UChooserTable* DuplicateNestedChooser(UChooserTable* Chooser, UChooserTable* NewOuter)
 {
-	if (TObjectPtr<UChooserTable>* FoundTable = NewOuter->NestedChoosers.FindByPredicate([Chooser](UChooserTable* Table)
+	UChooserTable* RootTable = NewOuter->GetRootChooser();
+	if (TObjectPtr<UChooserTable>* FoundTable = RootTable->NestedChoosers.FindByPredicate([Chooser](UChooserTable* Table)
 		{
 			return Table->GetName() == Chooser->GetName();
 		}))
@@ -2261,8 +2262,8 @@ UChooserTable* DuplicateNestedChooser(UChooserTable* Chooser, UChooserTable* New
 	UChooserTable* NewTable = NewObject<UChooserTable>(NewOuter, Chooser->GetFName());
 	NewTable->ColumnsStructs = Chooser->ColumnsStructs;
 	NewTable->ResultsStructs = Chooser->ResultsStructs;
-	NewTable->RootChooser = NewOuter;
-	NewOuter->NestedChoosers.Add(NewTable);
+	NewTable->RootChooser = RootTable;
+	RootTable->AddNestedChooser(NewTable);
 
 	for (FInstancedStruct& ResultData : NewTable->ResultsStructs)
 	{
