Fix LOCTEXT duplicate identifier

diff --git a/Engine/Plugins/Chooser/Source/ChooserEditor/Private/ChooserTableEditor.cpp b/Engine/Plugins/Chooser/Source/ChooserEditor/Private/ChooserTableEditor.cpp
--- a/Engine/Plugins/Chooser/Source/ChooserEditor/Private/ChooserTableEditor.cpp
+++ b/Engine/Plugins/Chooser/Source/ChooserEditor/Private/ChooserTableEditor.cpp
@@ -1557,7 +1557,7 @@ TSharedRef<SWidget>	FChooserTableEditor::MakeCreateRowMenu()
 	{
 		MenuBuilder.AddMenuEntry(
 			LOCTEXT("Add Output Row", "Add Output Row"),
-			LOCTEXT("Add Output Row", "Add a regular row to the chooser"),
+			LOCTEXT("Add Output Row Tooltip", "Add a regular row to the chooser"),
 			FSlateIcon(),
 			FUIAction(FExecuteAction::CreateLambda([this]()
 			{
