Fix issue with binding an input struct directly to a GameplayTags column (as opposed to an FGameplayTagContainer property on an input struct/class)

diff --git a/Engine/Plugins/Chooser/Source/ChooserEditor/Private/SPropertyAccessChainWidget.cpp b/Engine/Plugins/Chooser/Source/ChooserEditor/Private/SPropertyAccessChainWidget.cpp
--- a/Engine/Plugins/Chooser/Source/ChooserEditor/Private/SPropertyAccessChainWidget.cpp
+++ b/Engine/Plugins/Chooser/Source/ChooserEditor/Private/SPropertyAccessChainWidget.cpp
@@ -93,7 +93,8 @@ TSharedRef<SWidget> SPropertyAccessChainWidget::CreatePropertyAccessWidget()
 			}
 			else
 			{
-				return StructType->GetName()  == TypeFilter;
+				const FString CPPName = StructType->GetPrefixCPP() +  StructType->GetName();
+				return CPPName  == TypeFilter;
 			}
 		}
 		return false;
