Remove space from OutputFloatColumn.cpp filename

diff --git a/Engine/Plugins/Chooser/Source/Chooser/Private/OutputFloatColumn .cpp b/Engine/Plugins/Chooser/Source/Chooser/Private/OutputFloatColumn .cpp
deleted file mode 100644
diff --git a/Engine/Plugins/Chooser/Source/Chooser/Private/OutputFloatColumn.cpp b/Engine/Plugins/Chooser/Source/Chooser/Private/OutputFloatColumn.cpp
new file mode 100644
--- /dev/null
+++ b/Engine/Plugins/Chooser/Source/Chooser/Private/OutputFloatColumn.cpp
@@ -0,0 +1,63 @@
+// Copyright Epic Games, Inc. All Rights Reserved.
+#include "OutputFloatColumn.h"
+#include "ChooserPropertyAccess.h"
+#include "FloatRangeColumn.h"
+#if WITH_EDITOR
+#include "StructUtils/PropertyBag.h"
+#endif
+
+FOutputFloatColumn::FOutputFloatColumn()
+{
+	InputValue.InitializeAs(FFloatContextProperty::StaticStruct());
+}
+
+void FOutputFloatColumn::SetOutputs(FChooserEvaluationContext& Context, int RowIndex) const
+{
+	if (InputValue.IsValid())
+	{
+		if (RowIndex == ChooserColumn_SpecialIndex_Fallback || RowValues.IsValidIndex(RowIndex))
+		{
+			double Value = GetValueForIndex(RowIndex);
+			InputValue.Get<FChooserParameterFloatBase>().SetValue(Context, Value);
+						
+			#if WITH_EDITOR
+			if (Context.DebuggingInfo.bCurrentDebugTarget)
+			{
+				TestValue = Value;
+			}
+			#endif
+		}
+		else
+		{
+#if CHOOSER_DEBUGGING_ENABLED
+			UE_ASSET_LOG(LogChooser, Error, Context.DebuggingInfo.CurrentChooser, TEXT("Invalid index passed to FOutputFloatColumn::SetOutputs"), RowIndex);
+#else
+			UE_LOG(LogChooser, Error, TEXT("Invalid index %d passed to FOutputFloatColumn::SetOutputs"), RowIndex);
+#endif
+		}
+	}
+}
+
+#if WITH_EDITOR
+	void FOutputFloatColumn::AddToDetails(FInstancedPropertyBag& PropertyBag, int32 ColumnIndex, int32 RowIndex)
+	{
+		FText DisplayName;
+		InputValue.Get<FChooserParameterBase>().GetDisplayName(DisplayName);
+		FName PropertyName("RowData",ColumnIndex);
+		FPropertyBagPropertyDesc PropertyDesc(PropertyName, EPropertyBagPropertyType::Float);
+		PropertyDesc.MetaData.Add(FPropertyBagPropertyDescMetaData("DisplayName", DisplayName.ToString()));
+		PropertyBag.AddProperties({PropertyDesc});
+		PropertyBag.SetValueFloat(PropertyName, GetValueForIndex(RowIndex));
+	}
+
+	void FOutputFloatColumn::SetFromDetails(FInstancedPropertyBag& PropertyBag, int32 ColumnIndex, int32 RowIndex)
+	{
+		FName PropertyName("RowData", ColumnIndex);
+			
+		TValueOrError<float, EPropertyBagResult> Result = PropertyBag.GetValueFloat(PropertyName);
+		if (float* Value = Result.TryGetValue())
+		{
+			GetValueForIndex(RowIndex) = *Value;
+		}
+	}
+#endif
\ No newline at end of file
