Chooser: Work around Internal directory being included from Public in Chooser.
This will need to be properly resolved.

#rnx

diff --git a/Engine/Plugins/Chooser/Source/Chooser/Chooser.Build.cs b/Engine/Plugins/Chooser/Source/Chooser/Chooser.Build.cs
--- a/Engine/Plugins/Chooser/Source/Chooser/Chooser.Build.cs
+++ b/Engine/Plugins/Chooser/Source/Chooser/Chooser.Build.cs
@@ -5,7 +5,10 @@ namespace UnrealBuildTool.Rules
 	public class Chooser : ModuleRules
 	{
 		public Chooser(ReadOnlyTargetRules Target) : base(Target)
-		{			
+		{
+			// TODO: Chooser/Public/ChooserFunctionLibrary.h is including Chooser/Internal/Chooser.h
+			PublicIncludePaths.Add(System.IO.Path.Combine(ModuleDirectory, "Internal"));
+
 			PublicDependencyModuleNames.AddRange(
 				new string[]
 				{
