Update Chooser Table Editor to save and restore column widths

diff --git a/Engine/Plugins/Chooser/Source/Chooser/Internal/Chooser.h b/Engine/Plugins/Chooser/Source/Chooser/Internal/Chooser.h
--- a/Engine/Plugins/Chooser/Source/Chooser/Internal/Chooser.h
+++ b/Engine/Plugins/Chooser/Source/Chooser/Internal/Chooser.h
@@ -165,6 +165,9 @@ public:
 	// deprecated in favor of RootChooser
 	UPROPERTY()
 	TObjectPtr<UChooserTable> ParentTable;
+	
+	UPROPERTY()
+	float EditorResultsColumnWidth = 300;
 
 	UPROPERTY()
 	uint32 Version = 0;
diff --git a/Engine/Plugins/Chooser/Source/Chooser/Public/IChooserColumn.h b/Engine/Plugins/Chooser/Source/Chooser/Public/IChooserColumn.h
--- a/Engine/Plugins/Chooser/Source/Chooser/Public/IChooserColumn.h
+++ b/Engine/Plugins/Chooser/Source/Chooser/Public/IChooserColumn.h
@@ -92,6 +92,9 @@ public:
 #if WITH_EDITORONLY_DATA
 	UPROPERTY(EditAnywhere, Category=Base)
 	bool bDisabled = false;
+	
+	UPROPERTY()
+	float EditorColumnWidth = 200;
 #endif
 };
 
diff --git a/Engine/Plugins/Chooser/Source/ChooserEditor/Private/ChooserTableEditor.cpp b/Engine/Plugins/Chooser/Source/ChooserEditor/Private/ChooserTableEditor.cpp
--- a/Engine/Plugins/Chooser/Source/ChooserEditor/Private/ChooserTableEditor.cpp
+++ b/Engine/Plugins/Chooser/Source/ChooserEditor/Private/ChooserTableEditor.cpp
@@ -1075,7 +1075,8 @@ void FChooserTableEditor::UpdateTableColumns()
 	if (Chooser->ResultType != EObjectChooserResultType::NoPrimaryResult)
 	{
 		HeaderRow->AddColumn(SHeaderRow::Column("Result")
-						.ManualWidth(300)
+						.ManualWidth_Lambda([Chooser]() { return Chooser->EditorResultsColumnWidth; } )
+						.OnWidthChanged_Lambda([Chooser](float NewWidth) { Chooser->EditorResultsColumnWidth = NewWidth; })
 						.HeaderContent()
 						[
 							SNew(SVerticalBox)
@@ -1103,7 +1104,9 @@ void FChooserTableEditor::UpdateTableColumns()
 		
 		HeaderRow->AddColumn(SHeaderRow::FColumn::FArguments()
 			.ColumnId(ColumnId)
-			.ManualWidth(200)
+			.ManualWidth(Column.EditorColumnWidth)
+			.ManualWidth_Lambda([&Column]() { return Column.EditorColumnWidth; } )
+			.OnWidthChanged_Lambda([&Column](float NewWidth) { Column.EditorColumnWidth = NewWidth; })
 			.HeaderComboVisibility(EHeaderComboVisibility::Ghosted)
 			.HeaderContent()
 			[
