Fix multi-edit support in PropertyAccessChainCustomization for choosers, and add GetPerObjectValueData to PropertyHandle

diff --git a/Engine/Plugins/Chooser/Source/Chooser/Private/ChooserPropertyAccess.cpp b/Engine/Plugins/Chooser/Source/Chooser/Private/ChooserPropertyAccess.cpp
--- a/Engine/Plugins/Chooser/Source/Chooser/Private/ChooserPropertyAccess.cpp
+++ b/Engine/Plugins/Chooser/Source/Chooser/Private/ChooserPropertyAccess.cpp
@@ -311,6 +311,53 @@ void FChooserPropertyBinding::Compile(IHasContextClass* Owner, bool bForce)
  	}
 }
 
+#if WITH_EDITORONLY_DATA
+void FChooserEnumPropertyBinding::SetPropertyData(const IHasContextClass* HasContext, FField* Property)
+{
+	if (const FEnumProperty* EnumProperty = CastField<const FEnumProperty>(Property))
+	{
+		Enum = EnumProperty->GetEnum();
+	}
+	else if (const FByteProperty* ByteProperty = CastField<const FByteProperty>(Property))
+	{
+		Enum = ByteProperty->Enum;
+	}
+}
+
+void FChooserObjectPropertyBinding::SetPropertyData(const IHasContextClass* HasContext, FField* Property)
+{
+	if (const FObjectPropertyBase* ObjectProperty = CastField<const FObjectPropertyBase>(Property))
+	{
+		AllowedClass = ObjectProperty->PropertyClass;
+	}
+}
+
+void FChooserStructPropertyBinding::SetPropertyData(const IHasContextClass* HasContext, FField* Property)
+{
+	StructType = nullptr;
+
+	if (const FStructProperty* StructProperty = CastField<const FStructProperty>(Property))
+	{
+		StructType = StructProperty->Struct;
+	}
+	else if (this->IsBoundToRoot)
+	{
+		// direct binding to a context struct
+		if (HasContext)
+		{
+			TConstArrayView<FInstancedStruct> ContextData = HasContext->GetContextData();
+			if (ContextData.IsValidIndex(ContextIndex))
+			{
+				if (const FContextObjectTypeStruct* StructContext = ContextData[ContextIndex].GetPtr<FContextObjectTypeStruct>())
+				{
+					StructType = StructContext->Struct;
+				}
+			}
+		}
+	}
+}
+#endif
+
 
 namespace UE::Chooser
 {
diff --git a/Engine/Plugins/Chooser/Source/Chooser/Public/ChooserPropertyAccess.h b/Engine/Plugins/Chooser/Source/Chooser/Public/ChooserPropertyAccess.h
--- a/Engine/Plugins/Chooser/Source/Chooser/Public/ChooserPropertyAccess.h
+++ b/Engine/Plugins/Chooser/Source/Chooser/Public/ChooserPropertyAccess.h
@@ -79,6 +79,8 @@ USTRUCT()
 struct CHOOSER_API FChooserPropertyBinding 
 {
 	GENERATED_BODY()
+
+	virtual ~FChooserPropertyBinding() {}
 	
 	UPROPERTY()
 	TArray<FName> PropertyBindingChain;
@@ -94,6 +96,8 @@ struct CHOOSER_API FChooserPropertyBinding
 	FString DisplayName;
 
 	FText CompileMessage;
+
+	virtual void SetPropertyData(const IHasContextClass* HasContext, FField* Property) {}
 #endif
 
 	TSharedPtr<UE::Chooser::FCompiledBinding> CompiledBinding;
@@ -101,7 +105,6 @@ struct CHOOSER_API FChooserPropertyBinding
 	FName GetUniqueId() const;
 
 	void Compile(IHasContextClass* HasContext, bool bForce = false);
-
 	
 	template <typename T>
 	bool GetValuePtr(FChooserEvaluationContext& Context, T*& Value) const;
@@ -117,35 +120,41 @@ struct CHOOSER_API FChooserPropertyBinding
 };
 
 USTRUCT()
-struct FChooserEnumPropertyBinding : public FChooserPropertyBinding
+struct CHOOSER_API FChooserEnumPropertyBinding : public FChooserPropertyBinding
 {
 	GENERATED_BODY()
 	
 #if WITH_EDITORONLY_DATA
 	UPROPERTY()
 	TObjectPtr<const UEnum> Enum = nullptr;
+
+	virtual void SetPropertyData(const IHasContextClass* HasContext, FField* Property) override;
 #endif
 };
 
 USTRUCT()
-struct FChooserObjectPropertyBinding : public FChooserPropertyBinding
+struct CHOOSER_API FChooserObjectPropertyBinding : public FChooserPropertyBinding
 {
 	GENERATED_BODY()
 
 #if WITH_EDITORONLY_DATA
 	UPROPERTY()
 	TObjectPtr<UClass> AllowedClass = nullptr;
+	
+	virtual void SetPropertyData(const IHasContextClass* HasContext, FField* Property) override;
 #endif
 };
 
 USTRUCT()
-struct FChooserStructPropertyBinding : public FChooserPropertyBinding
+struct CHOOSER_API FChooserStructPropertyBinding : public FChooserPropertyBinding
 {
 	GENERATED_BODY()
 
 #if WITH_EDITORONLY_DATA
 	UPROPERTY()
 	TObjectPtr<UScriptStruct> StructType = nullptr;
+	
+	virtual void SetPropertyData(const IHasContextClass* HasContext, FField* Property) override;
 #endif
 };
 
diff --git a/Engine/Plugins/Chooser/Source/ChooserEditor/Private/PropertyAccessChainCustomization.cpp b/Engine/Plugins/Chooser/Source/ChooserEditor/Private/PropertyAccessChainCustomization.cpp
--- a/Engine/Plugins/Chooser/Source/ChooserEditor/Private/PropertyAccessChainCustomization.cpp
+++ b/Engine/Plugins/Chooser/Source/ChooserEditor/Private/PropertyAccessChainCustomization.cpp
@@ -62,8 +62,36 @@ void FPropertyAccessChainCustomization::CustomizeHeader(TSharedRef<IPropertyHand
 			PropertyHandle->GetValueData(data);
 			return reinterpret_cast<FChooserPropertyBinding*>(data);
 		})
-		.OnValueChanged_Lambda([PropertyHandle]()
+		.OnAddBinding_Lambda([PropertyHandle, this, HasContext, TypeFilter](FName /*InPropertyName*/, const TArray<FBindingChainElement>& InBindingChain)
 		{
+			FChooserPropertyBinding ContextProperty;
+			SPropertyAccessChainWidget::SetPropertyBinding(HasContext, InBindingChain, ContextProperty);
+
+			UObject* TransactionObject = Cast<UObject>(HasContext);
+            	
+            const FScopedTransaction Transaction(NSLOCTEXT("PropertyAccessChainCustomization", "Set Binding", "Set Binding"));
+   
+			TArray<UObject*> OuterObjects;
+			PropertyHandle->GetOuterObjects(OuterObjects);
+			for (UObject* OuterObject : OuterObjects)
+			{
+				OuterObject->Modify();				
+			}
+			
+			
+			if (FStructProperty* StructProperty = CastField<FStructProperty>(PropertyHandle->GetProperty()))
+			{
+				if (StructProperty->Struct == FChooserPropertyBinding::StaticStruct())
+				{
+					TArray<void*> PropertyDatas;
+					PropertyHandle->AccessRawData(PropertyDatas);
+					for (void* PropertyData : PropertyDatas)
+					{
+						*static_cast<FChooserPropertyBinding*>(PropertyData) = ContextProperty;
+					}
+				}
+			}
+			
 			PropertyHandle->NotifyPostChange(EPropertyChangeType::ValueSet);
 		})
 	];
diff --git a/Engine/Plugins/Chooser/Source/ChooserEditor/Private/SPropertyAccessChainWidget.cpp b/Engine/Plugins/Chooser/Source/ChooserEditor/Private/SPropertyAccessChainWidget.cpp
--- a/Engine/Plugins/Chooser/Source/ChooserEditor/Private/SPropertyAccessChainWidget.cpp
+++ b/Engine/Plugins/Chooser/Source/ChooserEditor/Private/SPropertyAccessChainWidget.cpp
@@ -16,6 +16,58 @@
 namespace UE::ChooserEditor
 {
 	
+void SPropertyAccessChainWidget::SetPropertyBinding(IHasContextClass* HasContext, const TArray<FBindingChainElement>& InBindingChain, FChooserPropertyBinding& OutPropertyBinding)
+{
+	Chooser::CopyPropertyChain(InBindingChain, OutPropertyBinding);
+	
+	FField* Property = InBindingChain.Last().Field.ToField();
+		
+	OutPropertyBinding.DisplayName = "";
+	if (InBindingChain.Num() == 1)
+	{
+		// direct binding to a context struct/class, set display name to the struct/clase name
+		if (HasContext)
+		{
+			TConstArrayView<FInstancedStruct> ContextData = HasContext->GetContextData();
+			if (ContextData.IsValidIndex(InBindingChain[0].ArrayIndex))
+			{
+				if (const FContextObjectTypeStruct* StructContext = ContextData[InBindingChain[0].ArrayIndex].GetPtr<FContextObjectTypeStruct>())
+				{
+					if (StructContext->Struct)
+					{
+						OutPropertyBinding.DisplayName = StructContext->Struct->GetDisplayNameText().ToString();
+					}
+				}
+				else if (const FContextObjectTypeClass* ClassContext = ContextData[InBindingChain[0].ArrayIndex].GetPtr<FContextObjectTypeClass>())
+				{
+					if (ClassContext->Class)
+					{
+						OutPropertyBinding.DisplayName = ClassContext->Class->GetDisplayNameText().ToString();
+					}
+				}
+			}
+		}
+	}
+	else
+	{
+		// set displayname from property name
+		if (Property)
+		{
+			OutPropertyBinding.DisplayName = Property->GetDisplayNameText().ToString();
+			static const int ShortNameLength = 5;
+			if (OutPropertyBinding.DisplayName.Len() < ShortNameLength && InBindingChain.Num() > 2)
+			{
+				FField* ParentProperty = InBindingChain[InBindingChain.Num() - 2].Field.ToField();
+				OutPropertyBinding.DisplayName = ParentProperty->GetDisplayNameText().ToString() + "." + OutPropertyBinding.DisplayName;
+			}
+		}
+	}
+
+	OutPropertyBinding.SetPropertyData(HasContext,Property);
+	
+	OutPropertyBinding.Compile(HasContext);
+}
+
 TSharedRef<SWidget> SPropertyAccessChainWidget::CreatePropertyAccessWidget()
 {
 	FPropertyBindingWidgetArgs Args;
@@ -161,107 +213,14 @@ TSharedRef<SWidget> SPropertyAccessChainWidget::CreatePropertyAccessWidget()
 			{
 				if (PropertyBindingValue.IsSet())
 				{
-					UObject* TransactionObject = Cast<UObject>(ContextClassOwner);
 					FChooserPropertyBinding* ContextProperty = PropertyBindingValue.Get();
-					const FScopedTransaction Transaction(NSLOCTEXT("ContextPropertyWidget", "Change Property Binding", "Change Property Binding"));
-					TransactionObject->Modify(true);
-
-					// todo: move these to a virtual function on FChooserPropertyBinding
-					Chooser::CopyPropertyChain(InBindingChain, *PropertyBindingValue.Get());
-					
-					FField* Property = InBindingChain.Last().Field.ToField();
-						
-					ContextProperty->DisplayName = "";
-					if (InBindingChain.Num() == 1)
-					{
-						// direct binding to a context struct/class, set display name to the struct/clase name
-						if (ContextClassOwner)
-						{
-							TConstArrayView<FInstancedStruct> ContextData = ContextClassOwner->GetContextData();
-							if (ContextData.IsValidIndex(InBindingChain[0].ArrayIndex))
-							{
-								if (const FContextObjectTypeStruct* StructContext = ContextData[InBindingChain[0].ArrayIndex].GetPtr<FContextObjectTypeStruct>())
-								{
-									if (StructContext->Struct)
-									{
-										ContextProperty->DisplayName = StructContext->Struct->GetDisplayNameText().ToString();
-									}
-								}
-								else if (const FContextObjectTypeClass* ClassContext = ContextData[InBindingChain[0].ArrayIndex].GetPtr<FContextObjectTypeClass>())
-								{
-									if (ClassContext->Class)
-									{
-										ContextProperty->DisplayName = ClassContext->Class->GetDisplayNameText().ToString();
-									}
-								}
-							}
-						}
-					}
-					else
-					{
-						// set displayname from property name
-						if (Property)
-						{
-							ContextProperty->DisplayName = Property->GetDisplayNameText().ToString();
-							static const int ShortNameLength = 5;
-							if (ContextProperty->DisplayName.Len() < ShortNameLength && InBindingChain.Num() > 2)
-							{
-								FField* ParentProperty = InBindingChain[InBindingChain.Num() - 2].Field.ToField();
-								ContextProperty->DisplayName = ParentProperty->GetDisplayNameText().ToString() + "." + ContextProperty->DisplayName;
-							}
-						}
-					}
 					
-					if (TypeFilter == "enum")
-					{
-						FChooserEnumPropertyBinding* EnumPropertyValue = static_cast<FChooserEnumPropertyBinding*>(ContextProperty);
-						
-						if (const FEnumProperty* EnumProperty = CastField<const FEnumProperty>(Property))
-						{
-							EnumPropertyValue->Enum = EnumProperty->GetEnum();
-						}
-						else if (const FByteProperty* ByteProperty = CastField<const FByteProperty>(Property))
-						{
-							EnumPropertyValue->Enum = ByteProperty->Enum;
-						}
-					}
-					if (TypeFilter == "object")
-					{
-						FChooserObjectPropertyBinding* ObjectPropertyBinding = static_cast<FChooserObjectPropertyBinding*>(ContextProperty);
-						
-						if (const FObjectPropertyBase* ObjectProperty = CastField<const FObjectPropertyBase>(Property))
-						{
-							ObjectPropertyBinding->AllowedClass = ObjectProperty->PropertyClass;
-						}
-					}
-					if (TypeFilter == "struct")
-					{
-						FChooserStructPropertyBinding* StructPropertyBinding = static_cast<FChooserStructPropertyBinding*>(ContextProperty);
-
-						StructPropertyBinding->StructType = nullptr;
-
-						if (const FStructProperty* StructProperty = CastField<const FStructProperty>(Property))
-						{
-							StructPropertyBinding->StructType = StructProperty->Struct;
-						}
-						else if (InBindingChain.Num() == 1)
-						{
-							// direct binding to a context struct
-							if (ContextClassOwner)
-							{
-								TConstArrayView<FInstancedStruct> ContextData = ContextClassOwner->GetContextData();
-								if (ContextData.IsValidIndex(InBindingChain[0].ArrayIndex))
-								{
-									if (const FContextObjectTypeStruct* StructContext = ContextData[InBindingChain[0].ArrayIndex].GetPtr<FContextObjectTypeStruct>())
-									{
-										StructPropertyBinding->StructType = StructContext->Struct;
-									}
-								}
-							}
-						}
-					}
+					UObject* TransactionObject = Cast<UObject>(ContextClassOwner);
+                    
+                    const FScopedTransaction Transaction(NSLOCTEXT("ContextPropertyWidget", "Change Property Binding", "Change Property Binding"));
+                    TransactionObject->Modify(true);
 
-					ContextProperty->Compile(ContextClassOwner);
+					SetPropertyBinding(ContextClassOwner, InBindingChain, *PropertyBindingValue.Get());
 					OnValueChanged.ExecuteIfBound();
 				}
 			});
@@ -298,13 +257,18 @@ TSharedRef<SWidget> SPropertyAccessChainWidget::CreatePropertyAccessWidget()
 
 		return CurrentValue;	
 	});
-	
+
 	Args.CurrentBindingText = MakeAttributeLambda([this]()
 			{
 				const FText Bind = NSLOCTEXT("ContextPropertyWidget", "Bind", "Bind");
 				FText CurrentValue = Bind;
 		
 				const FChooserPropertyBinding* PropertyValue = PropertyBindingValue.Get();
+
+				if (PropertyValue == nullptr)
+				{
+					return FText();
+				}
 	
 				int BindingChainLength = PropertyValue->PropertyBindingChain.Num();
 				if (BindingChainLength == 0)
diff --git a/Engine/Plugins/Chooser/Source/ChooserEditor/Public/SPropertyAccessChainWidget.h b/Engine/Plugins/Chooser/Source/ChooserEditor/Public/SPropertyAccessChainWidget.h
--- a/Engine/Plugins/Chooser/Source/ChooserEditor/Public/SPropertyAccessChainWidget.h
+++ b/Engine/Plugins/Chooser/Source/ChooserEditor/Public/SPropertyAccessChainWidget.h
@@ -10,6 +10,7 @@
 namespace UE::ChooserEditor
 {
 DECLARE_DELEGATE(FPropertyAccessChainChanged)
+
 		
 // Wrapper widget for Property access widget, which can update when the target Class changes
 class CHOOSEREDITOR_API SPropertyAccessChainWidget : public SCompoundWidget
@@ -30,6 +31,9 @@ public:
 
 	void Construct( const FArguments& InArgs);
 	virtual ~SPropertyAccessChainWidget();
+	
+	// Assign result from Property Access Widget to an FChooserPropertyBinding (Helper for OnAddBinding implementations)
+	static void SetPropertyBinding(IHasContextClass* HasContext, const TArray<FBindingChainElement>& InBindingChain, FChooserPropertyBinding& OutPropertyBinding);
 
 private:
 	TSharedRef<SWidget> CreatePropertyAccessWidget();
