Fix mirroring in ChooserPlayer pose matching

diff --git a/Engine/Plugins/Chooser/Source/Chooser/Private/AnimNode_ChooserPlayer.cpp b/Engine/Plugins/Chooser/Source/Chooser/Private/AnimNode_ChooserPlayer.cpp
--- a/Engine/Plugins/Chooser/Source/Chooser/Private/AnimNode_ChooserPlayer.cpp
+++ b/Engine/Plugins/Chooser/Source/Chooser/Private/AnimNode_ChooserPlayer.cpp
@@ -55,6 +55,7 @@ UAnimationAsset* FAnimNode_ChooserPlayer::ChooseAsset(const FAnimationUpdateCont
 					{
 						bForceBlendTo = true;
 						Settings.StartTime = SearchResult.TimeOffsetSeconds;
+						Settings.bMirror = SearchResult.bMirrored;
 					}
 					return SelectedAnimationAsset;
 				}
