Chooser - FAnimNode_ChooserPlayer::DefaultSettings as BlueprintReadWrite

diff --git a/Engine/Plugins/Chooser/Source/Chooser/Internal/AnimNode_ChooserPlayer.h b/Engine/Plugins/Chooser/Source/Chooser/Internal/AnimNode_ChooserPlayer.h
--- a/Engine/Plugins/Chooser/Source/Chooser/Internal/AnimNode_ChooserPlayer.h
+++ b/Engine/Plugins/Chooser/Source/Chooser/Internal/AnimNode_ChooserPlayer.h
@@ -128,7 +128,7 @@ public:
 	bool bUpdateAllActiveBlendSpaces = true;
 	
 	// Settings when starting an animation - these can be overridden per animation asset by the chooser itself
-	UPROPERTY(EditAnywhere, Category = Settings)
+	UPROPERTY(EditAnywhere, BlueprintReadWrite, Category = Settings, meta = (PinHiddenByDefault))
 	FChooserPlayerSettings DefaultSettings;
 	
 	UPROPERTY()
