Improved Chooser support for Soft Object References

diff --git a/Engine/Plugins/Chooser/Source/Chooser/Internal/Chooser.h b/Engine/Plugins/Chooser/Source/Chooser/Internal/Chooser.h
--- a/Engine/Plugins/Chooser/Source/Chooser/Internal/Chooser.h
+++ b/Engine/Plugins/Chooser/Source/Chooser/Internal/Chooser.h
@@ -201,6 +201,7 @@ public:
 
 	bool IsCookedData() const { return !CookedResults.IsEmpty(); }
 	static FObjectChooserBase::EIteratorStatus EvaluateChooser(FChooserEvaluationContext& Context, const UChooserTable* Chooser, FObjectChooserBase::FObjectChooserIteratorCallback Callback);
+	static FObjectChooserBase::EIteratorStatus EvaluateChooser(FChooserEvaluationContext& Context, const UChooserTable* Chooser, FObjectChooserBase::FObjectChooserSoftObjectIteratorCallback Callback);
 	static FObjectChooserBase::EIteratorStatus IterateChooser(const UChooserTable* Chooser, FObjectChooserBase::FObjectChooserIteratorCallback Callback);
 };
 
@@ -209,8 +210,10 @@ struct CHOOSER_API FNestedChooser : public FObjectChooserBase
 {
 	GENERATED_BODY()
 
+	virtual void ChooseObject(FChooserEvaluationContext& Context, TSoftObjectPtr<UObject>& Result) const final override;
 	virtual UObject* ChooseObject(FChooserEvaluationContext& Context) const final override;
 	virtual EIteratorStatus ChooseMulti(FChooserEvaluationContext &Context, FObjectChooserIteratorCallback Callback) const final override;
+	virtual EIteratorStatus ChooseMulti(FChooserEvaluationContext &Context, FObjectChooserSoftObjectIteratorCallback Callback) const final override;
 	virtual EIteratorStatus IterateObjects(FObjectChooserIteratorCallback Callback) const final override;
 	virtual void GetDebugName(FString& OutDebugName) const override;
 	
@@ -226,9 +229,10 @@ USTRUCT(DisplayName = "Evaluate Chooser", Meta = (Category = "Chooser", Tooltip
 struct CHOOSER_API FEvaluateChooser : public FObjectChooserBase
 {
 	GENERATED_BODY()
-
+	virtual void ChooseObject(FChooserEvaluationContext& Context, TSoftObjectPtr<UObject>& Result) const final override;
 	virtual UObject* ChooseObject(FChooserEvaluationContext& Context) const final override;
 	virtual EIteratorStatus ChooseMulti(FChooserEvaluationContext &Context, FObjectChooserIteratorCallback Callback) const final override;
+	virtual EIteratorStatus ChooseMulti(FChooserEvaluationContext &Context, FObjectChooserSoftObjectIteratorCallback Callback) const final override;
 	virtual EIteratorStatus IterateObjects(FObjectChooserIteratorCallback Callback) const final override;
 	virtual void GetDebugName(FString& OutDebugName) const override;
 	
diff --git a/Engine/Plugins/Chooser/Source/Chooser/Internal/ObjectChooser_Asset.h b/Engine/Plugins/Chooser/Source/Chooser/Internal/ObjectChooser_Asset.h
--- a/Engine/Plugins/Chooser/Source/Chooser/Internal/ObjectChooser_Asset.h
+++ b/Engine/Plugins/Chooser/Source/Chooser/Internal/ObjectChooser_Asset.h
@@ -29,6 +29,7 @@ struct CHOOSER_API FSoftAssetChooser : public FObjectChooserBase
 	GENERATED_BODY()
 
 	// FObjectChooserBase interface
+	virtual void ChooseObject(FChooserEvaluationContext& Context, TSoftObjectPtr<UObject>& Result) const final override;
 	virtual UObject* ChooseObject(FChooserEvaluationContext& Context) const final override;
 	virtual EIteratorStatus IterateObjects(FObjectChooserIteratorCallback Callback) const final override;
 	
diff --git a/Engine/Plugins/Chooser/Source/Chooser/Private/Chooser.cpp b/Engine/Plugins/Chooser/Source/Chooser/Private/Chooser.cpp
--- a/Engine/Plugins/Chooser/Source/Chooser/Private/Chooser.cpp
+++ b/Engine/Plugins/Chooser/Source/Chooser/Private/Chooser.cpp
@@ -398,6 +398,15 @@ bool UChooserTable::ResultAssetFilter(const FAssetData& AssetData)
 }
 
 FObjectChooserBase::EIteratorStatus UChooserTable::EvaluateChooser(FChooserEvaluationContext& Context, const UChooserTable* Chooser, FObjectChooserBase::FObjectChooserIteratorCallback Callback)
+{
+	return EvaluateChooser(Context, Chooser, FObjectChooserBase::FObjectChooserSoftObjectIteratorCallback::CreateLambda([Callback](const TSoftObjectPtr<UObject>& ObjectPtr)
+	{
+		return Callback.Execute(ObjectPtr.Get());
+	}));
+}
+
+FObjectChooserBase::EIteratorStatus UChooserTable::EvaluateChooser(FChooserEvaluationContext& Context, const UChooserTable* Chooser,
+	FObjectChooserBase::FObjectChooserSoftObjectIteratorCallback Callback)
 {
 	TRACE_CPUPROFILER_EVENT_SCOPE(EvaluateChooser);
 
@@ -607,6 +616,15 @@ FObjectChooserBase::EIteratorStatus UChooserTable::IterateChooser(const UChooser
 	return FObjectChooserBase::EIteratorStatus::Continue;
 }
 
+void FEvaluateChooser::ChooseObject(FChooserEvaluationContext& Context, TSoftObjectPtr<UObject>& Result) const
+{
+    UChooserTable::EvaluateChooser(Context, Chooser, FObjectChooserSoftObjectIteratorCallback::CreateLambda([&Result](const TSoftObjectPtr<UObject>& InResult)
+    {
+    	Result = InResult;
+    	return FObjectChooserBase::EIteratorStatus::Stop;
+    }));
+}
+
 UObject* FEvaluateChooser::ChooseObject(FChooserEvaluationContext& Context) const
 {
 	UObject* Result = nullptr;
@@ -624,6 +642,12 @@ FObjectChooserBase::EIteratorStatus FEvaluateChooser::ChooseMulti(FChooserEvalua
 	return UChooserTable::EvaluateChooser(Context, Chooser, Callback);
 }
 
+FObjectChooserBase::EIteratorStatus FEvaluateChooser::ChooseMulti(FChooserEvaluationContext& Context,
+	FObjectChooserSoftObjectIteratorCallback Callback) const
+{
+	return UChooserTable::EvaluateChooser(Context, Chooser, Callback);
+}
+
 FObjectChooserBase::EIteratorStatus FEvaluateChooser::IterateObjects(FObjectChooserIteratorCallback Callback) const
 {
 	return UChooserTable::IterateChooser(Chooser, Callback);
@@ -642,6 +666,15 @@ FNestedChooser::FNestedChooser()
 {
 }
 
+void FNestedChooser::ChooseObject(FChooserEvaluationContext& Context, TSoftObjectPtr<UObject>& Result) const
+{
+    UChooserTable::EvaluateChooser(Context, Chooser, FObjectChooserSoftObjectIteratorCallback::CreateLambda([&Result](const TSoftObjectPtr<UObject>& InResult)
+    {
+    	Result = InResult;
+    	return FObjectChooserBase::EIteratorStatus::Stop;
+    }));
+}
+
 UObject* FNestedChooser::ChooseObject(FChooserEvaluationContext& Context) const
 {
 	UObject* Result = nullptr;
@@ -659,6 +692,11 @@ FObjectChooserBase::EIteratorStatus FNestedChooser::ChooseMulti(FChooserEvaluati
 	return UChooserTable::EvaluateChooser(Context, Chooser, Callback);
 }
 
+FObjectChooserBase::EIteratorStatus FNestedChooser::ChooseMulti(FChooserEvaluationContext& Context, FObjectChooserSoftObjectIteratorCallback Callback) const
+{
+	return UChooserTable::EvaluateChooser(Context, Chooser, Callback);
+}
+
 FObjectChooserBase::EIteratorStatus FNestedChooser::IterateObjects(FObjectChooserIteratorCallback Callback) const
 {
 	return UChooserTable::IterateChooser(Chooser, Callback);
@@ -688,6 +726,28 @@ static FAutoConsoleCommand CCmdTestCookChoosers(
 
 void UChooserTable::AddNestedChooser(UChooserTable* Chooser)
 {
+	FName NewChooserName = Chooser->GetFName();
+	FString NewChooserPlainName = NewChooserName.GetPlainNameString();
+	bool bNameConflict = false;
+	int MaxNumber = NewChooserName.GetNumber();
+	for(UChooserTable* NestedChooserIterator : NestedChoosers)
+	{
+		FName NestedChooserName = NestedChooserIterator->GetFName();
+		if (NestedChooserName.GetPlainNameString() == NewChooserPlainName)
+		{
+			MaxNumber = FMath::Max(MaxNumber, NestedChooserName.GetNumber());
+			if(NestedChooserName.GetNumber() == NewChooserName.GetNumber())
+			{
+				bNameConflict = true;
+			}
+		}
+	}
+	if (bNameConflict)
+	{
+		NewChooserName.SetNumber(MaxNumber + 1);
+		Chooser->Rename(*NewChooserName.ToString());
+	}
+	
 	NestedChoosers.Add(Chooser);
 	NestedChoosersChanged.Broadcast();
 }
diff --git a/Engine/Plugins/Chooser/Source/Chooser/Private/ChooserFunctionLibrary.cpp b/Engine/Plugins/Chooser/Source/Chooser/Private/ChooserFunctionLibrary.cpp
--- a/Engine/Plugins/Chooser/Source/Chooser/Private/ChooserFunctionLibrary.cpp
+++ b/Engine/Plugins/Chooser/Source/Chooser/Private/ChooserFunctionLibrary.cpp
@@ -114,7 +114,7 @@ DEFINE_FUNCTION(UChooserFunctionLibrary::execGetChooserStructOutput)
 }
 
 UObject* UChooserFunctionLibrary::EvaluateObjectChooserBase(FChooserEvaluationContext& Context, const FInstancedStruct& ObjectChooser, TSubclassOf<UObject> ObjectClass, bool bResultIsClass)
- {
+{
  	UObject* Result = nullptr;
 
 	if (const FObjectChooserBase* ObjectChooserPtr = ObjectChooser.GetPtr<FObjectChooserBase>())
@@ -151,6 +151,19 @@ UObject* UChooserFunctionLibrary::EvaluateObjectChooserBase(FChooserEvaluationCo
  	return Result;
  }
 
+TSoftObjectPtr<UObject> UChooserFunctionLibrary::EvaluateObjectChooserBaseSoft(FChooserEvaluationContext& Context, const FInstancedStruct& ObjectChooser, TSubclassOf<UObject> ObjectClass, bool bResultIsClass)
+{
+	TSoftObjectPtr<UObject> Result;
+
+	if (const FObjectChooserBase* ObjectChooserPtr = ObjectChooser.GetPtr<FObjectChooserBase>())
+	{
+		ObjectChooserPtr->ChooseObject(Context, Result);
+	}
+	
+ 	return Result;
+}
+
+
 
 TArray<UObject*> UChooserFunctionLibrary::EvaluateObjectChooserBaseMulti(FChooserEvaluationContext& Context, const FInstancedStruct& ObjectChooser, TSubclassOf<UObject> ObjectClass, bool bResultIsClass)
 {
@@ -185,7 +198,25 @@ TArray<UObject*> UChooserFunctionLibrary::EvaluateObjectChooserBaseMulti(FChoose
 	}
 	
  	return Result;
- }
+}
+
+TArray<TSoftObjectPtr<UObject>> UChooserFunctionLibrary::EvaluateObjectChooserBaseMultiSoft(FChooserEvaluationContext& Context, const FInstancedStruct& ObjectChooser, TSubclassOf<UObject> ObjectClass, bool bResultIsClass)
+{
+ 	TArray<TSoftObjectPtr<UObject>> Result;
+
+	if (const FObjectChooserBase* ObjectChooserPtr = ObjectChooser.GetPtr<FObjectChooserBase>())
+	{
+		ObjectChooserPtr->ChooseMulti(Context, FObjectChooserBase::FObjectChooserSoftObjectIteratorCallback::CreateLambda([&Result, bResultIsClass, ObjectClass](const TSoftObjectPtr<UObject>& InResult)
+		{
+			Result.Add(InResult);
+			
+			// trigger output columns only on the first result
+			return Result.Num() == 1 ? FObjectChooserBase::EIteratorStatus::ContinueWithOutputs : FObjectChooserBase::EIteratorStatus::Continue;
+		}));
+	}
+	
+ 	return Result;
+}
 
 FInstancedStruct UChooserFunctionLibrary::MakeEvaluateChooser(UChooserTable* Chooser)
 {
diff --git a/Engine/Plugins/Chooser/Source/Chooser/Private/ObjectChooser_Asset.cpp b/Engine/Plugins/Chooser/Source/Chooser/Private/ObjectChooser_Asset.cpp
--- a/Engine/Plugins/Chooser/Source/Chooser/Private/ObjectChooser_Asset.cpp
+++ b/Engine/Plugins/Chooser/Source/Chooser/Private/ObjectChooser_Asset.cpp
@@ -12,9 +12,14 @@ FObjectChooserBase::EIteratorStatus FAssetChooser::IterateObjects(FObjectChooser
 	return Callback.Execute(Asset);
 }
 
+void FSoftAssetChooser::ChooseObject(FChooserEvaluationContext& Context, TSoftObjectPtr<UObject>& Result) const
+{
+	Result = Asset.ToSoftObjectPath();
+}
+
 UObject* FSoftAssetChooser::ChooseObject(FChooserEvaluationContext& Context) const
 {
-	return Asset.LoadSynchronous();
+	return Asset.Get();
 }
 
 FObjectChooserBase::EIteratorStatus FSoftAssetChooser::IterateObjects(FObjectChooserIteratorCallback Callback) const
diff --git a/Engine/Plugins/Chooser/Source/Chooser/Public/ChooserFunctionLibrary.h b/Engine/Plugins/Chooser/Source/Chooser/Public/ChooserFunctionLibrary.h
--- a/Engine/Plugins/Chooser/Source/Chooser/Public/ChooserFunctionLibrary.h
+++ b/Engine/Plugins/Chooser/Source/Chooser/Public/ChooserFunctionLibrary.h
@@ -6,6 +6,7 @@
 #include "Kismet/BlueprintFunctionLibrary.h"
 #include "Chooser.h"
 #include "Templates/SubclassOf.h"
+#include "UObject/SoftObjectPtr.h"
 #include "ChooserFunctionLibrary.generated.h"
 
 /**
@@ -46,7 +47,11 @@ public:
 	* @param bResultIsClass		(in) The Object being returned is a UClass, and the ObjectClass parameter indicates what it must be a subclass of
 	*/
 	UFUNCTION(BlueprintCallable, meta = (BlueprintThreadSafe, BlueprintInternalUseOnly = "true", DeterminesOutputType = "ObjectClass"))
-	static UObject* EvaluateObjectChooserBase(UPARAM(Ref) FChooserEvaluationContext& Context,UPARAM(Ref) const FInstancedStruct& ObjectChooser, TSubclassOf<UObject> ObjectClass, bool bResultIsClass = false);
+	static UObject* EvaluateObjectChooserBase(UPARAM(Ref) struct FChooserEvaluationContext& Context,UPARAM(Ref) const struct FInstancedStruct& ObjectChooser, TSubclassOf<UObject> ObjectClass, bool bResultIsClass = false);
+	
+	
+	UFUNCTION(BlueprintCallable, meta = (BlueprintThreadSafe, BlueprintInternalUseOnly = "true", DeterminesOutputType = "ObjectClass"))
+	static TSoftObjectPtr<UObject> EvaluateObjectChooserBaseSoft(UPARAM(Ref) struct FChooserEvaluationContext& Context,UPARAM(Ref) const struct FInstancedStruct& ObjectChooser, TSubclassOf<UObject> ObjectClass, bool bResultIsClass = false);
 
 	/**
 	* Evaluate a chooser table and return all selected UObjects
@@ -58,7 +63,9 @@ public:
 	*/
 	UFUNCTION(BlueprintCallable, meta = (BlueprintThreadSafe, BlueprintInternalUseOnly = "true", DeterminesOutputType = "ObjectClass"))
 	static TArray<UObject*> EvaluateObjectChooserBaseMulti(UPARAM(Ref) FChooserEvaluationContext& Context,UPARAM(Ref) const FInstancedStruct& ObjectChooser, TSubclassOf<UObject> ObjectClass, bool bResultIsClass = false);
-
+	
+	UFUNCTION(BlueprintCallable, meta = (BlueprintThreadSafe, BlueprintInternalUseOnly = "true", DeterminesOutputType = "ObjectClass"))
+	static TArray<TSoftObjectPtr<UObject>> EvaluateObjectChooserBaseMultiSoft(UPARAM(Ref) FChooserEvaluationContext& Context,UPARAM(Ref) const FInstancedStruct& ObjectChooser, TSubclassOf<UObject> ObjectClass, bool bResultIsClass = false);
 
 	/**
 	* Add an Object to a ChooserEvaluation context
diff --git a/Engine/Plugins/Chooser/Source/Chooser/Public/IObjectChooser.h b/Engine/Plugins/Chooser/Source/Chooser/Public/IObjectChooser.h
--- a/Engine/Plugins/Chooser/Source/Chooser/Public/IObjectChooser.h
+++ b/Engine/Plugins/Chooser/Source/Chooser/Public/IObjectChooser.h
@@ -101,10 +101,17 @@ public:
 	enum class EIteratorStatus { Continue, ContinueWithOutputs, Stop };
 
 	DECLARE_DELEGATE_RetVal_OneParam( EIteratorStatus, FObjectChooserIteratorCallback, UObject*);
+	DECLARE_DELEGATE_RetVal_OneParam( EIteratorStatus, FObjectChooserSoftObjectIteratorCallback, const TSoftObjectPtr<UObject>&);
 
 	virtual void Compile(class IHasContextClass* HasContext, bool bForce) {};
 	virtual bool HasCompileErrors(FText& Message) { return false; }
 
+	virtual void ChooseObject(FChooserEvaluationContext& Context, TSoftObjectPtr<UObject>& Result) const
+	{
+		// fallback implementation
+		Result = TSoftObjectPtr<UObject>(ChooseObject(Context));
+	}
+	
 	virtual UObject* ChooseObject(FChooserEvaluationContext& Context) const { return nullptr; };
 	virtual EIteratorStatus ChooseMulti(FChooserEvaluationContext& ContextData, FObjectChooserIteratorCallback Callback) const
 	{
@@ -115,6 +122,20 @@ public:
 		}
 		return EIteratorStatus::Continue;
 	}
+	
+	virtual EIteratorStatus ChooseMulti(FChooserEvaluationContext& ContextData, FObjectChooserSoftObjectIteratorCallback Callback) const
+	{
+		// fallback implementation just calls the single version
+		TSoftObjectPtr<UObject> Result;
+		ChooseObject(ContextData, Result);
+		if (!Result.IsNull())
+		{
+			return Callback.Execute(Result);
+		}
+	
+		return EIteratorStatus::Continue;
+	}
+
 
 	virtual EIteratorStatus IterateObjects(FObjectChooserIteratorCallback Callback) const { return EIteratorStatus::Continue; }
 	virtual void GetDebugName(FString& OutDebugName) const {};
diff --git a/Engine/Plugins/Chooser/Source/ChooserUncooked/Private/EvaluateChooserNode.cpp b/Engine/Plugins/Chooser/Source/ChooserUncooked/Private/EvaluateChooserNode.cpp
--- a/Engine/Plugins/Chooser/Source/ChooserUncooked/Private/EvaluateChooserNode.cpp
+++ b/Engine/Plugins/Chooser/Source/ChooserUncooked/Private/EvaluateChooserNode.cpp
@@ -507,7 +507,9 @@ void UK2Node_EvaluateChooser2::AllocateDefaultPins()
 
 			if (ResultPin)
 			{
-				ResultPin->PinType.PinCategory = bResultIsClass ? UEdGraphSchema_K2::PC_Class : UEdGraphSchema_K2::PC_Object;
+				ResultPin->PinType.PinCategory = bReturnSoftObjectReference ? UEdGraphSchema_K2::PC_SoftObject
+												: bResultIsClass ? UEdGraphSchema_K2::PC_Class : UEdGraphSchema_K2::PC_Object;
+
 				ResultPin->PinType.PinSubCategoryObject = ResultType;
 				ResultPin->PinType.ContainerType = (Mode == EEvaluateChooserMode::AllResults) ? EPinContainerType::Array : EPinContainerType::None;
 			}
@@ -515,7 +517,8 @@ void UK2Node_EvaluateChooser2::AllocateDefaultPins()
 			{
 				UEdGraphNode::FCreatePinParams PinParams;
 				PinParams.ContainerType = (Mode == EEvaluateChooserMode::AllResults) ? EPinContainerType::Array : EPinContainerType::None;
-				PinParams.ValueTerminalType.TerminalCategory =  bResultIsClass ? UEdGraphSchema_K2::PC_Class : UEdGraphSchema_K2::PC_Object;
+				PinParams.ValueTerminalType.TerminalCategory = bReturnSoftObjectReference ? UEdGraphSchema_K2::PC_SoftObject
+												: bResultIsClass ? UEdGraphSchema_K2::PC_Class : UEdGraphSchema_K2::PC_Object;
 				
 				CreatePin(EGPD_Output, PinParams.ValueTerminalType.TerminalCategory, ResultType, TEXT("Result"), PinParams);
 			}
@@ -550,7 +553,7 @@ void UK2Node_EvaluateChooser2::PostEditChangeProperty(struct FPropertyChangedEve
 	{
 		ChooserChanged();
 	}
-	if (PropertyChangedEvent.Property->GetName() == "Mode")
+	if (PropertyChangedEvent.Property->GetName() == "Mode" || PropertyChangedEvent.Property->GetName() == "bReturnSoftObjectReference")
 	{
 		ReconstructNode();
 	}
@@ -613,13 +616,27 @@ void UK2Node_EvaluateChooser2::ExpandNode(class FKismetCompilerContext& Compiler
 		UK2Node_CallFunction* CallFunction = CompilerContext.SpawnIntermediateNode<UK2Node_CallFunction>(this, SourceGraph);
 		CompilerContext.MessageLog.NotifyIntermediateObjectCreation(CallFunction, this);
 
-		if (Mode == EEvaluateChooserMode::AllResults)
+		if (bReturnSoftObjectReference)
 		{
-			CallFunction->SetFromFunction(UChooserFunctionLibrary::StaticClass()->FindFunctionByName(GET_MEMBER_NAME_CHECKED(UChooserFunctionLibrary, EvaluateObjectChooserBaseMulti)));
+			if (Mode == EEvaluateChooserMode::AllResults)
+			{
+				CallFunction->SetFromFunction(UChooserFunctionLibrary::StaticClass()->FindFunctionByName(GET_MEMBER_NAME_CHECKED(UChooserFunctionLibrary, EvaluateObjectChooserBaseMultiSoft)));
+			}
+			else
+			{
+				CallFunction->SetFromFunction(UChooserFunctionLibrary::StaticClass()->FindFunctionByName(GET_MEMBER_NAME_CHECKED(UChooserFunctionLibrary, EvaluateObjectChooserBaseSoft)));
+			}
 		}
 		else
 		{
-			CallFunction->SetFromFunction(UChooserFunctionLibrary::StaticClass()->FindFunctionByName(GET_MEMBER_NAME_CHECKED(UChooserFunctionLibrary, EvaluateObjectChooserBase)));
+			if (Mode == EEvaluateChooserMode::AllResults)
+			{
+				CallFunction->SetFromFunction(UChooserFunctionLibrary::StaticClass()->FindFunctionByName(GET_MEMBER_NAME_CHECKED(UChooserFunctionLibrary, EvaluateObjectChooserBaseMulti)));
+			}
+			else
+			{
+				CallFunction->SetFromFunction(UChooserFunctionLibrary::StaticClass()->FindFunctionByName(GET_MEMBER_NAME_CHECKED(UChooserFunctionLibrary, EvaluateObjectChooserBase)));
+			}
 		}
 		CallFunction->AllocateDefaultPins();
 		
diff --git a/Engine/Plugins/Chooser/Source/ChooserUncooked/Private/EvaluateChooserNode.h b/Engine/Plugins/Chooser/Source/ChooserUncooked/Private/EvaluateChooserNode.h
--- a/Engine/Plugins/Chooser/Source/ChooserUncooked/Private/EvaluateChooserNode.h
+++ b/Engine/Plugins/Chooser/Source/ChooserUncooked/Private/EvaluateChooserNode.h
@@ -133,6 +133,9 @@ private:
 	UPROPERTY(EditAnywhere, Category = "Chooser", Meta = (EditConditionHides, EditCondition="bShowModeSetting"))
 	EEvaluateChooserMode Mode;
 
+	UPROPERTY(EditAnywhere, Category = "Chooser", Meta = (EditConditionHides, EditCondition="bShowModeSetting"))
+	bool bReturnSoftObjectReference = false;
+
 	UPROPERTY(Transient)
 	bool bShowModeSetting = false;
 
