Performance improvments for chooser debugging

diff --git a/Engine/Plugins/Chooser/Source/Chooser/Internal/Chooser.h b/Engine/Plugins/Chooser/Source/Chooser/Internal/Chooser.h
--- a/Engine/Plugins/Chooser/Source/Chooser/Internal/Chooser.h
+++ b/Engine/Plugins/Chooser/Source/Chooser/Internal/Chooser.h
@@ -64,11 +64,7 @@ public:
 	const FString& GetDebugTargetName() const { return DebugTargetName; }
 	
 
-	void AddRecentContextObject(const FString& ObjectName) const
-	{
-		RecentContextObjects.AddUnique(ObjectName);
-		RecentContextObjects.StableSort();
-	}
+	void AddRecentContextObject(const FString& ObjectName) const;
 	void IterateRecentContextObjects(TFunction<void(const FString&)> Callback) const;
 	void UpdateDebugging(FChooserEvaluationContext& Context) const;
 	
diff --git a/Engine/Plugins/Chooser/Source/Chooser/Private/Chooser.cpp b/Engine/Plugins/Chooser/Source/Chooser/Private/Chooser.cpp
--- a/Engine/Plugins/Chooser/Source/Chooser/Private/Chooser.cpp
+++ b/Engine/Plugins/Chooser/Source/Chooser/Private/Chooser.cpp
@@ -333,9 +333,16 @@ void UChooserTable::PostEditChangeProperty(FPropertyChangedEvent& PropertyChange
 	}
 }
 
+void UChooserTable::AddRecentContextObject(const FString& ObjectName) const
+{
+	FScopeLock Lock(&DebugLock);
+	RecentContextObjects.AddUnique(ObjectName);
+}
+
 void UChooserTable::IterateRecentContextObjects(TFunction<void(const FString&)> Callback) const
 {
 	FScopeLock Lock(&DebugLock);
+	RecentContextObjects.StableSort();
 	for (const FString& ObjectName : RecentContextObjects)
 	{
 		Callback(ObjectName);
@@ -344,8 +351,6 @@ void UChooserTable::IterateRecentContextObjects(TFunction<void(const FString&)>
 
 void UChooserTable::UpdateDebugging(FChooserEvaluationContext& Context) const
 {
-	FScopeLock Lock(&DebugLock);
-
 	const UChooserTable* RootTable = GetRootChooser();
 
 	for (const FStructView& Param : Context.Params)
