Passthrough inputs correctly when properties aren't set on certain chooser column types

diff --git a/Engine/Plugins/Chooser/Source/Chooser/Private/BoolColumn.cpp b/Engine/Plugins/Chooser/Source/Chooser/Private/BoolColumn.cpp
--- a/Engine/Plugins/Chooser/Source/Chooser/Private/BoolColumn.cpp
+++ b/Engine/Plugins/Chooser/Source/Chooser/Private/BoolColumn.cpp
@@ -26,11 +26,10 @@ FBoolColumn::FBoolColumn()
 
 void FBoolColumn::Filter(FChooserEvaluationContext& Context, const FChooserIndexArray& IndexListIn, FChooserIndexArray& IndexListOut) const
 {
-	if (InputValue.IsValid())
+	bool Result = false;
+	if (InputValue.IsValid() && 
+		InputValue.Get<FChooserParameterBoolBase>().GetValue(Context, Result))
 	{
-		bool Result = false;
-		InputValue.Get<FChooserParameterBoolBase>().GetValue(Context,Result);
-
 		TRACE_CHOOSER_VALUE(Context, ToCStr(InputValue.Get<FChooserParameterBase>().GetDebugName()), Result);
 
 	#if WITH_EDITOR
diff --git a/Engine/Plugins/Chooser/Source/Chooser/Private/FloatDistanceColumn.cpp b/Engine/Plugins/Chooser/Source/Chooser/Private/FloatDistanceColumn.cpp
--- a/Engine/Plugins/Chooser/Source/Chooser/Private/FloatDistanceColumn.cpp
+++ b/Engine/Plugins/Chooser/Source/Chooser/Private/FloatDistanceColumn.cpp
@@ -92,11 +92,10 @@ float FFloatDistanceColumn::EditorTestCost(int32 RowIndex) const
 
 void FFloatDistanceColumn::Filter(FChooserEvaluationContext& Context, const FChooserIndexArray& IndexListIn, FChooserIndexArray& IndexListOut) const
 {
-	if (InputValue.IsValid())
-	{
-		double Result = 0.0f;
-		InputValue.Get<FChooserParameterFloatBase>().GetValue(Context, Result);
-        	
+	double Result = 0.0f;
+	if (InputValue.IsValid() && 
+		InputValue.Get<FChooserParameterFloatBase>().GetValue(Context, Result))
+	{	
 		TRACE_CHOOSER_VALUE(Context, ToCStr(InputValue.Get<FChooserParameterBase>().GetDebugName()), Result);
 
 #if WITH_EDITOR
diff --git a/Engine/Plugins/Chooser/Source/Chooser/Private/FloatRangeColumn.cpp b/Engine/Plugins/Chooser/Source/Chooser/Private/FloatRangeColumn.cpp
--- a/Engine/Plugins/Chooser/Source/Chooser/Private/FloatRangeColumn.cpp
+++ b/Engine/Plugins/Chooser/Source/Chooser/Private/FloatRangeColumn.cpp
@@ -52,11 +52,10 @@ bool FFloatRangeColumn::EditorTestFilter(int32 RowIndex) const
 
 void FFloatRangeColumn::Filter(FChooserEvaluationContext& Context, const FChooserIndexArray& IndexListIn, FChooserIndexArray& IndexListOut) const
 {
-	if (InputValue.IsValid())
-	{
-		double Result = 0.0f;
-		InputValue.Get<FChooserParameterFloatBase>().GetValue(Context, Result);
-		
+	double Result = 0.0f;
+	if (InputValue.IsValid() &&
+		InputValue.Get<FChooserParameterFloatBase>().GetValue(Context, Result))
+	{		
 		TRACE_CHOOSER_VALUE(Context, ToCStr(InputValue.Get<FChooserParameterBase>().GetDebugName()), Result);
 
 #if WITH_EDITOR
