Fix ChooserIndexArray equals operator

diff --git a/Engine/Plugins/Chooser/Source/Chooser/Public/ChooserIndexArray.h b/Engine/Plugins/Chooser/Source/Chooser/Public/ChooserIndexArray.h
--- a/Engine/Plugins/Chooser/Source/Chooser/Public/ChooserIndexArray.h
+++ b/Engine/Plugins/Chooser/Source/Chooser/Public/ChooserIndexArray.h
@@ -71,7 +71,7 @@ public:
 	{
 		check(MaxSize >= Other.Size);
 		Size = Other.Size;
-		FMemory::Memcpy(Data,Other.Data, Size*sizeof(uint32));
+		FMemory::Memcpy(Data,Other.Data, Size*sizeof(FIndexData));
 	}
 
 	FIndexData* GetData() { return Data; }
