Autopopulate choosers on save

diff --git a/Engine/Plugins/Chooser/Source/ChooserEditor/Private/ChooserTableEditor.cpp b/Engine/Plugins/Chooser/Source/ChooserEditor/Private/ChooserTableEditor.cpp
--- a/Engine/Plugins/Chooser/Source/ChooserEditor/Private/ChooserTableEditor.cpp
+++ b/Engine/Plugins/Chooser/Source/ChooserEditor/Private/ChooserTableEditor.cpp
@@ -486,6 +486,12 @@ void FChooserTableEditor::InitToolMenuContext(FToolMenuContext& MenuContext)
 	MenuContext.AddObject(Context);
 }
 
+void FChooserTableEditor::SaveAsset_Execute()
+{
+	AutoPopulateAll();
+	FAssetEditorToolkit::SaveAsset_Execute();
+}
+
 
 void FChooserTableEditor::BindCommands()
 {
diff --git a/Engine/Plugins/Chooser/Source/ChooserEditor/Private/ChooserTableEditor.h b/Engine/Plugins/Chooser/Source/ChooserEditor/Private/ChooserTableEditor.h
--- a/Engine/Plugins/Chooser/Source/ChooserEditor/Private/ChooserTableEditor.h
+++ b/Engine/Plugins/Chooser/Source/ChooserEditor/Private/ChooserTableEditor.h
@@ -75,6 +75,7 @@ namespace UE::ChooserEditor
 		virtual bool IsPrimaryEditor() const override { return true; }
 		virtual bool IsSimpleAssetEditor() const override { return false; }
 		virtual void InitToolMenuContext(FToolMenuContext& MenuContext) override;
+		virtual void SaveAsset_Execute() override;
 
 		/** FEditorUndoClient Interface */
 		virtual bool MatchesContext( const FTransactionContext& InContext, const TArray<TPair<UObject*, FTransactionObjectEvent>>& TransactionObjectContexts ) const override;
