Changed the Iteration method for proxy tables and choosers to handle reading an output from a Context
Added a general template for non arithemtic types for the ChooserPropertyBinding Set Value
Used it for FGameplayTagQuery
Implemented an output column for gameplay tag queries

diff --git a/Engine/Plugins/Chooser/Source/Chooser/Internal/Chooser.h b/Engine/Plugins/Chooser/Source/Chooser/Internal/Chooser.h
--- a/Engine/Plugins/Chooser/Source/Chooser/Internal/Chooser.h
+++ b/Engine/Plugins/Chooser/Source/Chooser/Internal/Chooser.h
@@ -198,6 +198,7 @@ public:
 	bool IsCookedData() const { return !CookedResults.IsEmpty(); }
 	static CHOOSER_API FObjectChooserBase::EIteratorStatus EvaluateChooser(FChooserEvaluationContext& Context, const UChooserTable* Chooser, FObjectChooserBase::FObjectChooserIteratorCallback Callback);
 	static CHOOSER_API FObjectChooserBase::EIteratorStatus EvaluateChooser(FChooserEvaluationContext& Context, const UChooserTable* Chooser, FObjectChooserBase::FObjectChooserSoftObjectIteratorCallback Callback);
+	static CHOOSER_API FObjectChooserBase::EIteratorStatus IterateChooser(FChooserEvaluationContext& Context, const UChooserTable* Chooser, FObjectChooserBase::FObjectChooserIteratorCallback Callback);
 	static CHOOSER_API FObjectChooserBase::EIteratorStatus IterateChooser(const UChooserTable* Chooser, FObjectChooserBase::FObjectChooserIteratorCallback Callback);
 };
 
@@ -210,7 +211,7 @@ struct FNestedChooser : public FObjectChooserBase
 	CHOOSER_API virtual UObject* ChooseObject(FChooserEvaluationContext& Context) const final override;
 	CHOOSER_API virtual EIteratorStatus ChooseMulti(FChooserEvaluationContext &Context, FObjectChooserIteratorCallback Callback) const final override;
 	CHOOSER_API virtual EIteratorStatus ChooseMulti(FChooserEvaluationContext &Context, FObjectChooserSoftObjectIteratorCallback Callback) const final override;
-	CHOOSER_API virtual EIteratorStatus IterateObjects(FObjectChooserIteratorCallback Callback) const final override;
+	CHOOSER_API virtual EIteratorStatus IterateObjects(FChooserEvaluationContext& Context, FObjectChooserIteratorCallback Callback) const final override;
 	CHOOSER_API virtual void GetDebugName(FString& OutDebugName) const override;
 	
 	public:
@@ -229,7 +230,7 @@ struct FEvaluateChooser : public FObjectChooserBase
 	CHOOSER_API virtual UObject* ChooseObject(FChooserEvaluationContext& Context) const final override;
 	CHOOSER_API virtual EIteratorStatus ChooseMulti(FChooserEvaluationContext &Context, FObjectChooserIteratorCallback Callback) const final override;
 	CHOOSER_API virtual EIteratorStatus ChooseMulti(FChooserEvaluationContext &Context, FObjectChooserSoftObjectIteratorCallback Callback) const final override;
-	CHOOSER_API virtual EIteratorStatus IterateObjects(FObjectChooserIteratorCallback Callback) const final override;
+	CHOOSER_API virtual EIteratorStatus IterateObjects(FChooserEvaluationContext& Context, FObjectChooserIteratorCallback Callback) const final override;
 	CHOOSER_API virtual void GetDebugName(FString& OutDebugName) const override;
 	
 	public:
diff --git a/Engine/Plugins/Chooser/Source/Chooser/Internal/ObjectChooser_Asset.h b/Engine/Plugins/Chooser/Source/Chooser/Internal/ObjectChooser_Asset.h
--- a/Engine/Plugins/Chooser/Source/Chooser/Internal/ObjectChooser_Asset.h
+++ b/Engine/Plugins/Chooser/Source/Chooser/Internal/ObjectChooser_Asset.h
@@ -15,7 +15,7 @@ struct FAssetChooser : public FObjectChooserBase
 	
 	// FObjectChooserBase interface
 	UE_API virtual UObject* ChooseObject(FChooserEvaluationContext& Context) const final override;
-	UE_API virtual EIteratorStatus IterateObjects(FObjectChooserIteratorCallback Callback) const final override;
+	UE_API virtual EIteratorStatus IterateObjects(FChooserEvaluationContext& Context, FObjectChooserIteratorCallback Callback) const final override;
 
 #if WITH_EDITOR
 	virtual UObject* GetReferencedObject() const override { return Asset; }
@@ -33,7 +33,7 @@ struct FSoftAssetChooser : public FObjectChooserBase
 	// FObjectChooserBase interface
 	UE_API virtual void ChooseObject(FChooserEvaluationContext& Context, TSoftObjectPtr<UObject>& Result) const final override;
 	UE_API virtual UObject* ChooseObject(FChooserEvaluationContext& Context) const final override;
-	UE_API virtual EIteratorStatus IterateObjects(FObjectChooserIteratorCallback Callback) const final override;
+	UE_API virtual EIteratorStatus IterateObjects(FChooserEvaluationContext& Context, FObjectChooserIteratorCallback Callback) const final override;
 	
 #if WITH_EDITOR
 	virtual UObject* GetReferencedObject() const override { return Asset.LoadSynchronous(); }
diff --git a/Engine/Plugins/Chooser/Source/Chooser/Internal/ObjectChooser_Class.h b/Engine/Plugins/Chooser/Source/Chooser/Internal/ObjectChooser_Class.h
--- a/Engine/Plugins/Chooser/Source/Chooser/Internal/ObjectChooser_Class.h
+++ b/Engine/Plugins/Chooser/Source/Chooser/Internal/ObjectChooser_Class.h
@@ -15,7 +15,7 @@ struct FClassChooser : public FObjectChooserBase
 	
 	// FObjectChooserBase interface
 	UE_API virtual UObject* ChooseObject(FChooserEvaluationContext& Context) const final override;
-	UE_API virtual EIteratorStatus IterateObjects(FObjectChooserIteratorCallback Callback) const final override;
+	UE_API virtual EIteratorStatus IterateObjects(FChooserEvaluationContext& Context, FObjectChooserIteratorCallback Callback) const final override;
 public: 
 	UPROPERTY(EditAnywhere, Category = "Parameters")
 	TObjectPtr<UClass> Class;
diff --git a/Engine/Plugins/Chooser/Source/Chooser/Internal/OutputGameplayTagQueryColumn.h b/Engine/Plugins/Chooser/Source/Chooser/Internal/OutputGameplayTagQueryColumn.h
new file mode 100644
--- /dev/null
+++ b/Engine/Plugins/Chooser/Source/Chooser/Internal/OutputGameplayTagQueryColumn.h
@@ -0,0 +1,74 @@
+// Copyright Epic Games, Inc. All Rights Reserved.
+#pragma once
+
+#include "IChooserColumn.h"
+#include "IChooserParameterGameplayTag.h"
+#include "ChooserPropertyAccess.h"
+#include "GameplayTagContainer.h"
+#include "StructUtils/InstancedStruct.h"
+#include "Serialization/MemoryReader.h"
+#include "OutputGameplayTagQueryColumn.generated.h"
+
+#define UE_API CHOOSER_API
+
+USTRUCT(DisplayName = "Gameplay Tags Query Property Binding")
+struct FGameplayTagQueryContextProperty :  public FChooserParameterGameplayTagQueryBase
+{
+	GENERATED_BODY()
+
+public:
+	UPROPERTY(EditAnywhere, NoClear, Meta = (BindingType = "FGameplayTagQuery", BindingColor = "StructPinTypeColor"), Category = "Binding")
+	FChooserPropertyBinding Binding;
+	
+	UE_API virtual bool SetValue(FChooserEvaluationContext& Context, const FGameplayTagQuery& Value) const override;
+	
+	CHOOSER_PARAMETER_BOILERPLATE();
+};
+
+USTRUCT(DisplayName = "Output Gameplay Tag Query", Meta = (Category = "Output", Tooltip = "A column which writes a Gameplay Tag Query value."))
+struct FOutputGameplayTagQueryColumn : public FChooserColumnBase
+{
+	GENERATED_BODY()
+	public:
+	UE_API FOutputGameplayTagQueryColumn();
+
+	virtual bool HasFilters() const override { return false; }
+	virtual bool HasOutputs() const override { return true; }
+	UE_API virtual void SetOutputs(FChooserEvaluationContext& Context, int RowIndex) const override;
+	
+	UPROPERTY(EditAnywhere,  NoClear, Meta = (ExcludeBaseStruct, BaseStruct = "/Script/Chooser.FChooserParameterGameplayTagQueryBase", ToolTip="The Gameplay Tag Query property this column will write to"), Category = "Hidden")
+	FInstancedStruct InputValue;
+
+#if WITH_EDITORONLY_DATA
+	UPROPERTY(EditAnywhere, Category=Data, meta=(ToolTip="DefaultRowValue will be assigned to cells when new rows are created"));
+	FGameplayTagQuery DefaultRowValue;
+#endif
+	
+	UPROPERTY()
+	TArray<FGameplayTagQuery> RowValues;
+	
+#if WITH_EDITOR
+	mutable FGameplayTagQuery TestValue;
+	UE_API virtual void AddToDetails(FInstancedPropertyBag& PropertyBag, int32 ColumnIndex, int32 RowIndex) override;
+	UE_API virtual void SetFromDetails(FInstancedPropertyBag& PropertyBag, int32 ColumnIndex, int32 RowIndex) override;
+	virtual void CopyFallback(FChooserColumnBase& SourceColumn) override { FallbackValue = static_cast<FOutputGameplayTagQueryColumn&>(SourceColumn).FallbackValue; }
+#endif
+
+	// FallbackValue will be used as the output value if the all rows in the chooser fail, and the FallbackResult from the chooser is used.
+	UPROPERTY();
+   	FGameplayTagQuery FallbackValue;
+
+	FGameplayTagQuery& GetValueForIndex(int32 Index)
+	{
+		return Index == ChooserColumn_SpecialIndex_Fallback ? FallbackValue : RowValues[Index];
+	}
+	
+	const FGameplayTagQuery& GetValueForIndex(int32 Index) const
+	{
+		return Index == ChooserColumn_SpecialIndex_Fallback ? FallbackValue : RowValues[Index];
+	}
+
+	CHOOSER_COLUMN_BOILERPLATE(FChooserParameterGameplayTagQueryBase);
+};
+
+#undef UE_API
diff --git a/Engine/Plugins/Chooser/Source/Chooser/Private/Chooser.cpp b/Engine/Plugins/Chooser/Source/Chooser/Private/Chooser.cpp
--- a/Engine/Plugins/Chooser/Source/Chooser/Private/Chooser.cpp
+++ b/Engine/Plugins/Chooser/Source/Chooser/Private/Chooser.cpp
@@ -569,6 +569,12 @@ FObjectChooserBase::EIteratorStatus UChooserTable::EvaluateChooser(FChooserEvalu
 }
 
 FObjectChooserBase::EIteratorStatus UChooserTable::IterateChooser(const UChooserTable* Chooser, FObjectChooserBase::FObjectChooserIteratorCallback Callback)
+{
+	FChooserEvaluationContext Context;
+	return IterateChooser(Context, Chooser, Callback);
+}
+
+FObjectChooserBase::EIteratorStatus UChooserTable::IterateChooser(FChooserEvaluationContext& Context, const UChooserTable* Chooser, FObjectChooserBase::FObjectChooserIteratorCallback Callback)
 {
 	TRACE_CPUPROFILER_EVENT_SCOPE(IterateChooser);
 
@@ -604,7 +610,22 @@ FObjectChooserBase::EIteratorStatus UChooserTable::IterateChooser(const UChooser
 		if (ResultsArray->IsValidIndex(SelectedIndexData.Index))
 		{
 			const FObjectChooserBase& SelectedResult = (*ResultsArray)[SelectedIndexData.Index].Get<FObjectChooserBase>();
-			FObjectChooserBase::EIteratorStatus Status = SelectedResult.IterateObjects(Callback);
+
+			for (const FInstancedStruct& ColumnData : Chooser->ColumnsStructs)
+			{
+				const FChooserColumnBase& Column = ColumnData.Get<FChooserColumnBase>();
+
+				#if WITH_EDITORONLY_DATA
+				if (Column.bDisabled)
+				{
+					continue;
+				}
+				#endif
+				
+				Column.SetOutputs(Context, SelectedIndexData.Index);
+			}
+
+			FObjectChooserBase::EIteratorStatus Status = SelectedResult.IterateObjects(Context, Callback);
 			if (Status == FObjectChooserBase::EIteratorStatus::Stop)
 			{
 				return FObjectChooserBase::EIteratorStatus::Stop;
@@ -615,7 +636,13 @@ FObjectChooserBase::EIteratorStatus UChooserTable::IterateChooser(const UChooser
 	if (Chooser->FallbackResult.IsValid())
 	{
 		const FObjectChooserBase& SelectedResult = Chooser->FallbackResult.Get<FObjectChooserBase>();
-		return SelectedResult.IterateObjects(Callback);
+		for (const FInstancedStruct& ColumnData : Chooser->ColumnsStructs)
+		{
+			const FChooserColumnBase& Column = ColumnData.Get<FChooserColumnBase>();
+			Column.SetOutputs(Context, ChooserColumn_SpecialIndex_Fallback);
+		}
+
+		return SelectedResult.IterateObjects(Context, Callback);
 	}
 
 	return FObjectChooserBase::EIteratorStatus::Continue;
@@ -653,9 +680,9 @@ FObjectChooserBase::EIteratorStatus FEvaluateChooser::ChooseMulti(FChooserEvalua
 	return UChooserTable::EvaluateChooser(Context, Chooser, Callback);
 }
 
-FObjectChooserBase::EIteratorStatus FEvaluateChooser::IterateObjects(FObjectChooserIteratorCallback Callback) const
+FObjectChooserBase::EIteratorStatus FEvaluateChooser::IterateObjects(FChooserEvaluationContext& Context, FObjectChooserIteratorCallback Callback) const
 {
-	return UChooserTable::IterateChooser(Chooser, Callback);
+	return UChooserTable::IterateChooser(Context, Chooser, Callback);
 }
 
 void FEvaluateChooser::GetDebugName(FString& OutDebugName) const
@@ -702,9 +729,9 @@ FObjectChooserBase::EIteratorStatus FNestedChooser::ChooseMulti(FChooserEvaluati
 	return UChooserTable::EvaluateChooser(Context, Chooser, Callback);
 }
 
-FObjectChooserBase::EIteratorStatus FNestedChooser::IterateObjects(FObjectChooserIteratorCallback Callback) const
+FObjectChooserBase::EIteratorStatus FNestedChooser::IterateObjects(FChooserEvaluationContext& Context, FObjectChooserIteratorCallback Callback) const
 {
-	return UChooserTable::IterateChooser(Chooser, Callback);
+	return UChooserTable::IterateChooser(Context, Chooser, Callback);
 }
 
 void FNestedChooser::GetDebugName(FString& OutDebugName) const
diff --git a/Engine/Plugins/Chooser/Source/Chooser/Private/ObjectChooser_Asset.cpp b/Engine/Plugins/Chooser/Source/Chooser/Private/ObjectChooser_Asset.cpp
--- a/Engine/Plugins/Chooser/Source/Chooser/Private/ObjectChooser_Asset.cpp
+++ b/Engine/Plugins/Chooser/Source/Chooser/Private/ObjectChooser_Asset.cpp
@@ -7,7 +7,7 @@ UObject* FAssetChooser::ChooseObject(FChooserEvaluationContext& Context) const
 	return Asset;
 }
 
-FObjectChooserBase::EIteratorStatus FAssetChooser::IterateObjects(FObjectChooserIteratorCallback Callback) const
+FObjectChooserBase::EIteratorStatus FAssetChooser::IterateObjects(FChooserEvaluationContext& Context, FObjectChooserIteratorCallback Callback) const
 {
 	return Callback.Execute(Asset);
 }
@@ -22,7 +22,7 @@ UObject* FSoftAssetChooser::ChooseObject(FChooserEvaluationContext& Context) con
 	return Asset.Get();
 }
 
-FObjectChooserBase::EIteratorStatus FSoftAssetChooser::IterateObjects(FObjectChooserIteratorCallback Callback) const
+FObjectChooserBase::EIteratorStatus FSoftAssetChooser::IterateObjects(FChooserEvaluationContext& Context, FObjectChooserIteratorCallback Callback) const
 {
 	return Callback.Execute(Asset.LoadSynchronous());
 }
diff --git a/Engine/Plugins/Chooser/Source/Chooser/Private/ObjectChooser_Class.cpp b/Engine/Plugins/Chooser/Source/Chooser/Private/ObjectChooser_Class.cpp
--- a/Engine/Plugins/Chooser/Source/Chooser/Private/ObjectChooser_Class.cpp
+++ b/Engine/Plugins/Chooser/Source/Chooser/Private/ObjectChooser_Class.cpp
@@ -7,7 +7,7 @@ UObject* FClassChooser::ChooseObject(FChooserEvaluationContext& Context) const
 	return Class;
 }
 
-FObjectChooserBase::EIteratorStatus FClassChooser::IterateObjects(FObjectChooserIteratorCallback Callback) const
+FObjectChooserBase::EIteratorStatus FClassChooser::IterateObjects(FChooserEvaluationContext& Context, FObjectChooserIteratorCallback Callback) const
 {
 	return Callback.Execute(Class);
 }
diff --git a/Engine/Plugins/Chooser/Source/Chooser/Private/OutputGameplayTagQueryColumn.cpp b/Engine/Plugins/Chooser/Source/Chooser/Private/OutputGameplayTagQueryColumn.cpp
new file mode 100644
--- /dev/null
+++ b/Engine/Plugins/Chooser/Source/Chooser/Private/OutputGameplayTagQueryColumn.cpp
@@ -0,0 +1,69 @@
+// Copyright Epic Games, Inc. All Rights Reserved.
+#include "OutputGameplayTagQueryColumn.h"
+
+#include "ChooserIndexArray.h"
+#include "ChooserTrace.h"
+#if WITH_EDITOR
+#include "StructUtils/PropertyBag.h"
+#endif
+
+
+bool FGameplayTagQueryContextProperty::SetValue(FChooserEvaluationContext& Context, const FGameplayTagQuery& Value) const
+{
+	return Binding.SetValue(Context, Value);
+}
+
+FOutputGameplayTagQueryColumn::FOutputGameplayTagQueryColumn()
+{
+	InputValue.InitializeAs(FGameplayTagQueryContextProperty::StaticStruct());
+}
+
+void FOutputGameplayTagQueryColumn::SetOutputs(FChooserEvaluationContext& Context, int RowIndex) const
+{
+	if (InputValue.IsValid())
+	{
+		if (RowIndex == ChooserColumn_SpecialIndex_Fallback || RowValues.IsValidIndex(RowIndex))
+		{
+			const FGameplayTagQuery& OutputTagQuery = GetValueForIndex(RowIndex);
+			InputValue.Get<FChooserParameterGameplayTagQueryBase>().SetValue(Context, OutputTagQuery);
+#if WITH_EDITOR
+			if (Context.DebuggingInfo.bCurrentDebugTarget)
+			{
+				TestValue = OutputTagQuery;
+			}
+#endif
+		}
+		else
+		{
+#if CHOOSER_DEBUGGING_ENABLED
+			UE_ASSET_LOG(LogChooser, Error, Context.DebuggingInfo.CurrentChooser, TEXT("Invalid index %d passed to FOutputGameplayTagQueryColumn::SetOutputs"), RowIndex);
+#else
+			UE_LOG(LogChooser, Error, TEXT("Invalid index %d passed to FOutputGameplayTagQueryColumn::SetOutputs"), RowIndex);
+#endif
+		}
+	}
+}
+
+#if WITH_EDITOR
+	void FOutputGameplayTagQueryColumn::AddToDetails(FInstancedPropertyBag& PropertyBag, int32 ColumnIndex, int32 RowIndex)
+	{
+		FText DisplayName;
+		InputValue.Get<FChooserParameterBase>().GetDisplayName(DisplayName);
+		FName PropertyName("RowData", ColumnIndex);
+		FPropertyBagPropertyDesc PropertyDesc(PropertyName, EPropertyBagPropertyType::Struct, FGameplayTagQuery::StaticStruct());
+		PropertyDesc.MetaData.Add(FPropertyBagPropertyDescMetaData("DisplayName", DisplayName.ToString()));
+		PropertyBag.AddProperties({PropertyDesc});
+		PropertyBag.SetValueStruct(PropertyName, RowValues[RowIndex]);
+	}
+
+	void FOutputGameplayTagQueryColumn::SetFromDetails(FInstancedPropertyBag& PropertyBag, int32 ColumnIndex, int32 RowIndex)
+	{
+		FName PropertyName("RowData", ColumnIndex);
+		
+		TValueOrError<FStructView, EPropertyBagResult> Result = PropertyBag.GetValueStruct(PropertyName, FGameplayTagQuery::StaticStruct());
+		if (FStructView* StructView = Result.TryGetValue())
+		{
+			RowValues[RowIndex] = StructView->Get<FGameplayTagQuery>();
+		}
+	}	
+#endif
diff --git a/Engine/Plugins/Chooser/Source/Chooser/Public/ChooserPropertyAccess.h b/Engine/Plugins/Chooser/Source/Chooser/Public/ChooserPropertyAccess.h
--- a/Engine/Plugins/Chooser/Source/Chooser/Public/ChooserPropertyAccess.h
+++ b/Engine/Plugins/Chooser/Source/Chooser/Public/ChooserPropertyAccess.h
@@ -3,6 +3,7 @@
 #pragma once
 
 #include "CoreMinimal.h"
+#include "GameplayTagContainer.h"
 #include "IHasContext.h"
 #include "StructUtils/InstancedStruct.h"
 #include "IObjectChooser.h"
@@ -116,9 +117,12 @@ struct FChooserPropertyBinding
 	
 	template <typename T>
 	bool GetValue(FChooserEvaluationContext& Context, T& Value) const;
-	
+
+	template <typename T>
+	typename TEnableIf<TIsArithmetic<T>::Value, bool>::Type SetValue(FChooserEvaluationContext& Context, const T& Value) const;
+
 	template <typename T>
-	bool SetValue(FChooserEvaluationContext& Context, const T& Value) const;
+	typename TEnableIf<!TIsArithmetic<T>::Value, bool>::Type SetValue(FChooserEvaluationContext& Context, const T& Value) const;
 };
 
 USTRUCT()
@@ -369,7 +373,7 @@ bool FChooserPropertyBinding::GetValue(FChooserEvaluationContext& Context, T& Ou
 }
 
 template <typename T>
-bool FChooserPropertyBinding::SetValue(FChooserEvaluationContext& Context, const T& InValue) const
+typename TEnableIf<TIsArithmetic<T>::Value, bool>::Type FChooserPropertyBinding::SetValue(FChooserEvaluationContext& Context, const T& Value) const
 {
 	using namespace  UE::Chooser;
 
@@ -380,19 +384,24 @@ bool FChooserPropertyBinding::SetValue(FChooserEvaluationContext& Context, const
 		{
 			switch (Result.PropertyType)
 			{
+			case EChooserPropertyAccessType::None:
+			{
+				*reinterpret_cast<T*>(Result.Container + Result.PropertyOffset) = Value;
+				break;
+			}
 			case EChooserPropertyAccessType::Float:
 				{
-					*reinterpret_cast<float*>(Result.Container + Result.PropertyOffset) = static_cast<float>(InValue);
+					*reinterpret_cast<float*>(Result.Container + Result.PropertyOffset) = static_cast<float>(Value);
 					break;
 				}
 			case EChooserPropertyAccessType::Double:
 				{
-					*reinterpret_cast<double*>(Result.Container + Result.PropertyOffset) = static_cast<double>(InValue);
+					*reinterpret_cast<double*>(Result.Container + Result.PropertyOffset) = static_cast<double>(Value);
 					break;
 				}
 			case EChooserPropertyAccessType::Int32:
 				{
-					*reinterpret_cast<int*>(Result.Container + Result.PropertyOffset) = static_cast<int>(InValue);
+					*reinterpret_cast<int*>(Result.Container + Result.PropertyOffset) = static_cast<int>(Value);
 					break;
 				}
 			case EChooserPropertyAccessType::Bool:
@@ -400,16 +409,16 @@ bool FChooserPropertyBinding::SetValue(FChooserEvaluationContext& Context, const
 				uint8* ByteValue = Result.Container + Result.PropertyOffset;
 				if (Result.Mask == 255) // regular bool
 				{
-					*reinterpret_cast<bool*>(Result.Container + Result.PropertyOffset) = static_cast<bool>(InValue);
+					*reinterpret_cast<bool*>(Result.Container + Result.PropertyOffset) = static_cast<bool>(Value);
 				}
 				else // bitset bool
 				{
-					*ByteValue = ((*ByteValue) & ~Result.Mask) | (InValue ? Result.Mask : 0);
+					*ByteValue = ((*ByteValue) & ~Result.Mask) | (Value ? Result.Mask : 0);
 				}
 				break;
 			}
 			default:
-				*reinterpret_cast<T*>(Result.Container + Result.PropertyOffset) = InValue;
+				*reinterpret_cast<T*>(Result.Container + Result.PropertyOffset) = Value;
 				break;
 			}
 			return true;
@@ -418,6 +427,23 @@ bool FChooserPropertyBinding::SetValue(FChooserEvaluationContext& Context, const
 	return false;
 }
 
+template <typename T>
+typename TEnableIf<!TIsArithmetic<T>::Value, bool>::Type FChooserPropertyBinding::SetValue(FChooserEvaluationContext& Context, const T& Value) const
+{
+	using namespace  UE::Chooser;
+
+	FResolvedPropertyChainResult Result;
+	if (ResolvePropertyChain(Context, *this, Result))
+	{
+		if (Result.Function == nullptr)
+		{
+			*reinterpret_cast<T*>(Result.Container + Result.PropertyOffset) = Value;
+			return true;
+		}
+	}
+	return false;
+}
+
 #if WITH_EDITOR
 #define CHOOSER_PARAMETER_BOILERPLATE() \
 	virtual void Compile(IHasContextClass* Owner, bool bForce) override\
diff --git a/Engine/Plugins/Chooser/Source/Chooser/Public/IChooserParameterGameplayTag.h b/Engine/Plugins/Chooser/Source/Chooser/Public/IChooserParameterGameplayTag.h
--- a/Engine/Plugins/Chooser/Source/Chooser/Public/IChooserParameterGameplayTag.h
+++ b/Engine/Plugins/Chooser/Source/Chooser/Public/IChooserParameterGameplayTag.h
@@ -10,6 +10,7 @@
 #include "IChooserParameterGameplayTag.generated.h"
 
 struct FGameplayTagContainer;
+struct FGameplayTagQuery;
 
 UINTERFACE(MinimalAPI, NotBlueprintType, meta = (CannotImplementInterfaceInBlueprint))
 class UChooserParameterGameplayTag : public UInterface
@@ -30,3 +31,10 @@ struct FChooserParameterGameplayTagBase : public FChooserParameterBase
 	GENERATED_BODY()
 	virtual bool GetValue(FChooserEvaluationContext& Context, const FGameplayTagContainer*& OutResult) const { return false; }
 };
+
+USTRUCT()
+struct FChooserParameterGameplayTagQueryBase : public FChooserParameterBase
+{
+	GENERATED_BODY()
+	virtual bool SetValue(FChooserEvaluationContext& Context, const FGameplayTagQuery& Value) const { return false; }
+};
diff --git a/Engine/Plugins/Chooser/Source/Chooser/Public/IObjectChooser.h b/Engine/Plugins/Chooser/Source/Chooser/Public/IObjectChooser.h
--- a/Engine/Plugins/Chooser/Source/Chooser/Public/IObjectChooser.h
+++ b/Engine/Plugins/Chooser/Source/Chooser/Public/IObjectChooser.h
@@ -149,7 +149,7 @@ public:
 	}
 
 
-	virtual EIteratorStatus IterateObjects(FObjectChooserIteratorCallback Callback) const { return EIteratorStatus::Continue; }
+	virtual EIteratorStatus IterateObjects(FChooserEvaluationContext& Context, FObjectChooserIteratorCallback Callback) const { return EIteratorStatus::Continue; }
 	virtual void GetDebugName(FString& OutDebugName) const {};
 
 #if WITH_EDITOR
diff --git a/Engine/Plugins/Chooser/Source/ChooserEditor/Private/GameplayTagQueryColumnEditor.cpp b/Engine/Plugins/Chooser/Source/ChooserEditor/Private/GameplayTagQueryColumnEditor.cpp
--- a/Engine/Plugins/Chooser/Source/ChooserEditor/Private/GameplayTagQueryColumnEditor.cpp
+++ b/Engine/Plugins/Chooser/Source/ChooserEditor/Private/GameplayTagQueryColumnEditor.cpp
@@ -6,8 +6,10 @@
 
 #include "GameplayTagQueryColumn.h"
 #include "ObjectChooserWidgetFactories.h"
+#include "OutputGameplayTagQueryColumn.h"
 #include "SGameplayTagQueryEntryBox.h"
 #include "SGameplayTagWidget.h"
+#include "SPropertyAccessChainWidget.h"
 #include "SSimpleComboButton.h"
 #include "Widgets/Images/SImage.h"
 
@@ -70,6 +72,7 @@ TSharedRef<SWidget> CreateGameplayTagQueryColumnWidget(UChooserTable* Chooser, F
 			return FGameplayTagQuery::EmptyQuery;
 		})
 		.ReadOnly(false)
+		.DescriptionMaxWidth(1000.f)
 		.OnTagQueryChanged_Lambda([GameplayTagQueryColumn, Row](const FGameplayTagQuery& UpdatedQuery)
 		{
 			if (Row < GameplayTagQueryColumn->RowValues.Num())
@@ -78,10 +81,63 @@ TSharedRef<SWidget> CreateGameplayTagQueryColumnWidget(UChooserTable* Chooser, F
 			}
 		});
 }
+
+TSharedRef<SWidget> CreateOutputGameplayTagQueryColumnWidget(UChooserTable* Chooser, FChooserColumnBase* Column, int Row)
+{
+	FOutputGameplayTagQueryColumn* OutputGameplayTagQueryColumn = static_cast<struct FOutputGameplayTagQueryColumn*>(Column);
 	
+	if (Row == ColumnWidget_SpecialIndex_Fallback)
+	{
+		return SNullWidget::NullWidget;
+	}
+	else if (Row == ColumnWidget_SpecialIndex_Header)
+	{
+		// create column header widget
+		const FSlateBrush* ColumnIcon = FCoreStyle::Get().GetBrush("Icons.ArrowRight");
+		const FText ColumnTooltip = LOCTEXT("Output Gameplay Tag Query Tooltip", "Output Gameplay Tag Query: writes the value from cell in the result row to the bound variable");
+		const FText ColumnName = LOCTEXT("Output Gameplay Tag Query", "Output Gameplay Tag Query");
+		
+		TSharedPtr<SWidget> DebugWidget = nullptr;
+		return MakeColumnHeaderWidget(Chooser, Column, ColumnName, ColumnTooltip, ColumnIcon, DebugWidget);
+	}
+
+	// create cell widget
+	return SNew(SGameplayTagQueryEntryBox)
+		.TagQuery_Lambda([OutputGameplayTagQueryColumn, Row]()
+		{
+			if (Row < OutputGameplayTagQueryColumn->RowValues.Num())
+			{
+				return OutputGameplayTagQueryColumn->RowValues[Row];
+			}
+			return FGameplayTagQuery::EmptyQuery;
+		})
+		.ReadOnly(false)
+		.DescriptionMaxWidth(1000.f)
+		.OnTagQueryChanged_Lambda([OutputGameplayTagQueryColumn, Row](const FGameplayTagQuery& UpdatedQuery)
+		{
+			if (Row < OutputGameplayTagQueryColumn->RowValues.Num())
+			{
+				OutputGameplayTagQueryColumn->RowValues[Row] = UpdatedQuery;
+			}
+		});
+}
+
+TSharedRef<SWidget> CreateGameplayTagQueryPropertyWidget(bool bReadOnly, UObject* TransactionObject, void* Value, UClass* ResultBaseClass, FChooserWidgetValueChanged ValueChanged)
+{
+	IHasContextClass* HasContextClass = Cast<IHasContextClass>(TransactionObject);
+
+	FGameplayTagQueryContextProperty* ContextProperty = reinterpret_cast<FGameplayTagQueryContextProperty*>(Value);
+
+	return SNew(SPropertyAccessChainWidget).ContextClassOwner(HasContextClass).BindingColor("StructPinTypeColor").TypeFilter("FGameplayTagQuery")
+		.PropertyBindingValue(&ContextProperty->Binding)
+		.OnValueChanged(ValueChanged);
+}
+
 void RegisterGameplayTagQueryWidgets()
 {
+	FObjectChooserWidgetFactories::RegisterWidgetCreator(FGameplayTagQueryContextProperty::StaticStruct(), CreateGameplayTagQueryPropertyWidget);
 	FObjectChooserWidgetFactories::RegisterColumnWidgetCreator(FGameplayTagQueryColumn::StaticStruct(), CreateGameplayTagQueryColumnWidget);
+	FObjectChooserWidgetFactories::RegisterColumnWidgetCreator(FOutputGameplayTagQueryColumn::StaticStruct(), CreateOutputGameplayTagQueryColumnWidget);
 	// No need to make and register a creator for gameplay tag containers - it's already registered in GameplayTagColumnEditor
 }
 	
diff --git a/Engine/Plugins/Chooser/Source/ProxyTable/Private/ProxyTable.cpp b/Engine/Plugins/Chooser/Source/ProxyTable/Private/ProxyTable.cpp
--- a/Engine/Plugins/Chooser/Source/ProxyTable/Private/ProxyTable.cpp
+++ b/Engine/Plugins/Chooser/Source/ProxyTable/Private/ProxyTable.cpp
@@ -262,7 +262,7 @@ FObjectChooserBase::EIteratorStatus UProxyTable::FindProxyObjectMulti(const FGui
     return FObjectChooserBase::EIteratorStatus::Continue;
 }
 
-FObjectChooserBase::EIteratorStatus UProxyTable::IterateProxyObjects(const FGuid& Key, FObjectChooserBase::FObjectChooserIteratorCallback Callback) const
+FObjectChooserBase::EIteratorStatus UProxyTable::IterateProxyObjects(const FGuid& Key, FChooserEvaluationContext& Context, FObjectChooserBase::FObjectChooserIteratorCallback Callback) const
 {
 	const int FoundIndex = Algo::BinarySearch(Keys, Key);
 	if (FoundIndex != INDEX_NONE)
@@ -270,13 +270,19 @@ FObjectChooserBase::EIteratorStatus UProxyTable::IterateProxyObjects(const FGuid
 		const FRuntimeProxyValue& EntryValueData = RuntimeValues[FoundIndex];
 		const FObjectChooserBase& EntryValue = EntryValueData.Value.Get<const FObjectChooserBase>();
 
-		FObjectChooserBase::EIteratorStatus Result = EntryValue.IterateObjects(Callback);
+		FObjectChooserBase::EIteratorStatus Result = EntryValue.IterateObjects(Context, Callback);
 		return Result;
 	}
 
 	return FObjectChooserBase::EIteratorStatus::Continue;
 }
 
+FObjectChooserBase::EIteratorStatus UProxyTable::IterateProxyObjects(const FGuid& Key, FObjectChooserBase::FObjectChooserIteratorCallback Callback) const
+{
+	FChooserEvaluationContext Context;
+	return IterateProxyObjects(Key, Context, Callback);
+}
+
 UProxyTable::UProxyTable(const FObjectInitializer& Initializer)
 	:Super(Initializer)
 {
diff --git a/Engine/Plugins/Chooser/Source/ProxyTable/Public/ProxyTable.h b/Engine/Plugins/Chooser/Source/ProxyTable/Public/ProxyTable.h
--- a/Engine/Plugins/Chooser/Source/ProxyTable/Public/ProxyTable.h
+++ b/Engine/Plugins/Chooser/Source/ProxyTable/Public/ProxyTable.h
@@ -87,7 +87,7 @@ public:
 
 	UE_API FObjectChooserBase::EIteratorStatus FindProxyObjectMulti(const FGuid& Key, FChooserEvaluationContext &Context, FObjectChooserBase::FObjectChooserIteratorCallback Callback) const;
 	UE_API UObject* FindProxyObject(const FGuid& Key, FChooserEvaluationContext& Context) const;
-
+	UE_API FObjectChooserBase::EIteratorStatus IterateProxyObjects(const FGuid& Key, FChooserEvaluationContext& Context, FObjectChooserBase::FObjectChooserIteratorCallback Callback) const;
 	UE_API FObjectChooserBase::EIteratorStatus IterateProxyObjects(const FGuid& Key, FObjectChooserBase::FObjectChooserIteratorCallback Callback) const;
 
 	UE_API virtual void PostLoad() override;
