Chooser - Moved FChooserPlayerSettings to a public header file so it can be referenced from native code on the game side

diff --git a/Engine/Plugins/Chooser/Source/Chooser/Internal/AnimNode_ChooserPlayer.h b/Engine/Plugins/Chooser/Source/Chooser/Internal/AnimNode_ChooserPlayer.h
--- a/Engine/Plugins/Chooser/Source/Chooser/Internal/AnimNode_ChooserPlayer.h
+++ b/Engine/Plugins/Chooser/Source/Chooser/Internal/AnimNode_ChooserPlayer.h
@@ -2,7 +2,6 @@
 
 #pragma once
 
-#include "AlphaBlend.h"
 #include "Animation/AnimNode_RelevantAssetPlayerBase.h"
 #include "Animation/AnimationAsset.h"
 #include "CoreMinimal.h"
@@ -10,7 +9,7 @@
 #include "StructUtils/InstancedStruct.h"
 #include "IObjectChooser.h"
 #include "BlendStack/AnimNode_BlendStack.h"
-
+#include "ChooserTypes.h"
 #include "AnimNode_ChooserPlayer.generated.h"
 
 UENUM()
@@ -22,78 +21,6 @@ enum class EChooserEvaluationFrequency
 	OnUpdate
 };
 
-USTRUCT(BlueprintType)
-struct FAnimCurveOverride
-{
-	GENERATED_BODY()
-	// Name of curve to override
-	UPROPERTY(EditAnywhere, BlueprintReadWrite, Category = "CurveValue")
-	FName CurveName;
-	// Value to set to the curve
-	UPROPERTY(EditAnywhere, BlueprintReadWrite, Category = "CurveValue")
-	float CurveValue = 0.0f;
-};
-
-template <>
-struct TTypeTraits<FAnimCurveOverride> : public TTypeTraitsBase < FAnimCurveOverride >
-{
-	enum { IsBytewiseComparable = true };
-};
-
-USTRUCT(BlueprintType)
-struct FAnimCurveOverrideList
-{
-	GENERATED_BODY()
-	UPROPERTY(EditAnywhere, BlueprintReadWrite, Category = "Values", meta=(ExpandByDefault=true))
-	TArray<FAnimCurveOverride> Values;
-
-	UPROPERTY()
-	uint32 Hash = 0;
-	
-	CHOOSER_API void ComputeHash();
-};
-
-
-USTRUCT(BlueprintType)
-struct FChooserPlayerSettings
-{
-	GENERATED_BODY()
-
-	// Set this value to mirror animations - the MirrorDataTable must also be set on the AnimNode
-	UPROPERTY(EditAnywhere, BlueprintReadWrite, Category = "Settings")
-	bool bMirror = false;
-	
-	// Start offset when starting the Animation Asset
-	UPROPERTY(EditAnywhere, BlueprintReadWrite, Category = "Settings", meta=(FrameTimeEditor))
-    float StartTime = 0;
-
-	// Loop the animation asset, even if the asset is not set as looping
-	UPROPERTY(EditAnywhere, BlueprintReadWrite, Category = "Settings")
-    bool bForceLooping = false;
-
-	// playback rate modifier
-	UPROPERTY(EditAnywhere, BlueprintReadWrite ,Category = "Settings")
-    float PlaybackRate = 1.0;
-
-	// List of curve values to set 
-	UPROPERTY(EditAnywhere, BlueprintReadWrite, Category = "Settings", meta=(ExpandByDefault=true))
-	FAnimCurveOverrideList CurveOverrides;
-	
-	UPROPERTY(EditAnywhere, BlueprintReadWrite, Category = "Blending", meta=(FrameTimeEditor))
-	float BlendTime = 0.2;
-	
-	// Set Blend Profiles (editable in the skeleton) to determine how the blending is distributed among your character's bones. It could be used to differentiate between upper body and lower body to blend timing.
-	UPROPERTY(EditAnywhere, BlueprintReadWrite, Category = Blending, meta = (UseAsBlendProfile = true))
-	TObjectPtr<UBlendProfile> BlendProfile;
- 
-	// How the blend is applied over time to the bones. Common selections are linear, ease in, ease out, and ease in and out.
-	UPROPERTY(EditAnywhere, BlueprintReadWrite, Category = Blending)
-	EAlphaBlendOption BlendOption = EAlphaBlendOption::Linear;
-
-	UPROPERTY(EditAnywhere, BlueprintReadWrite, Category = Blending)
-	bool bUseInertialBlend = false;
-};
-
 USTRUCT(BlueprintInternalUseOnly)
 struct FAnimNode_ChooserPlayer : public FAnimNode_BlendStack_Standalone, public IHasContextClass
 {
diff --git a/Engine/Plugins/Chooser/Source/Chooser/Public/ChooserTypes.h b/Engine/Plugins/Chooser/Source/Chooser/Public/ChooserTypes.h
new file mode 100644
--- /dev/null
+++ b/Engine/Plugins/Chooser/Source/Chooser/Public/ChooserTypes.h
@@ -0,0 +1,79 @@
+// Copyright Epic Games, Inc. All Rights Reserved.
+
+#pragma once
+
+#include "CoreMinimal.h"
+#include "AlphaBlend.h"
+#include "Animation/BlendProfile.h"
+#include "ChooserTypes.generated.h"
+
+USTRUCT(BlueprintType)
+struct FAnimCurveOverride
+{
+	GENERATED_BODY()
+	// Name of curve to override
+	UPROPERTY(EditAnywhere, BlueprintReadWrite, Category = "CurveValue")
+	FName CurveName;
+	// Value to set to the curve
+	UPROPERTY(EditAnywhere, BlueprintReadWrite, Category = "CurveValue")
+	float CurveValue = 0.0f;
+};
+
+template <>
+struct TTypeTraits<FAnimCurveOverride> : public TTypeTraitsBase < FAnimCurveOverride >
+{
+	enum { IsBytewiseComparable = true };
+};
+
+USTRUCT(BlueprintType)
+struct FAnimCurveOverrideList
+{
+	GENERATED_BODY()
+	UPROPERTY(EditAnywhere, BlueprintReadWrite, Category = "Values", meta = (ExpandByDefault = true))
+	TArray<FAnimCurveOverride> Values;
+
+	UPROPERTY()
+	uint32 Hash = 0;
+
+	CHOOSER_API void ComputeHash();
+};
+
+USTRUCT(BlueprintType)
+struct FChooserPlayerSettings
+{
+	GENERATED_BODY()
+
+	// Set this value to mirror animations - the MirrorDataTable must also be set on the AnimNode
+	UPROPERTY(EditAnywhere, BlueprintReadWrite, Category = "Settings")
+	bool bMirror = false;
+
+	// Start offset when starting the Animation Asset
+	UPROPERTY(EditAnywhere, BlueprintReadWrite, Category = "Settings", meta = (FrameTimeEditor))
+	float StartTime = 0;
+
+	// Loop the animation asset, even if the asset is not set as looping
+	UPROPERTY(EditAnywhere, BlueprintReadWrite, Category = "Settings")
+	bool bForceLooping = false;
+
+	// playback rate modifier
+	UPROPERTY(EditAnywhere, BlueprintReadWrite, Category = "Settings")
+	float PlaybackRate = 1.0;
+
+	// List of curve values to set 
+	UPROPERTY(EditAnywhere, BlueprintReadWrite, Category = "Settings", meta = (ExpandByDefault = true))
+	FAnimCurveOverrideList CurveOverrides;
+
+	UPROPERTY(EditAnywhere, BlueprintReadWrite, Category = "Blending", meta = (FrameTimeEditor))
+	float BlendTime = 0.2;
+
+	// Set Blend Profiles (editable in the skeleton) to determine how the blending is distributed among your character's bones. It could be used to differentiate between upper body and lower body to blend timing.
+	UPROPERTY(EditAnywhere, BlueprintReadWrite, Category = Blending, meta = (UseAsBlendProfile = true))
+	TObjectPtr<UBlendProfile> BlendProfile;
+
+	// How the blend is applied over time to the bones. Common selections are linear, ease in, ease out, and ease in and out.
+	UPROPERTY(EditAnywhere, BlueprintReadWrite, Category = Blending)
+	EAlphaBlendOption BlendOption = EAlphaBlendOption::Linear;
+
+	UPROPERTY(EditAnywhere, BlueprintReadWrite, Category = Blending)
+	bool bUseInertialBlend = false;
+};
\ No newline at end of file
