Fix crash with Proxy table compilation

diff --git a/Engine/Plugins/Chooser/Source/ProxyTable/Private/ProxyTable.cpp b/Engine/Plugins/Chooser/Source/ProxyTable/Private/ProxyTable.cpp
--- a/Engine/Plugins/Chooser/Source/ProxyTable/Private/ProxyTable.cpp
+++ b/Engine/Plugins/Chooser/Source/ProxyTable/Private/ProxyTable.cpp
@@ -170,25 +170,34 @@ void UProxyTable::PostLoad()
 #endif
 
 	// compilation for property accesses
-	for(FRuntimeProxyValue& Entry : RuntimeValues) 
+	for (FRuntimeProxyValue& Entry : RuntimeValues)
 	{
-		if (FObjectChooserBase* Result = Entry.Value.GetMutablePtr<FObjectChooserBase>())
-		{
-			// need to compile results in case one is a LookupProxy
-			Result->Compile(Entry.ProxyAsset, false);
-		}
-		
-		/// compile any struct output property references
-		for(FProxyStructOutput& StructOutput : Entry.OutputStructData)
+
+		if (Entry.ProxyAsset)
 		{
-			StructOutput.Binding.Compile(Entry.ProxyAsset);
-#if WITH_EDITORONLY_DATA
-			if (!StructOutput.Binding.CompileMessage.IsEmpty())
+			if (FObjectChooserBase* Result = Entry.Value.GetMutablePtr<FObjectChooserBase>())
 			{
-				UE_LOG(LogChooser, Error, TEXT("ProxyTable Struct Output Compile Error (%s - %s): %s."),  *GetName(), *Entry.ProxyAsset.GetName(), *StructOutput.Binding.CompileMessage.ToString());
+				// need to compile results in case one is a LookupProxy
+				Result->Compile(Entry.ProxyAsset, false);
+			}
+
+			/// compile any struct output property references
+			for (FProxyStructOutput& StructOutput : Entry.OutputStructData)
+			{
+				StructOutput.Binding.Compile(Entry.ProxyAsset);
+	#if WITH_EDITORONLY_DATA
+				if (!StructOutput.Binding.CompileMessage.IsEmpty())
+				{
+					UE_LOG(LogChooser, Error, TEXT("ProxyTable Struct Output Compile Error (%s - %s): %s."), *GetName(), *Entry.ProxyAsset.GetName(), *StructOutput.Binding.CompileMessage.ToString());
+				}
+	#endif
 			}
-#endif
 		}
+		else
+		{
+			UE_LOG(LogChooser, Warning, TEXT("ProxyTable has an entry with no ProxyAsset assigned (%s)"), *GetName());
+		}
+
 	}
 }
 
