UE: Lower verbosity to display for now

diff --git a/Engine/Plugins/Chooser/Source/ProxyTable/Private/ProxyTable.cpp b/Engine/Plugins/Chooser/Source/ProxyTable/Private/ProxyTable.cpp
--- a/Engine/Plugins/Chooser/Source/ProxyTable/Private/ProxyTable.cpp
+++ b/Engine/Plugins/Chooser/Source/ProxyTable/Private/ProxyTable.cpp
@@ -195,7 +195,7 @@ void UProxyTable::PostLoad()
 		}
 		else
 		{
-			UE_LOG(LogChooser, Warning, TEXT("ProxyTable has an entry with no ProxyAsset assigned (%s)"), *GetName());
+			UE_LOG(LogChooser, Display, TEXT("ProxyTable has an entry with no ProxyAsset assigned (%s)"), *GetName());
 		}
 
 	}
