Fix chooser player not resetting its definition context on paste, potentially keeping references to its previous owner anim BP.

diff --git a/Engine/Plugins/Chooser/Source/ChooserUncooked/Private/AnimGraphNode_ChooserPlayer.cpp b/Engine/Plugins/Chooser/Source/ChooserUncooked/Private/AnimGraphNode_ChooserPlayer.cpp
--- a/Engine/Plugins/Chooser/Source/ChooserUncooked/Private/AnimGraphNode_ChooserPlayer.cpp
+++ b/Engine/Plugins/Chooser/Source/ChooserUncooked/Private/AnimGraphNode_ChooserPlayer.cpp
@@ -33,6 +33,13 @@ void UAnimGraphNode_ChooserPlayer::PostPlacedNewNode()
 	UpdateContextData();
 }
 
+void UAnimGraphNode_ChooserPlayer::PostPasteNode()
+{
+	Super::PostPasteNode();
+	Node.ChooserContextDefinition.Empty();
+	UpdateContextData();
+}
+
 FLinearColor UAnimGraphNode_ChooserPlayer::GetNodeTitleColor() const
 {
 	return FLinearColor(0.10f, 0.60f, 0.12f);
diff --git a/Engine/Plugins/Chooser/Source/ChooserUncooked/Public/AnimGraphNode_ChooserPlayer.h b/Engine/Plugins/Chooser/Source/ChooserUncooked/Public/AnimGraphNode_ChooserPlayer.h
--- a/Engine/Plugins/Chooser/Source/ChooserUncooked/Public/AnimGraphNode_ChooserPlayer.h
+++ b/Engine/Plugins/Chooser/Source/ChooserUncooked/Public/AnimGraphNode_ChooserPlayer.h
@@ -26,6 +26,7 @@ public:
 	
 	// UEdGraphNode interface
 	virtual void PostPlacedNewNode() override;
+	virtual void PostPasteNode() override;
 	virtual FLinearColor GetNodeTitleColor() const override;
 	virtual FText GetTooltipText() const override;
 	virtual FText GetNodeTitle(ENodeTitleType::Type TitleType) const override;
