Update title of Chooser Player nodes to contain the asset name, which also allows find in blueprints to find those nodes when searching for a chooser or proxy asset name

diff --git a/Engine/Plugins/Chooser/Source/ChooserUncooked/Private/AnimGraphNode_ChooserPlayer.cpp b/Engine/Plugins/Chooser/Source/ChooserUncooked/Private/AnimGraphNode_ChooserPlayer.cpp
--- a/Engine/Plugins/Chooser/Source/ChooserUncooked/Private/AnimGraphNode_ChooserPlayer.cpp
+++ b/Engine/Plugins/Chooser/Source/ChooserUncooked/Private/AnimGraphNode_ChooserPlayer.cpp
@@ -73,6 +73,12 @@ FText UAnimGraphNode_ChooserPlayer::GetTooltipText() const
 
 FText UAnimGraphNode_ChooserPlayer::GetNodeTitle(ENodeTitleType::Type TitleType) const
 {
+	if (Node.Chooser.IsValid())
+	{
+		FString Name;
+		Node.Chooser.Get<FObjectChooserBase>().GetDebugName(Name);
+		return FText::Format(LOCTEXT("EvaluateChooser2_TitleWithChooser", "Chooser Player: {0}"), {FText::FromString(Name)});
+	}
 	return LOCTEXT("NodeTitle", "Chooser Player");
 }
 
