Fix transactional flag on nested choosers when they are duplicated by copy/paste

diff --git a/Engine/Plugins/Chooser/Source/ChooserEditor/Private/ChooserTableEditor.cpp b/Engine/Plugins/Chooser/Source/ChooserEditor/Private/ChooserTableEditor.cpp
--- a/Engine/Plugins/Chooser/Source/ChooserEditor/Private/ChooserTableEditor.cpp
+++ b/Engine/Plugins/Chooser/Source/ChooserEditor/Private/ChooserTableEditor.cpp
@@ -2271,7 +2271,7 @@ UChooserTable* DuplicateNestedChooser(UChooserTable* Chooser, UChooserTable* New
 		return *FoundTable;
 	}
 	
-	UChooserTable* NewTable = NewObject<UChooserTable>(NewOuter, Chooser->GetFName());
+	UChooserTable* NewTable = NewObject<UChooserTable>(NewOuter, Chooser->GetFName(), RF_Transactional);
 	NewTable->ColumnsStructs = Chooser->ColumnsStructs;
 	NewTable->ResultsStructs = Chooser->ResultsStructs;
 	NewTable->RootChooser = RootTable;
