Chooser - fix for crash when Column.GetInputValue() is null in FRewindDebuggerChooser::Update

diff --git a/Engine/Plugins/Chooser/Source/ChooserEditor/Private/RewindDebuggerChooser.cpp b/Engine/Plugins/Chooser/Source/ChooserEditor/Private/RewindDebuggerChooser.cpp
--- a/Engine/Plugins/Chooser/Source/ChooserEditor/Private/RewindDebuggerChooser.cpp
+++ b/Engine/Plugins/Chooser/Source/ChooserEditor/Private/RewindDebuggerChooser.cpp
@@ -80,9 +80,12 @@ void FRewindDebuggerChooser::Update(float DeltaTime, IRewindDebugger* RewindDebu
 										for(FInstancedStruct& ColumnStruct : Chooser->ColumnsStructs)
 										{
 											FChooserColumnBase& Column = ColumnStruct.GetMutable<FChooserColumnBase>();
-											if(ValueData.Key == Column.GetInputValue()->GetDebugName())
+											if (const FChooserParameterBase* InputValue = Column.GetInputValue())
 											{
-												Column.SetTestValue(ValueData.Value);
+												if (ValueData.Key == InputValue->GetDebugName())
+												{
+													Column.SetTestValue(ValueData.Value);
+												}
 											}
 										}
 									}
