Fix incorrect base type on chooser column instanced structs

diff --git a/Engine/Plugins/Chooser/Source/Chooser/Internal/Chooser.h b/Engine/Plugins/Chooser/Source/Chooser/Internal/Chooser.h
--- a/Engine/Plugins/Chooser/Source/Chooser/Internal/Chooser.h
+++ b/Engine/Plugins/Chooser/Source/Chooser/Internal/Chooser.h
@@ -176,7 +176,7 @@ public:
 	TArray<FInstancedStruct> CookedResults;
 
 	// Columns which filter Results
-	UPROPERTY(EditAnywhere, NoClear, DisplayName = "Results", Meta = (ExcludeBaseStruct, BaseStruct = "/Script/Chooser.ObjectChooserBase"), Category = "Hidden")
+	UPROPERTY(EditAnywhere, NoClear, DisplayName = "Columns", Meta = (ExcludeBaseStruct, BaseStruct = "/Script/Chooser.ChooserColumnBase"), Category = "Hidden")
 	TArray<FInstancedStruct> ColumnsStructs;
 	
 	// The kind of output this chooser has (Object or Class or No primary result)
