Add basic support for nesting arbitrary asset types inside chooser tables

diff --git a/Engine/Plugins/Chooser/Source/Chooser/Internal/Chooser.h b/Engine/Plugins/Chooser/Source/Chooser/Internal/Chooser.h
--- a/Engine/Plugins/Chooser/Source/Chooser/Internal/Chooser.h
+++ b/Engine/Plugins/Chooser/Source/Chooser/Internal/Chooser.h
@@ -153,10 +153,17 @@ public:
 
 	UPROPERTY()
 	TArray<TObjectPtr<UChooserTable>> NestedChoosers;
+	
+	UPROPERTY()
+	TArray<TObjectPtr<UObject>> NestedObjects;
 
 	CHOOSER_API void AddNestedChooser(UChooserTable* Chooser);
 	CHOOSER_API void RemoveNestedChooser(UChooserTable* Chooser);
 	FSimpleMulticastDelegate NestedChoosersChanged;
+	
+	CHOOSER_API void AddNestedObject(UObject* Object);
+	CHOOSER_API void RemoveNestedObject(UObject* Object);
+	FSimpleMulticastDelegate NestedObjectsChanged;
 
 	// deprecated in favor of RootChooser
 	UPROPERTY()
@@ -166,7 +173,7 @@ public:
 	float EditorResultsColumnWidth = 300;
 
 	UPROPERTY()
-	uint32 Version = 0;
+	uint32 Version = 2;
 #endif
 
 	UFUNCTION()
diff --git a/Engine/Plugins/Chooser/Source/Chooser/Private/Chooser.cpp b/Engine/Plugins/Chooser/Source/Chooser/Private/Chooser.cpp
--- a/Engine/Plugins/Chooser/Source/Chooser/Private/Chooser.cpp
+++ b/Engine/Plugins/Chooser/Source/Chooser/Private/Chooser.cpp
@@ -47,7 +47,7 @@ void UChooserTable::PostLoad()
 	CachedPreviousOutputObjectType = OutputObjectType;
 	CachedPreviousResultType = ResultType;
 
-	if (Version != 1)
+	if (Version < 1)
 	{
 		if (ParentTable)
 		{
@@ -63,20 +63,28 @@ void UChooserTable::PostLoad()
 			}
 		}
 
-		if(RootChooser == nullptr && NestedChoosers.IsEmpty())
+		if(RootChooser == nullptr && NestedObjects.IsEmpty())
 		{
-			// data upgrade for root tables: add elements to NestedTables list
+			// data upgrade for root tables: add elements to NestedObjects list
 			TArray<UObject*> ChildObjects;
 			GetObjectsWithOuter(GetPackage(), ChildObjects, true);
 			for(UObject* ChildObject : ChildObjects)
 			{
 				if(UChooserTable* Chooser = Cast<UChooserTable>(ChildObject))
 				{
-					NestedChoosers.Add(Chooser);
+					NestedObjects.Add(Chooser);
 				}
 			}
 		}
-		Version = 1;
+		Version = 2;
+	}
+	if (Version < 2)
+	{
+		for(UChooserTable* NestedChooser : NestedChoosers)
+		{
+			NestedObjects.Add(NestedChooser);
+		}
+		Version = 2;
 	}
 #endif
 
@@ -758,17 +766,22 @@ static FAutoConsoleCommand CCmdTestCookChoosers(
 
 void UChooserTable::AddNestedChooser(UChooserTable* Chooser)
 {
-	FName NewChooserName = Chooser->GetFName();
-	FString NewChooserPlainName = NewChooserName.GetPlainNameString();
+	AddNestedObject(Chooser);
+}
+
+void UChooserTable::AddNestedObject(UObject* Object)
+{
+	FName NewObjectName = Object->GetFName();
+	FString NewChooserPlainName = NewObjectName.GetPlainNameString();
 	bool bNameConflict = false;
-	int MaxNumber = NewChooserName.GetNumber();
-	for(UChooserTable* NestedChooserIterator : NestedChoosers)
+	int MaxNumber = NewObjectName.GetNumber();
+	for(UObject* NestedObjectIterator : NestedObjects)
 	{
-		FName NestedChooserName = NestedChooserIterator->GetFName();
-		if (NestedChooserName.GetPlainNameString() == NewChooserPlainName)
+		FName NestedObjectName = NestedObjectIterator->GetFName();
+		if (NestedObjectName.GetPlainNameString() == NewChooserPlainName)
 		{
-			MaxNumber = FMath::Max(MaxNumber, NestedChooserName.GetNumber());
-			if(NestedChooserName.GetNumber() == NewChooserName.GetNumber())
+			MaxNumber = FMath::Max(MaxNumber, NestedObjectName.GetNumber());
+			if(NestedObjectName.GetNumber() == NewObjectName.GetNumber())
 			{
 				bNameConflict = true;
 			}
@@ -776,18 +789,23 @@ void UChooserTable::AddNestedChooser(UChooserTable* Chooser)
 	}
 	if (bNameConflict)
 	{
-		NewChooserName.SetNumber(MaxNumber + 1);
-		Chooser->Rename(*NewChooserName.ToString());
+		NewObjectName.SetNumber(MaxNumber + 1);
+		Object->Rename(*NewObjectName.ToString());
 	}
 	
-	NestedChoosers.Add(Chooser);
-	NestedChoosersChanged.Broadcast();
+	NestedObjects.Add(Object);
+	NestedObjectsChanged.Broadcast();
 }
 
 void UChooserTable::RemoveNestedChooser(UChooserTable* Chooser)
 {
-	NestedChoosers.Remove(Chooser);
-	NestedChoosersChanged.Broadcast();
+	RemoveNestedObject(Chooser);
+}
+
+void UChooserTable::RemoveNestedObject(UObject* Object)
+{
+	NestedObjects.Remove(Object);
+	NestedObjectsChanged.Broadcast();
 }
 
 #endif
\ No newline at end of file
diff --git a/Engine/Plugins/Chooser/Source/ChooserEditor/Private/ChooserEditorWidgets.cpp b/Engine/Plugins/Chooser/Source/ChooserEditor/Private/ChooserEditorWidgets.cpp
--- a/Engine/Plugins/Chooser/Source/ChooserEditor/Private/ChooserEditorWidgets.cpp
+++ b/Engine/Plugins/Chooser/Source/ChooserEditor/Private/ChooserEditorWidgets.cpp
@@ -190,42 +190,23 @@ TSharedRef<SWidget> CreateNestedChooserWidget(bool bReadOnly, UObject* Transacti
 													UChooserTable* ParentChooser = Cast<UChooserTable>(TransactionObject);
 													UChooserTable* RootChooser = nullptr;
 													
-													if (ParentChooser)
-													{
-														RootChooser = ParentChooser->GetRootChooser();
-
-														bool bNameConflict = false;
-														int MaxNumber = NewChooserName.GetNumber();
-														for(UChooserTable* NestedChooserIterator : RootChooser->NestedChoosers)
-														{
-															FName NestedChooserName = NestedChooserIterator->GetFName();
-															if (NestedChooserName.GetPlainNameString() == NewChooserPlainName)
-															{
-																MaxNumber = FMath::Max(MaxNumber, NestedChooserName.GetNumber());
-																if(NestedChooserName.GetNumber() == NewChooserName.GetNumber())
-																{
-																	bNameConflict = true;
-																}
-															}
-														}
-														if (bNameConflict)
-														{
-															NewChooserName.SetNumber(MaxNumber + 1);
-														}
-													}
-													
 													UChooserTable* NewChooser = NewObject<UChooserTable>(TransactionObject, UChooserTable::StaticClass(), NewChooserName, RF_Transactional);
 												
 													const FScopedTransaction Transaction(LOCTEXT("Assign New Nested Chooser", "Assign New Nested Chooser"));
 												
 													TransactionObject->Modify();
 													NestedChooser->Chooser = NewChooser;
+													
+													if (ParentChooser)
+													{
+														RootChooser = ParentChooser->GetRootChooser();
+													}
 
 													if (RootChooser)
 													{
 														NestedChooser->Chooser->RootChooser = RootChooser;
 														RootChooser->Modify();
-														RootChooser->AddNestedChooser(NestedChooser->Chooser);
+														RootChooser->AddNestedObject(NestedChooser->Chooser);
 													}
 												}
 											});
@@ -252,24 +233,25 @@ TSharedRef<SWidget> CreateNestedChooserWidget(bool bReadOnly, UObject* Transacti
 										{
 											RootTable = OuterChooser->GetRootChooser();
 
-											for (UChooserTable* Chooser : RootTable->NestedChoosers)
+											for (UObject* Object : RootTable->NestedObjects)
 											{
-												 // exclude the main chooser 
-												if (Chooser != RootTable)
+												if (UChooserTable* Chooser = Cast<UChooserTable>(Object))
 												{
-													SubMenuBuilder.AddMenuEntry( FText::FromString(Chooser->GetName()), LOCTEXT("AddExistingObjectTooltip", "Add a reference to this existing Chooser Table."), FSlateIcon(),
-														FUIAction(FExecuteAction::CreateLambda([Chooser, Button, NestedChooser, TransactionObject]()
-														{
-															const FScopedTransaction Transaction(LOCTEXT("Set Nested Chooser", "Set Nested Chooser"));
-															TransactionObject->Modify(true);
-															Button->SetIsOpen(false);
-															NestedChooser->Chooser = Chooser;
-														})));
+													// exclude the main chooser 
+												   if (Chooser != RootTable)
+												   {
+													   SubMenuBuilder.AddMenuEntry( FText::FromString(Chooser->GetName()), LOCTEXT("AddExistingObjectTooltip", "Add a reference to this existing Chooser Table."), FSlateIcon(),
+														   FUIAction(FExecuteAction::CreateLambda([Chooser, Button, NestedChooser, TransactionObject]()
+														   {
+															   const FScopedTransaction Transaction(LOCTEXT("Set Nested Chooser", "Set Nested Chooser"));
+															   TransactionObject->Modify(true);
+															   Button->SetIsOpen(false);
+															   NestedChooser->Chooser = Chooser;
+														   })));
+												   }
 												}
 											}
 										}
-										
-
 									}));
                                      
 									
diff --git a/Engine/Plugins/Chooser/Source/ChooserEditor/Private/ChooserTableEditor.cpp b/Engine/Plugins/Chooser/Source/ChooserEditor/Private/ChooserTableEditor.cpp
--- a/Engine/Plugins/Chooser/Source/ChooserEditor/Private/ChooserTableEditor.cpp
+++ b/Engine/Plugins/Chooser/Source/ChooserEditor/Private/ChooserTableEditor.cpp
@@ -809,10 +809,10 @@ void FChooserTableEditor::RefreshAll()
 		}
 	}
 	
-	RefreshNestedChoosers();
+	RefreshNestedObjectTree();
 }
 
-void FChooserTableEditor::RefreshNestedChoosers()
+void FChooserTableEditor::RefreshNestedObjectTree()
 {
 	if (NestedChooserTree.IsValid())
 	{
@@ -1302,7 +1302,7 @@ void FChooserTableEditor::MakeChoosersMenuRecursive(UObject* Outer, FMenuBuilder
 	{
 		if (UChooserTable* Chooser = Cast<UChooserTable>(Object))
 		{
-			if (Chooser == RootChooser || Chooser->GetRootChooser()->NestedChoosers.Contains(Chooser))
+			if (Chooser == RootChooser || Chooser->GetRootChooser()->NestedObjects.Contains(Chooser))
 			{
 				MenuBuilder.AddMenuEntry( FText::FromString(Indent + Chooser->GetName()), LOCTEXT("Edit Chooser ToolTip", "Browse to this Nested Chooser Table"), FSlateIcon(),
 					FUIAction(FExecuteAction::CreateLambda([this, Chooser]()
@@ -2263,20 +2263,24 @@ void FChooserTableEditor::MoveColumnRight()
 UChooserTable* DuplicateNestedChooser(UChooserTable* Chooser, UChooserTable* NewOuter)
 {
 	UChooserTable* RootTable = NewOuter->GetRootChooser();
-	if (TObjectPtr<UChooserTable>* FoundTable = RootTable->NestedChoosers.FindByPredicate([Chooser](UChooserTable* Table)
+	if (TObjectPtr<UObject>* FoundTable = RootTable->NestedObjects.FindByPredicate([Chooser](UObject* Object)
 		{
-			return Table->GetName() == Chooser->GetName();
+			if (Cast<UChooserTable>(Object))
+			{
+				return Object->GetName() == Chooser->GetName();
+			}
+			return false;
 		}))
 	{
 		// we already duplicated this table
-		return *FoundTable;
+		return Cast<UChooserTable>(*FoundTable);
 	}
 	
 	UChooserTable* NewTable = NewObject<UChooserTable>(NewOuter, Chooser->GetFName(), RF_Transactional);
 	NewTable->ColumnsStructs = Chooser->ColumnsStructs;
 	NewTable->ResultsStructs = Chooser->ResultsStructs;
 	NewTable->RootChooser = RootTable;
-	RootTable->AddNestedChooser(NewTable);
+	RootTable->AddNestedObject(NewTable);
 
 	for (FInstancedStruct& ResultData : NewTable->ResultsStructs)
 	{
@@ -2506,7 +2510,7 @@ void FChooserTableEditor::PasteInternal(UChooserTable* PastedContent, int PasteR
 			}
 		}
 		
-		if (!PastedContent->NestedChoosers.IsEmpty())
+		if (!PastedContent->NestedObjects.IsEmpty())
 		{
 			// if there were nested choosers in the copy buffer we have to remap or paste them here
 			for (int ResultIndex = InsertIndex; ResultIndex < PastedContent->ResultsStructs.Num() + InsertIndex; ResultIndex++)
diff --git a/Engine/Plugins/Chooser/Source/ChooserEditor/Private/ChooserTableEditor.h b/Engine/Plugins/Chooser/Source/ChooserEditor/Private/ChooserTableEditor.h
--- a/Engine/Plugins/Chooser/Source/ChooserEditor/Private/ChooserTableEditor.h
+++ b/Engine/Plugins/Chooser/Source/ChooserEditor/Private/ChooserTableEditor.h
@@ -94,7 +94,7 @@ namespace UE::ChooserEditor
 		void PushChooserTableToEdit(UChooserTable* Chooser);
 		void PopChooserTableToEdit();
 		void RefreshAll();
-		void RefreshNestedChoosers();
+		void RefreshNestedObjectTree();
 	
 		/** Used to show or hide certain properties */
 		void SetPropertyVisibilityDelegate(FIsPropertyVisible InVisibilityDelegate);
diff --git a/Engine/Plugins/Chooser/Source/ChooserEditor/Private/SNestedChooserTree.cpp b/Engine/Plugins/Chooser/Source/ChooserEditor/Private/SNestedChooserTree.cpp
--- a/Engine/Plugins/Chooser/Source/ChooserEditor/Private/SNestedChooserTree.cpp
+++ b/Engine/Plugins/Chooser/Source/ChooserEditor/Private/SNestedChooserTree.cpp
@@ -19,7 +19,7 @@ namespace UE::ChooserEditor
 			return SNew(STableRow<TSharedPtr<FNestedChooserTreeEntry>>, OwnerTable)
 						.Content()
 						[
-							SNew(STextBlock).Text(FText::FromString(InItem->Chooser->GetName()))
+							SNew(STextBlock).Text(FText::FromString(InItem->Object->GetName()))
 						];
 	}
 	
@@ -27,7 +27,7 @@ namespace UE::ChooserEditor
     {
 		for(TSharedPtr<FNestedChooserTreeEntry>& Entry : AllChoosers)
 		{
-			if (Entry->Chooser->GetOuter() == InItem->Chooser)
+			if (Entry->Object->GetOuter() == InItem->Object)
 			{
 				OutChildren.Add(Entry);
 			}
@@ -36,7 +36,10 @@ namespace UE::ChooserEditor
 	
 	void SNestedChooserTree::TreeViewDoubleClicked(TSharedPtr<FNestedChooserTreeEntry> SelectedObject)
 	{
-		ChooserEditor->SetChooserTableToEdit(SelectedObject->Chooser);
+		if (UChooserTable* Chooser = Cast<UChooserTable>(SelectedObject->Object))
+		{
+			ChooserEditor->SetChooserTableToEdit(Chooser);
+		}
 	}
 
 	// Replace any references to a nested chooser, from either the Result column, or an OutputObject column
@@ -102,14 +105,17 @@ namespace UE::ChooserEditor
 	{
 		ReplaceReferencesInTable(ChooserToReplace, ReplaceWith, RootTable);
 
-		for(UChooserTable* NestedChooser : RootTable->NestedChoosers)
+		for(UObject* NestedObject : RootTable->NestedObjects)
 		{
-			ReplaceReferencesInTable(ChooserToReplace, ReplaceWith, NestedChooser);
-
-			// reparent any child tables to the root
-			if (NestedChooser->GetOuter() == ChooserToReplace)
+			if(UChooserTable* NestedChooser = Cast<UChooserTable>(NestedObject))
 			{
-				NestedChooser->Rename(nullptr, RootTable);
+				ReplaceReferencesInTable(ChooserToReplace, ReplaceWith, NestedChooser);
+
+				// reparent any child tables to the root
+				if (NestedChooser->GetOuter() == ChooserToReplace)
+				{
+					NestedChooser->Rename(nullptr, RootTable);
+				}
 			}
 		}
 	}
@@ -119,35 +125,36 @@ namespace UE::ChooserEditor
 		TArray<TSharedPtr<FNestedChooserTreeEntry>>SelectedItems = TreeView->GetSelectedItems();
 		if (SelectedItems.Num() == 1)
 		{
-			UChooserTable* ChooserToDelete = SelectedItems[0]->Chooser;
-
-			if (ChooserToDelete != RootChooser)
+			if (UChooserTable* ChooserToDelete = Cast<UChooserTable>(SelectedItems[0]->Object))
 			{
-				FScopedTransaction Transaction(LOCTEXT("Delete Nested Choosers", "Delete Nested Choosers"));
-				static FName DeletedChooserName = "DeletedNestedChooser";
-				DeletedChooserName.SetNumber(DeletedChooserName.GetNumber() + 1);
-				ChooserToDelete->Rename(*DeletedChooserName.ToString());
-				RootChooser->Modify(true);
-				RootChooser->RemoveNestedChooser(ChooserToDelete);
-				ReplaceReferences(ChooserToDelete, nullptr, RootChooser);
+				if (ChooserToDelete != RootChooser)
+				{
+					FScopedTransaction Transaction(LOCTEXT("Delete Nested Choosers", "Delete Nested Choosers"));
+					static FName DeletedChooserName = "DeletedNestedChooser";
+					DeletedChooserName.SetNumber(DeletedChooserName.GetNumber() + 1);
+					ChooserToDelete->Rename(*DeletedChooserName.ToString());
+					RootChooser->Modify(true);
+					RootChooser->RemoveNestedChooser(ChooserToDelete);
+					ReplaceReferences(ChooserToDelete, nullptr, RootChooser);
+				}
 			}
 		}
 	}
 
-	void SNestedChooserTree::RenameNestedChooser()
+	void SNestedChooserTree::RenameNestedObject()
 	{
 		TArray<TSharedPtr<FNestedChooserTreeEntry>>SelectedItems = TreeView->GetSelectedItems();
-		if (SelectedItems.Num() == 1)
+		if (SelectedItems.Num() == 1 && SelectedItems[0])
 		{
-			UChooserTable* ChooserToRename = SelectedItems[0]->Chooser;
+			UObject* ObjectToRename = SelectedItems[0]->Object;
 
-			if (ChooserToRename != RootChooser)
+			if (ObjectToRename != RootChooser)
 			{
 				TSharedRef<STextEntryPopup> TextEntry =
 					SNew(STextEntryPopup)
-					.DefaultText(FText::FromString(ChooserToRename->GetName()))
-					.Label(LOCTEXT("RenameNestedChooserLabel", "Rename Chooser"))
-					.OnTextCommitted_Lambda([ChooserToRename, this](FText InText, ETextCommit::Type InCommitType)
+					.DefaultText(FText::FromString(ObjectToRename->GetName()))
+					.Label(LOCTEXT("RenameNestedObjectLabel", "Rename"))
+					.OnTextCommitted_Lambda([ObjectToRename, this](FText InText, ETextCommit::Type InCommitType)
 					{
 						if (InCommitType == ETextCommit::OnEnter)
 						{
@@ -155,18 +162,18 @@ namespace UE::ChooserEditor
 
 							FString NewName = InText.ToString();
 							
-							for(UChooserTable* NestedChooser : RootChooser->NestedChoosers)
+							for(UObject* NestedObject : RootChooser->NestedObjects)
 							{
-								if (NestedChooser->GetName() == NewName)
+								if (NestedObject->GetName() == NewName)
 								{
 									return;
 								}
 							}
 							
-							const FScopedTransaction Transaction(LOCTEXT("Rename Nested Chooser", "Rename Nested Chooser"));
+							const FScopedTransaction Transaction(LOCTEXT("Rename Nested Object", "Rename Nested Object"));
 				
-							ChooserToRename->Modify(true);
-							ChooserToRename->Rename(*NewName);
+							ObjectToRename->Modify(true);
+							ObjectToRename->Rename(*NewName);
 							RefreshAll();
 						}
 					});
@@ -187,7 +194,7 @@ namespace UE::ChooserEditor
 	{
 		ChooserEditor = InArgs._ChooserEditor;
 		RootChooser = ChooserEditor->GetRootChooser();
-		RootChooser->NestedChoosersChanged.AddRaw(this, &SNestedChooserTree::RefreshAll);
+		RootChooser->NestedObjectsChanged.AddRaw(this, &SNestedChooserTree::RefreshAll);
 		TreeEntries.Add(MakeShared<FNestedChooserTreeEntry>(FNestedChooserTreeEntry({RootChooser})));
 
 		TreeView = SNew(STreeView<TSharedPtr<FNestedChooserTreeEntry>>)
@@ -201,7 +208,7 @@ namespace UE::ChooserEditor
 								}
 								else if (Event.GetKey() == EKeys::F2)
 								{
-									RenameNestedChooser();
+									RenameNestedObject();
 									return FReply::Handled();
 								}
 								return FReply::Unhandled();
@@ -227,21 +234,21 @@ namespace UE::ChooserEditor
 	{
 		if (RootChooser)
 		{
-			RootChooser->NestedChoosersChanged.RemoveAll(this);
+			RootChooser->NestedObjectsChanged.RemoveAll(this);
 		}
 	}
 
-	static TSharedPtr<FNestedChooserTreeEntry> MakeEntry(TArray<TSharedPtr<FNestedChooserTreeEntry>>& OldValues, UChooserTable* Chooser)
+	static TSharedPtr<FNestedChooserTreeEntry> MakeEntry(TArray<TSharedPtr<FNestedChooserTreeEntry>>& OldValues, UObject* Object)
 	{
 		for(TSharedPtr<FNestedChooserTreeEntry>& Entry : OldValues)
 		{
-			if (Entry->Chooser == Chooser)
+			if (Entry->Object == Object)
 			{
 				return Entry;
 			}
 		}
 
-		return MakeShared<FNestedChooserTreeEntry>(FNestedChooserTreeEntry({Chooser, true}));
+		return MakeShared<FNestedChooserTreeEntry>(FNestedChooserTreeEntry({Object, true}));
 	}
 
 	void SNestedChooserTree::RefreshAll()
@@ -250,9 +257,9 @@ namespace UE::ChooserEditor
 		AllChoosers.SetNum(0);
 		RootChooser = ChooserEditor->GetRootChooser();
 		AllChoosers.Add(MakeEntry(OldValues, RootChooser));
-		for(UChooserTable* Chooser : RootChooser->NestedChoosers)
+		for(UObject* Object : RootChooser->NestedObjects)
 		{
-			AllChoosers.Add(MakeEntry(OldValues, Chooser));
+			AllChoosers.Add(MakeEntry(OldValues, Object));
 		}
 
 		if (TreeView.IsValid())
@@ -274,9 +281,9 @@ namespace UE::ChooserEditor
 				FExecuteAction::CreateRaw(this, &SNestedChooserTree::DeleteNestedChooser)
 				)
 				);
-		MenuBuilder.AddMenuEntry(LOCTEXT("Rename", "Rename"), LOCTEXT("Rename Tooltip", "Rename Nested Chooser Table"), FSlateIcon(), FUIAction
+		MenuBuilder.AddMenuEntry(LOCTEXT("Rename", "Rename"), LOCTEXT("Rename Tooltip", "Rename Nested Object"), FSlateIcon(), FUIAction
 			(
-				FExecuteAction::CreateRaw(this, &SNestedChooserTree::RenameNestedChooser)
+				FExecuteAction::CreateRaw(this, &SNestedChooserTree::RenameNestedObject)
 				)
 				);
 
diff --git a/Engine/Plugins/Chooser/Source/ChooserEditor/Private/SNestedChooserTree.h b/Engine/Plugins/Chooser/Source/ChooserEditor/Private/SNestedChooserTree.h
--- a/Engine/Plugins/Chooser/Source/ChooserEditor/Private/SNestedChooserTree.h
+++ b/Engine/Plugins/Chooser/Source/ChooserEditor/Private/SNestedChooserTree.h
@@ -13,7 +13,7 @@ class FChooserTableEditor;
 
 struct FNestedChooserTreeEntry
 {
-	UChooserTable* Chooser = nullptr;
+	UObject* Object = nullptr;
 	bool bExpanded = true;
 };
 
@@ -35,7 +35,7 @@ private:
 	void TreeViewGetChildren(TSharedPtr<FNestedChooserTreeEntry> InItem, TArray<TSharedPtr<FNestedChooserTreeEntry>>& OutChildren);
 	void TreeViewDoubleClicked(TSharedPtr<FNestedChooserTreeEntry> SelectedObject);
 	void DeleteNestedChooser();
-	void RenameNestedChooser();
+	void RenameNestedObject();
 	
 	FChooserTableEditor* ChooserEditor = nullptr;
 	UChooserTable* RootChooser = nullptr;
