Fix large drag/drop operations in chooser table editor

diff --git a/Engine/Plugins/Chooser/Source/ChooserEditor/Private/ChooserTableEditor.cpp b/Engine/Plugins/Chooser/Source/ChooserEditor/Private/ChooserTableEditor.cpp
--- a/Engine/Plugins/Chooser/Source/ChooserEditor/Private/ChooserTableEditor.cpp
+++ b/Engine/Plugins/Chooser/Source/ChooserEditor/Private/ChooserTableEditor.cpp
@@ -205,11 +205,16 @@ FChooserTableEditor::~FChooserTableEditor()
 		SelectedColumn->RemoveFromRoot();
 		SelectedColumn = nullptr;
 	}
-	for (UObject* SelectedRow : SelectedRows)
+	for (UObject* SelectedRow : SelectedRowsCache)
 	{
-		SelectedRow->ClearFlags(RF_Standalone);
+		if (SelectedRow)
+		{
+			SelectedRow->ClearFlags(RF_Standalone);
+		}
 	}
 	SelectedRows.Empty();
+
+	FTSTicker::GetCoreTicker().RemoveTicker(TickDelegateHandle);
 	
 	FCoreUObjectDelegates::OnObjectsReplaced.RemoveAll(this);
 	FCoreUObjectDelegates::OnObjectTransacted.RemoveAll(this);
@@ -642,14 +647,25 @@ void FChooserTableEditor::OnObjectsTransacted(UObject* Object, const FTransactio
 					ColumnIndex++;
 				}
 			}
-
-			RefreshAll();
+			QueueRefreshAll();
 		}
 	}
 }
+	
+bool FChooserTableEditor::HandleTicker(float DeltaTime)
+{
+	if (bQueueRefreshAll)
+	{
+		RefreshAll();
+		bQueueRefreshAll = false;
+	}
+	return true;
+}
 
 void FChooserTableEditor::InitEditor( const EToolkitMode::Type Mode, const TSharedPtr< class IToolkitHost >& InitToolkitHost, const TArray<UObject*>& ObjectsToEdit, FGetDetailsViewObjects GetDetailsViewObjects )
 {
+	TickDelegateHandle = FTSTicker::GetCoreTicker().AddTicker(FTickerDelegate::CreateRaw(this, &FChooserTableEditor::HandleTicker), 1.f);
+	
 	UChooserTable* Chooser = Cast<UChooserTable>(ObjectsToEdit[0]);
 	RootChooser = Chooser->GetRootChooser();
 	check(RootChooser);
@@ -768,7 +784,7 @@ void FChooserTableEditor::RefreshAll()
 		ESelectionType CachedSelectionType = CurrentSelectionType;
 		int SelectedColumnIndex = -1;
 		UChooserTable* SelectedChooser = nullptr;
-		TArray<int> CachedSelectedRows;
+		TArray<int32> CachedSelectedRows;
 
 		if (CachedSelectionType == ESelectionType::Column)
 		{
@@ -801,17 +817,18 @@ void FChooserTableEditor::RefreshAll()
 		}
 		else if (CachedSelectionType == ESelectionType::Rows)
 		{
-			ClearSelectedRows();
-			for(int Row : CachedSelectedRows)
-			{
-				SelectRow(Row, false);
-			}
+			SelectRows(CachedSelectedRows);
 		}
 	}
 	
 	RefreshNestedObjectTree();
 }
 
+void FChooserTableEditor::QueueRefreshAll()
+{
+	bQueueRefreshAll = true;
+}
+
 void FChooserTableEditor::RefreshNestedObjectTree()
 {
 	if (NestedChooserTree.IsValid())
@@ -1047,6 +1064,23 @@ void FChooserTableEditor::SelectRow(int32 RowIndex, bool bClear)
 	}
 }
 	
+void FChooserTableEditor::SelectRows(const TArrayView<int32>& Rows)
+{
+	
+	TableView->ClearSelection();
+	TArray<TSharedPtr<FChooserTableRow>> RowsToSelect;
+
+	for (TSharedPtr<FChooserTableRow>& Row : TableRows)
+	{
+		if (Rows.Contains(Row->RowIndex))
+		{
+			RowsToSelect.Add(Row);
+		}
+	}
+	
+	TableView->SetItemSelection(RowsToSelect, true);
+}
+	
 void FChooserTableEditor::ClearSelectedRows() 
 {
 	SelectedRows.SetNum(0);
@@ -1205,10 +1239,6 @@ void FChooserTableEditor::AddColumn(const UScriptStruct* ColumnType)
 
 void FChooserTableEditor::RefreshRowSelectionDetails()
 {
-	for (UObject* SelectedRow : SelectedRows)
-	{
-		SelectedRow->ClearFlags(RF_Standalone);
-	}
 	SelectedRows.SetNum(0);
 	UChooserTable* Chooser = GetChooser();
 
@@ -1218,14 +1248,27 @@ void FChooserTableEditor::RefreshRowSelectionDetails()
 	
 	// Get the list of objects to edit the details of
 	TArray<TSharedPtr<FChooserTableRow>> SelectedItems = TableView->GetSelectedItems();
+
+	if (SelectedRowsCache.Num() < SelectedItems.Num())
+	{
+		SelectedRowsCache.SetNum(SelectedItems.Num());
+	}
+	
+	int32 SelectedRowsCacheIndex = 0;
 	for(TSharedPtr<FChooserTableRow>& SelectedItem : SelectedItems)
 	{
 		if (Chooser->ResultsStructs.IsValidIndex(SelectedItem->RowIndex))
 		{
-			TObjectPtr<UChooserRowDetails> Selection = NewObject<UChooserRowDetails>();
+			if (SelectedRowsCache[SelectedRowsCacheIndex] == nullptr)
+			{
+				SelectedRowsCache[SelectedRowsCacheIndex] = NewObject<UChooserRowDetails>();
+				SelectedRowsCache[SelectedRowsCacheIndex]->SetFlags(RF_Standalone | RF_Transactional);
+			}
+			TObjectPtr<UChooserRowDetails> Selection = SelectedRowsCache[SelectedRowsCacheIndex];
+			SelectedRowsCacheIndex ++;
+			
 			Selection->Chooser = Chooser;
 			Selection->Row = SelectedItem->RowIndex;
-			Selection->SetFlags(RF_Standalone | RF_Transactional);
 			
 			if (Chooser->ResultType != EObjectChooserResultType::NoPrimaryResult)
 			{
@@ -2613,20 +2656,24 @@ void FChooserTableEditor::PasteInternal(UChooserTable* PastedContent, int PasteR
 			}
 		}
 
+		ClearSelectedRows();
 		RefreshAll();
+
+		TArray<int32, TInlineAllocator<256>> SelectedIndices;
+		SelectedIndices.Reserve(RowsToPaste+1);
+		
 		// select the inserted rows
-		ClearSelectedRows();
 		for (int i = 0; i < RowsToPaste; i++)
 		{
-			SelectRow(InsertIndex + i, false);
+			SelectedIndices.Add(InsertIndex + i);
 		}
 		if(PastedContent->FallbackResult.IsValid())
 		{
-			SelectRow(ColumnWidget_SpecialIndex_Fallback, false);
+			SelectedIndices.Add(ColumnWidget_SpecialIndex_Fallback);
 		}
+
+		SelectRows(SelectedIndices);
 	}
-	
-	RefreshAll();
 }
 	
 void FChooserTableEditor::Paste()
diff --git a/Engine/Plugins/Chooser/Source/ChooserEditor/Private/ChooserTableEditor.h b/Engine/Plugins/Chooser/Source/ChooserEditor/Private/ChooserTableEditor.h
--- a/Engine/Plugins/Chooser/Source/ChooserEditor/Private/ChooserTableEditor.h
+++ b/Engine/Plugins/Chooser/Source/ChooserEditor/Private/ChooserTableEditor.h
@@ -42,6 +42,7 @@ namespace UE::ChooserEditor
 	class FChooserTableEditor : public FAssetEditorToolkit, public FSelfRegisteringEditorUndoClient, public FNotifyHook
 	{
 	public:
+		
 		/** Delegate that, given an array of assets, returns an array of objects to use in the details view of an FSimpleAssetEditor */
 		DECLARE_DELEGATE_RetVal_OneParam(TArray<UObject*>, FGetDetailsViewObjects, const TArray<UObject*>&);
 
@@ -94,6 +95,7 @@ namespace UE::ChooserEditor
 		void PushChooserTableToEdit(UChooserTable* Chooser);
 		void PopChooserTableToEdit();
 		void RefreshAll();
+		void QueueRefreshAll();
 		void RefreshNestedObjectTree();
 	
 		/** Used to show or hide certain properties */
@@ -142,6 +144,7 @@ namespace UE::ChooserEditor
 		void MoveRows(int TargetIndex);
 		int MoveRow(int SourceRowIndex, int TargetIndex);
 		int MoveColumn(int SourceIndex, int TargetIndex);
+		void SelectRows(const TArrayView<int32>& Rows);
 		void SelectRow(int32 RowIndex, bool bClear = true);
 		void ClearSelectedRows(); 
 		bool IsRowSelected(int32 RowIndex);
@@ -192,6 +195,13 @@ namespace UE::ChooserEditor
 		/** Details view */
 		TSharedPtr< class IDetailsView > DetailsView;
 
+		
+		bool bQueueRefreshAll = false;
+
+		FTSTicker::FDelegateHandle TickDelegateHandle;
+		
+		bool HandleTicker(float DeltaTime);
+		
 		/** App Identifier. */
 		static const FName ChooserEditorAppIdentifier;
 
@@ -212,6 +222,8 @@ namespace UE::ChooserEditor
 
 		UChooserColumnDetails* SelectedColumn = nullptr;
 		TArray<TObjectPtr<UChooserRowDetails>> SelectedRows;
+		// cache of objects to reuse for SelectedRows
+		TArray<TObjectPtr<UChooserRowDetails>> SelectedRowsCache;
 
 		TSharedPtr<SBreadcrumbTrail<UChooserTable*>> BreadcrumbTrail;
 		TRingBuffer<UChooserTable*> History;
diff --git a/Engine/Plugins/Chooser/Source/ChooserEditor/Private/SChooserTableRow.cpp b/Engine/Plugins/Chooser/Source/ChooserEditor/Private/SChooserTableRow.cpp
--- a/Engine/Plugins/Chooser/Source/ChooserEditor/Private/SChooserTableRow.cpp
+++ b/Engine/Plugins/Chooser/Source/ChooserEditor/Private/SChooserTableRow.cpp
@@ -513,14 +513,17 @@ namespace UE::ChooserEditor
 						Column.InsertRows(InsertRowIndex, NewResults.Num());
 					}
 
+					Editor->ClearSelectedRows();
 					Editor->RefreshAll();
 					
-					Editor->ClearSelectedRows();
+					TArray<int32, TInlineAllocator<256>> RowsToSelect;
+					RowsToSelect.Reserve(NewResults.Num());
 					for(int Index = InsertRowIndex; Index < InsertRowIndex + NewResults.Num(); Index++)
 					{
 						Editor->AutoPopulateRow(Index);
-						Editor->SelectRow(Index, false);
+						RowsToSelect.Add(Index);
 					}
+					Editor->SelectRows(RowsToSelect);
 				}
 			}
 		}
