Chooser - 16 aligning by default the column scratch area

diff --git a/Engine/Plugins/Chooser/Source/Chooser/Private/Chooser.cpp b/Engine/Plugins/Chooser/Source/Chooser/Private/Chooser.cpp
--- a/Engine/Plugins/Chooser/Source/Chooser/Private/Chooser.cpp
+++ b/Engine/Plugins/Chooser/Source/Chooser/Private/Chooser.cpp
@@ -474,8 +474,8 @@ FObjectChooserBase::EIteratorStatus UChooserTable::EvaluateChooser(FChooserEvalu
 		const FChooserColumnBase& Column = ColumnData.Get<FChooserColumnBase>();
 		TotalScratchAreaSize += Column.GetScratchAreaSize();
 	}
-	TArrayView<uint8> ScratchArea((uint8*)FMemory_Alloca(TotalScratchAreaSize * sizeof(uint8)), TotalScratchAreaSize);
-
+	static constexpr int32 ScratchAreaAlignment = 16;
+	TArrayView<uint8> ScratchArea((uint8*)FMemory_Alloca_Aligned(TotalScratchAreaSize * sizeof(uint8), ScratchAreaAlignment), TotalScratchAreaSize);
 	auto DeinitializeScratchAreas = [Chooser, ScratchArea, TotalScratchAreaSize]()
 		{
 			if (TotalScratchAreaSize > 0)
@@ -490,7 +490,7 @@ FObjectChooserBase::EIteratorStatus UChooserTable::EvaluateChooser(FChooserEvalu
 					// @todo: maybe initialize ScratchArea only for column that has filters and are not bDisabled
 					TArrayView<uint8> ColumnScratchArea = ScratchArea.Slice(ScratchAreaStart, ColumnScratchAreaSize);
 					Column.DeinitializeScratchArea(ColumnScratchArea);
-					ScratchAreaStart += ColumnScratchAreaSize;
+					ScratchAreaStart += Align(ColumnScratchAreaSize, ScratchAreaAlignment);
 				}
 			}
 		};
@@ -507,7 +507,7 @@ FObjectChooserBase::EIteratorStatus UChooserTable::EvaluateChooser(FChooserEvalu
 		// @todo: maybe initialize ScratchArea only for column that has filters and are not bDisabled
 		TArrayView<uint8> ColumnScratchArea = ScratchArea.Slice(ScratchAreaStart, ColumnScratchAreaSize);
 		Column.InitializeScratchArea(ColumnScratchArea);
-		ScratchAreaStart += ColumnScratchAreaSize;
+		ScratchAreaStart += Align(ColumnScratchAreaSize, ScratchAreaAlignment);
 
 #if WITH_EDITORONLY_DATA
 		if (Column.bDisabled)
@@ -557,7 +557,7 @@ FObjectChooserBase::EIteratorStatus UChooserTable::EvaluateChooser(FChooserEvalu
 					const int32 ColumnScratchAreaSize = Column.GetScratchAreaSize();
 
 					TArrayView<uint8> ColumnScratchArea = ScratchArea.Slice(ScratchAreaStart, ColumnScratchAreaSize);
-					ScratchAreaStart += ColumnScratchAreaSize;
+					ScratchAreaStart += Align(ColumnScratchAreaSize, ScratchAreaAlignment);
 
 					#if WITH_EDITORONLY_DATA
 					if (Column.bDisabled)
@@ -613,7 +613,7 @@ FObjectChooserBase::EIteratorStatus UChooserTable::EvaluateChooser(FChooserEvalu
 					const int32 ColumnScratchAreaSize = Column.GetScratchAreaSize();
 
 					TArrayView<uint8> ColumnScratchArea = ScratchArea.Slice(ScratchAreaStart, ColumnScratchAreaSize);
-					ScratchAreaStart += ColumnScratchAreaSize;
+					ScratchAreaStart += Align(ColumnScratchAreaSize, ScratchAreaAlignment);
 					Column.SetOutputs(Context, ChooserColumn_SpecialIndex_Fallback, ColumnScratchArea);
 				}
 				check(ScratchAreaStart == TotalScratchAreaSize);
