Support function calls for struct bindings in chooser, and use that for PoseSearchHistory binding in Pose Search column

diff --git a/Engine/Plugins/Animation/PoseSearch/Source/Runtime/Private/Chooser/PoseSearchChooserColumn.cpp b/Engine/Plugins/Animation/PoseSearch/Source/Runtime/Private/Chooser/PoseSearchChooserColumn.cpp
--- a/Engine/Plugins/Animation/PoseSearch/Source/Runtime/Private/Chooser/PoseSearchChooserColumn.cpp
+++ b/Engine/Plugins/Animation/PoseSearch/Source/Runtime/Private/Chooser/PoseSearchChooserColumn.cpp
@@ -44,13 +44,7 @@ namespace UE::PoseSearch
 
 bool FPoseHistoryContextProperty::GetValue(FChooserEvaluationContext& Context, FPoseHistoryReference& OutResult) const
 {
-	const FPoseHistoryReference* Reference;
-	if (Binding.GetValuePtr(Context, Reference))
-	{
-		OutResult = *Reference;
-		return true;
-	}
-	return false;
+	return Binding.GetStructValue(Context, OutResult);
 }
 
 FPoseSearchColumn::FPoseSearchColumn()
diff --git a/Engine/Plugins/Animation/PoseSearch/Source/Runtime/Public/PoseSearch/Chooser/PoseSearchChooserColumn.h b/Engine/Plugins/Animation/PoseSearch/Source/Runtime/Public/PoseSearch/Chooser/PoseSearchChooserColumn.h
--- a/Engine/Plugins/Animation/PoseSearch/Source/Runtime/Public/PoseSearch/Chooser/PoseSearchChooserColumn.h
+++ b/Engine/Plugins/Animation/PoseSearch/Source/Runtime/Public/PoseSearch/Chooser/PoseSearchChooserColumn.h
@@ -34,7 +34,7 @@ struct FPoseHistoryContextProperty :  public FChooserParameterPoseHistoryBase
 {
 	GENERATED_BODY()
 
-	UPROPERTY(EditAnywhere, Meta = (BindingType = "FPoseHistoryReference", BindingAllowFunctions = "false", BindingColor = "StructPinTypeColor"), Category = "Binding")
+	UPROPERTY(EditAnywhere, Meta = (BindingType = "FPoseHistoryReference", BindingAllowFunctions = "true", BindingColor = "StructPinTypeColor"), Category = "Binding")
 	FChooserPropertyBinding Binding;
 	
 	UE_API virtual bool GetValue(FChooserEvaluationContext& Context, FPoseHistoryReference& OutResult) const override;
diff --git a/Engine/Plugins/Chooser/Source/Chooser/Public/ChooserPropertyAccess.h b/Engine/Plugins/Chooser/Source/Chooser/Public/ChooserPropertyAccess.h
--- a/Engine/Plugins/Chooser/Source/Chooser/Public/ChooserPropertyAccess.h
+++ b/Engine/Plugins/Chooser/Source/Chooser/Public/ChooserPropertyAccess.h
@@ -115,6 +115,9 @@ struct FChooserPropertyBinding
 	template <typename T>
 	bool GetStructPtr(FChooserEvaluationContext& Context, T*& Value, UStruct const* &StructType) const;
 	
+	template <typename T>
+	bool GetStructValue(FChooserEvaluationContext& Context, T& Value) const;
+	
 	template <typename T>
 	bool GetValue(FChooserEvaluationContext& Context, T& Value) const;
 
@@ -260,6 +263,27 @@ bool FChooserPropertyBinding::GetStructPtr(FChooserEvaluationContext& Context, T
 	return false;
 }
 
+template <typename T>
+bool FChooserPropertyBinding::GetStructValue(FChooserEvaluationContext& Context, T& OutResult) const
+{
+	UE::Chooser::FResolvedPropertyChainResult Result;
+	if (ResolvePropertyChain(Context, *this, Result))
+	{
+		if (Result.Function)
+		{
+			UObject* Object = reinterpret_cast<UObject*>(Result.Container);
+			Object->ProcessEvent(Result.Function, &OutResult);
+			return true;
+		}
+		else
+		{
+			OutResult = *reinterpret_cast<T*>(Result.Container + Result.PropertyOffset);
+			return true;
+		}
+	}
+	
+	return false;
+	}
 
 template <typename T>
 bool FChooserPropertyBinding::GetValue(FChooserEvaluationContext& Context, T& OutResult) const
