Support for accessing any context member from PoseSearchFeatureChannel_Distance blueprint subclasses, rather than only the AnimInstance

diff --git a/Engine/Plugins/Animation/PoseSearch/Source/Runtime/Private/PoseSearchFeatureChannel_Distance.cpp b/Engine/Plugins/Animation/PoseSearch/Source/Runtime/Private/PoseSearchFeatureChannel_Distance.cpp
--- a/Engine/Plugins/Animation/PoseSearch/Source/Runtime/Private/PoseSearchFeatureChannel_Distance.cpp
+++ b/Engine/Plugins/Animation/PoseSearch/Source/Runtime/Private/PoseSearchFeatureChannel_Distance.cpp
@@ -72,15 +72,10 @@ void UPoseSearchFeatureChannel_Distance::BuildQuery(UE::PoseSearch::FSearchConte
 
 	if (bUseBlueprintQueryOverride)
 	{
-		if (const UAnimInstance* AnimInstance = Cast<UAnimInstance>(SearchContext.GetContext(SampleRole)->GetFirstObjectParam()))
+		if (FChooserEvaluationContext* ChooserContext = SearchContext.GetContext(SampleRole))
 		{
-			const float Distance = BP_GetDistance(AnimInstance);
-  			FFeatureVectorHelper::EncodeFloat(SearchContext.EditFeatureVector(), ChannelDataOffset, Distance);
-		}
-		else
-		{
-			// @todo: support non UAnimInstance anim contexts for AnimNext
-			UE_LOG(LogPoseSearch, Warning, TEXT("UPoseSearchFeatureChannel_Distance::BuildQuery - unsupported null UAnimInstance: WIP support for AnimNext!"));
+			const float Distance = BP_GetDistance(*ChooserContext);
+			FFeatureVectorHelper::EncodeFloat(SearchContext.EditFeatureVector(), ChannelDataOffset, Distance);
 		}
 		return;
 	}
diff --git a/Engine/Plugins/Animation/PoseSearch/Source/Runtime/Public/PoseSearch/PoseSearchFeatureChannel_Distance.h b/Engine/Plugins/Animation/PoseSearch/Source/Runtime/Public/PoseSearch/PoseSearchFeatureChannel_Distance.h
--- a/Engine/Plugins/Animation/PoseSearch/Source/Runtime/Public/PoseSearch/PoseSearchFeatureChannel_Distance.h
+++ b/Engine/Plugins/Animation/PoseSearch/Source/Runtime/Public/PoseSearch/PoseSearchFeatureChannel_Distance.h
@@ -4,6 +4,7 @@
 
 #include "PoseSearch/PoseSearchFeatureChannel.h"
 #include "BoneContainer.h"
+#include "IObjectChooser.h"
 #include "PoseSearchFeatureChannel_Distance.generated.h"
 
 #define UE_API POSESEARCH_API
@@ -88,7 +89,7 @@ public:
 #endif //WITH_EDITORONLY_DATA
 
 	UFUNCTION(BlueprintPure, BlueprintImplementableEvent, meta=(BlueprintThreadSafe, DisplayName = "Get Distance"), Category = "Settings")
-	UE_API float BP_GetDistance(const UAnimInstance* AnimInstance) const;
+	UE_API float BP_GetDistance(UPARAM(Ref) FChooserEvaluationContext& Context) const;
 
 	bool bUseBlueprintQueryOverride = false;
 
diff --git a/Engine/Plugins/Chooser/Source/Chooser/Internal/Chooser.h b/Engine/Plugins/Chooser/Source/Chooser/Internal/Chooser.h
--- a/Engine/Plugins/Chooser/Source/Chooser/Internal/Chooser.h
+++ b/Engine/Plugins/Chooser/Source/Chooser/Internal/Chooser.h
@@ -8,6 +8,7 @@
 #include "StructUtils/InstancedStruct.h"
 #include "IChooserColumn.h"
 #include "IHasContext.h"
+#include "ChooserSignature.h"
 #if WITH_EDITOR
 #include "Kismet2/StructureEditorUtils.h"
 #endif
@@ -18,7 +19,7 @@
 * Data table used to choose an asset based on input parameters
 */
 UCLASS(MinimalAPI, BlueprintType)
-class UChooserTable : public UObject, public IHasContextClass
+class UChooserTable : public UChooserSignature
 {
 	GENERATED_UCLASS_BODY()
 public:
@@ -142,7 +143,10 @@ public:
 	// FallbackResult will be used as the Result if there are no rows in the chooser which pass all filters.  If FallbackResult is not assigned, then the Chooser will return null in that case.
 	UPROPERTY()
 	FInstancedStruct FallbackResult;
-
+	
+	// IHasContextClass implementation
+	virtual TConstArrayView<FInstancedStruct> GetContextData() const override { return GetRootChooser()->ContextData; }
+	
 #if WITH_EDITORONLY_DATA
 	// Each possible Result (Rows of chooser table)
 	UPROPERTY()
@@ -186,22 +190,6 @@ public:
 	UPROPERTY(EditAnywhere, NoClear, DisplayName = "Columns", Meta = (ExcludeBaseStruct, BaseStruct = "/Script/Chooser.ChooserColumnBase"), Category = "Hidden")
 	TArray<FInstancedStruct> ColumnsStructs;
 	
-	// The kind of output this chooser has (Object or Class or No primary result)
-    UPROPERTY(EditAnywhere, DisplayName = "Result Type", Category="Result", Meta = (EditConditionHides, EditCondition="ResultType != EObjectChooserResultType::NoPrimaryResult"))
-    EObjectChooserResultType ResultType = EObjectChooserResultType::ObjectResult;
-
-	// The Class of Object this Chooser returns when ResultType is set to ObjectOfType, or the Parent Class of the Classes returned by this chooser when ResultType is set to ClassOfType
-	UPROPERTY(EditAnywhere, DisplayName= "Result Class", Category="Result", Meta = (AllowAbstract=true, EditConditionHides, EditCondition="ResultType != EObjectChooserResultType::NoPrimaryResult"))
-	TObjectPtr<UClass> OutputObjectType;
-	
-	// Parameter Objects or Structs from which the chooser can read or write properties 
-	UPROPERTY(EditAnywhere, NoClear, DisplayName = "Parameters", Meta = (ExcludeBaseStruct, BaseStruct = "/Script/Chooser.ContextObjectTypeBase"), Category = "Parameters")
-	TArray<FInstancedStruct> ContextData;
-
-	virtual TConstArrayView<FInstancedStruct> GetContextData() const override { return GetRootChooser()->ContextData; }
-	virtual FString GetContextOwnerName() const override { return GetName(); }
-	virtual UObject* GetContextOwnerAsset() override { return this; }
-
 	bool IsCookedData() const { return !CookedResults.IsEmpty(); }
 	static CHOOSER_API FObjectChooserBase::EIteratorStatus EvaluateChooser(FChooserEvaluationContext& Context, const UChooserTable* Chooser, FObjectChooserBase::FObjectChooserIteratorCallback Callback);
 	static CHOOSER_API FObjectChooserBase::EIteratorStatus EvaluateChooser(FChooserEvaluationContext& Context, const UChooserTable* Chooser, FObjectChooserBase::FObjectChooserSoftObjectIteratorCallback Callback);
diff --git a/Engine/Plugins/Chooser/Source/Chooser/Private/ChooserFunctionLibrary.cpp b/Engine/Plugins/Chooser/Source/Chooser/Private/ChooserFunctionLibrary.cpp
--- a/Engine/Plugins/Chooser/Source/Chooser/Private/ChooserFunctionLibrary.cpp
+++ b/Engine/Plugins/Chooser/Source/Chooser/Private/ChooserFunctionLibrary.cpp
@@ -38,6 +38,18 @@ void UChooserFunctionLibrary::AddChooserObjectInput(FChooserEvaluationContext& C
 	Context.AddObjectParam(Object);
 }
 
+UObject* UChooserFunctionLibrary::GetChooserObjectInput(FChooserEvaluationContext& Context, int32 Index, TSubclassOf<UObject> ObjectClass)
+{
+	if (Index < Context.Params.Num())
+	{
+		if (FChooserEvaluationInputObject* ObjectParam = Context.Params[Index].GetPtr<FChooserEvaluationInputObject>())
+		{
+			return ObjectParam->Object;
+		}
+	}
+	return nullptr;
+}
+
 void UChooserFunctionLibrary::AddChooserStructInput(FChooserEvaluationContext& Inputs, int32 Value)
 {
 	checkNoEntry();
diff --git a/Engine/Plugins/Chooser/Source/Chooser/Private/ChooserSignature.cpp b/Engine/Plugins/Chooser/Source/Chooser/Private/ChooserSignature.cpp
new file mode 100644
--- /dev/null
+++ b/Engine/Plugins/Chooser/Source/Chooser/Private/ChooserSignature.cpp
@@ -0,0 +1,8 @@
+// Copyright Epic Games, Inc. All Rights Reserved.
+#include "ChooserSignature.h"
+
+UChooserSignature::UChooserSignature(const FObjectInitializer& Initializer)
+	: Super(Initializer)
+{
+
+}
diff --git a/Engine/Plugins/Chooser/Source/Chooser/Public/ChooserFunctionLibrary.h b/Engine/Plugins/Chooser/Source/Chooser/Public/ChooserFunctionLibrary.h
--- a/Engine/Plugins/Chooser/Source/Chooser/Public/ChooserFunctionLibrary.h
+++ b/Engine/Plugins/Chooser/Source/Chooser/Public/ChooserFunctionLibrary.h
@@ -78,6 +78,15 @@ public:
 	UFUNCTION(BlueprintCallable, meta = (BlueprintThreadSafe, BlueprintInternalUseOnly = "true"))
 	static UE_API void AddChooserObjectInput(UPARAM(Ref) FChooserEvaluationContext& Context, UObject* Object);
 	
+	/**
+	* Get an Object from a ChooserEvaluation context
+	*
+	* @param Context			(in) A struct reference to the chooser evaluation context
+	* @param Index				(in) the context index to get the object from
+	*/
+	UFUNCTION(BlueprintPure, meta = (BlueprintThreadSafe ,BlueprintInternalUseOnly = "true", DeterminesOutputType = "ObjectClass"))
+	static UObject* GetChooserObjectInput(UPARAM(Ref) FChooserEvaluationContext& Context, int32 Index, TSubclassOf<UObject> ObjectClass);
+	
 	/**
 	* Add a Struct to a ChooserEvaluation context
 	*
diff --git a/Engine/Plugins/Chooser/Source/Chooser/Public/ChooserSignature.h b/Engine/Plugins/Chooser/Source/Chooser/Public/ChooserSignature.h
new file mode 100644
--- /dev/null
+++ b/Engine/Plugins/Chooser/Source/Chooser/Public/ChooserSignature.h
@@ -0,0 +1,40 @@
+// Copyright Epic Games, Inc. All Rights Reserved.
+#pragma once
+
+#include "AssetRegistry/AssetData.h"
+#include "CoreMinimal.h"
+#include "UObject/Object.h"
+#include "IObjectChooser.h"
+#include "StructUtils/InstancedStruct.h"
+#include "IHasContext.h"
+#include "ChooserSignature.generated.h"
+
+/**
+* Data table used to choose an asset based on input parameters
+*/
+UCLASS(MinimalAPI, BlueprintType)
+class UChooserSignature : public UObject, public IHasContextClass
+{
+	GENERATED_UCLASS_BODY()
+public:
+	UChooserSignature() {}
+	
+	// The kind of output this chooser has (Object or Class or No primary result)
+    UPROPERTY(EditAnywhere, DisplayName = "Result Type", Category="Result", Meta = (EditConditionHides, EditCondition="ResultType != EObjectChooserResultType::NoPrimaryResult"))
+    EObjectChooserResultType ResultType = EObjectChooserResultType::ObjectResult;
+
+	// The Class of Object this Chooser returns when ResultType is set to ObjectOfType, or the Parent Class of the Classes returned by this chooser when ResultType is set to ClassOfType
+	UPROPERTY(EditAnywhere, DisplayName= "Result Class", Category="Result", Meta = (AllowAbstract=true, EditConditionHides, EditCondition="ResultType != EObjectChooserResultType::NoPrimaryResult"))
+	TObjectPtr<UClass> OutputObjectType;
+	
+	// Parameter Objects or Structs from which the chooser can read or write properties 
+	UPROPERTY(EditAnywhere, NoClear, DisplayName = "Parameters", Meta = (ExcludeBaseStruct, BaseStruct = "/Script/Chooser.ContextObjectTypeBase"), Category = "Parameters")
+	TArray<FInstancedStruct> ContextData;
+	
+	// IHasContextClass implementation
+	virtual TConstArrayView<FInstancedStruct> GetContextData() const override { return ContextData; }
+	virtual FString GetContextOwnerName() const override { return GetName(); }
+	virtual UObject* GetContextOwnerAsset() override { return this; }
+};
+
+
diff --git a/Engine/Plugins/Chooser/Source/ChooserUncooked/Private/GetChooserContextParametersNode.cpp b/Engine/Plugins/Chooser/Source/ChooserUncooked/Private/GetChooserContextParametersNode.cpp
new file mode 100644
--- /dev/null
+++ b/Engine/Plugins/Chooser/Source/ChooserUncooked/Private/GetChooserContextParametersNode.cpp
@@ -0,0 +1,383 @@
+// Copyright Epic Games, Inc. All Rights Reserved.
+
+
+#include "GetChooserContextParametersNode.h"
+
+#include "BlueprintActionDatabaseRegistrar.h"
+#include "BlueprintNodeSpawner.h"
+#include "ChooserFunctionLibrary.h"
+#include "ChooserPropertyAccess.h"
+#include "Containers/UnrealString.h"
+#include "EdGraph/EdGraph.h"
+#include "EdGraph/EdGraphSchema.h"
+#include "EdGraphSchema_K2.h"
+#include "EditorCategoryUtils.h"
+#include "IHasContext.h"
+#include "Internationalization/Internationalization.h"
+#include "K2Node_CallFunction.h"
+#include "K2Node_MakeStruct.h"
+#include "Kismet2/BlueprintEditorUtils.h"
+#include "Kismet2/CompilerResultsLog.h"
+#include "Kismet/KismetSystemLibrary.h"
+#include "KismetCompiler.h"
+#include "Misc/AssertionMacros.h"
+#include "Misc/CString.h"
+#include "Templates/Casts.h"
+#include "UObject/Class.h"
+#include "UObject/ObjectPtr.h"
+#include "UObject/UObjectGlobals.h"
+#include "UObject/UnrealType.h"
+#include "UObject/WeakObjectPtrTemplates.h"
+#include "K2Node_Self.h"
+
+#include UE_INLINE_GENERATED_CPP_BY_NAME(GetChooserContextParametersNode)
+
+#define LOCTEXT_NAMESPACE "GetChooserContextParametersNode"
+
+
+//----------------------------------------------------------------------------------------------
+// UK2Node_GetChooserContextParameters
+// New Implementation of EvaluateChooser with support for passing in/out multiple objects and structs
+
+UK2Node_GetChooserContextParameters::UK2Node_GetChooserContextParameters(const FObjectInitializer& ObjectInitializer)
+	: Super(ObjectInitializer)
+{
+	NodeTooltip = LOCTEXT("GetChooserParametersNodeTooltip", "Returns all chooser parameters from a ChooserEvaluationContext struct reference.");
+}
+
+void UK2Node_GetChooserContextParameters::UnregisterCallback()
+{
+	if(CurrentCallbackObject)
+	{
+		CurrentCallbackObject->OnContextClassChanged.RemoveAll(this);
+		CurrentCallbackObject = nullptr;
+	}
+}
+
+void UK2Node_GetChooserContextParameters::BeginDestroy()
+{
+	UnregisterCallback();
+	Super::BeginDestroy();
+}
+
+void UK2Node_GetChooserContextParameters::PostEditUndo()
+{
+	Super::PostEditUndo();
+	ChooserSignatureChanged();
+}
+
+void UK2Node_GetChooserContextParameters::DestroyNode()
+{
+	UnregisterCallback();
+	Super::DestroyNode();
+}
+
+
+void UK2Node_GetChooserContextParameters::ChooserSignatureChanged()
+{
+	if (ChooserSignature != CurrentCallbackObject)
+	{
+		UnregisterCallback();
+		
+		if (ChooserSignature)
+		{
+			ChooserSignature->OnContextClassChanged.AddUObject(this, &UK2Node::ReconstructNode);
+		}
+	
+		CurrentCallbackObject = ChooserSignature;
+
+		ReconstructNode();
+	}
+}
+
+void UK2Node_GetChooserContextParameters::PreloadRequiredAssets()
+{
+	if (ChooserSignature)
+	{
+		ChooserSignature->ConditionalPreload();
+	}
+    		
+	Super::PreloadRequiredAssets();
+}
+
+void UK2Node_GetChooserContextParameters::AllocateDefaultPins()
+{
+	Super::AllocateDefaultPins();
+
+	FName ContextPinName("Context");
+	if (!FindPin(ContextPinName, EGPD_Input))
+	{
+		// Input -  Context Pin
+		FCreatePinParams StructContextPinCreationParams;
+		StructContextPinCreationParams.bIsReference = true;
+		CreatePin(EGPD_Input, UEdGraphSchema_K2::PC_Struct, FChooserEvaluationContext::StaticStruct(), ContextPinName, StructContextPinCreationParams);
+	}
+	
+	if (ChooserSignature)
+	{
+		FCreatePinParams StructContextPinCreationParams;
+		StructContextPinCreationParams.bIsReference = true;
+		
+		for (FInstancedStruct& ContextDataEntry : ChooserSignature->ContextData)
+		{
+			if (ContextDataEntry.IsValid())
+			{
+				const UScriptStruct* EntryType = ContextDataEntry.GetScriptStruct();
+				if (EntryType == FContextObjectTypeClass::StaticStruct())
+				{
+					const FContextObjectTypeClass& ClassContext = ContextDataEntry.Get<FContextObjectTypeClass>();
+					if (UEdGraphPin* Pin = FindPin(ClassContext.Class.GetFName(), EGPD_Output))
+					{
+						Pin->PinType.PinSubCategoryObject = ClassContext.Class;
+					}
+					else
+					{
+						CreatePin(EGPD_Output, UEdGraphSchema_K2::PC_Object, ClassContext.Class, ClassContext.Class.GetFName());
+					}
+				}
+				else if (EntryType == FContextObjectTypeStruct::StaticStruct())
+				{
+					const FContextObjectTypeStruct& StructContext = ContextDataEntry.Get<FContextObjectTypeStruct>();
+
+					if (UEdGraphPin* Pin = FindPin(StructContext.Struct.GetFName(), EGPD_Output))
+					{
+						Pin->PinType.PinSubCategoryObject = StructContext.Struct;
+					}
+					else
+					{
+						CreatePin(EGPD_Output, UEdGraphSchema_K2::PC_Struct, StructContext.Struct, StructContext.Struct.GetFName(), StructContextPinCreationParams);
+					}
+				}
+			}
+		}
+	}
+}
+
+FText UK2Node_GetChooserContextParameters::GetNodeTitle(ENodeTitleType::Type TitleType) const
+{
+	return LOCTEXT("GetChooserContextParametersTitle", "Get Chooser Context Parameters");
+}
+
+FText UK2Node_GetChooserContextParameters::GetPinDisplayName(const UEdGraphPin* Pin) const
+{
+	if (Pin->PinName == UEdGraphSchema_K2::PN_Execute)
+	{
+		return FText();
+	}
+	return FText::FromName(Pin->PinName);
+}
+
+void UK2Node_GetChooserContextParameters::PostEditChangeProperty(struct FPropertyChangedEvent& PropertyChangedEvent)
+{
+	if (PropertyChangedEvent.Property->GetName() == "ChooserSignature")
+	{
+		ChooserSignatureChanged();
+	}
+	Super::PostEditChangeProperty(PropertyChangedEvent);
+}
+
+void UK2Node_GetChooserContextParameters::PostLoad()
+{
+	Super::PostLoad();
+	
+	if (ChooserSignature)
+	{
+		ChooserSignature->OnContextClassChanged.AddUObject(this, &UK2Node::ReconstructNode);
+	}
+
+	CurrentCallbackObject = ChooserSignature;
+}
+
+void UK2Node_GetChooserContextParameters::PinConnectionListChanged(UEdGraphPin* Pin)
+{
+	Modify();
+
+	FBlueprintEditorUtils::MarkBlueprintAsStructurallyModified(GetBlueprint());
+}
+
+void UK2Node_GetChooserContextParameters::PinDefaultValueChanged(UEdGraphPin* Pin)
+{
+	Super::PinDefaultValueChanged(Pin);
+}
+
+void UK2Node_GetChooserContextParameters::PinTypeChanged(UEdGraphPin* Pin)
+{
+	Super::PinTypeChanged(Pin);
+}
+
+FText UK2Node_GetChooserContextParameters::GetTooltipText() const
+{
+	return NodeTooltip;
+}
+
+void UK2Node_GetChooserContextParameters::PostReconstructNode()
+{
+	Super::PostReconstructNode();
+}
+
+void UK2Node_GetChooserContextParameters::ExpandNode(class FKismetCompilerContext& CompilerContext, UEdGraph* SourceGraph)
+{
+	Super::ExpandNode(CompilerContext, SourceGraph);
+
+	UEdGraphPin* ContextPin = FindPin(TEXT("Context"));
+	check(ContextPin);
+	
+	if (ChooserSignature)
+	{
+		int ContextDataCount = ChooserSignature->ContextData.Num();
+		for(int ContextDataIndex = 0; ContextDataIndex < ContextDataCount; ContextDataIndex++)
+		{
+				FInstancedStruct& ContextDataEntry = ChooserSignature->ContextData[ContextDataIndex];
+				if (ContextDataEntry.IsValid())
+				{
+					const UScriptStruct* EntryType = ContextDataEntry.GetScriptStruct();
+					if (EntryType == FContextObjectTypeClass::StaticStruct())
+					{
+						const FContextObjectTypeClass& ClassContext = ContextDataEntry.Get<FContextObjectTypeClass>();
+						if (ClassContext.Class)
+						{
+							if (UEdGraphPin* Pin = FindPin(ClassContext.Class.GetFName(), EGPD_Output))
+							{
+								if (Pin->HasAnyConnections())
+								{
+									// setup class output pin
+									
+									UK2Node_CallFunction* GetParameterFunction = CompilerContext.SpawnIntermediateNode<UK2Node_CallFunction>(this, SourceGraph);
+									CompilerContext.MessageLog.NotifyIntermediateObjectCreation(GetParameterFunction, this);
+									
+									GetParameterFunction->SetFromFunction(UChooserFunctionLibrary::StaticClass()->FindFunctionByName(GET_MEMBER_NAME_CHECKED(UChooserFunctionLibrary, GetChooserObjectInput)));
+									GetParameterFunction->AllocateDefaultPins();
+									GetParameterFunction->FindPin(FName("Index"))->DefaultValue = FString::FromInt(ContextDataIndex);
+	 									
+									UEdGraphPin* ParameterContextPin = GetParameterFunction->FindPin(FName("Context"));
+									CompilerContext.CopyPinLinksToIntermediate(*ContextPin, *ParameterContextPin);
+									
+									if (UEdGraphPin* OutputClassPin = GetParameterFunction->FindPin(TEXT("ObjectClass")))
+									{
+										// this makes the output pin of the function become this class type, avoiding an extra cast node
+										GetParameterFunction->GetSchema()->TrySetDefaultObject(*OutputClassPin, ClassContext.Class);
+									}
+	 									
+									UEdGraphPin* ValuePin = GetParameterFunction->GetReturnValuePin();
+									
+									CompilerContext.MovePinLinksToIntermediate(*Pin, *ValuePin);
+								}
+							}
+						}
+					}
+	 				else if (EntryType == FContextObjectTypeStruct::StaticStruct())
+	 				{
+	 					const FContextObjectTypeStruct& StructContext = ContextDataEntry.Get<FContextObjectTypeStruct>();
+	 					if (StructContext.Struct)
+	 					{
+	 							if (UEdGraphPin* Pin = FindPin(StructContext.Struct.GetFName(), EGPD_Output))
+	 							{
+	 								if (Pin->HasAnyConnections())
+	 								{
+	 									// set up struct output pin
+	 									
+	 									UK2Node_CallFunction* GetStructFunction = CompilerContext.SpawnIntermediateNode<UK2Node_CallFunction>(this, SourceGraph);
+	 									CompilerContext.MessageLog.NotifyIntermediateObjectCreation(GetStructFunction, this);
+	 									// GetChooserStructOutput function also works for inputs (this node is outputting all members of the context input and output)
+	 									GetStructFunction->SetFromFunction(UChooserFunctionLibrary::StaticClass()->FindFunctionByName(GET_MEMBER_NAME_CHECKED(UChooserFunctionLibrary, GetChooserStructOutput)));
+	 									GetStructFunction->AllocateDefaultPins();
+	 									GetStructFunction->FindPin(FName("Index"))->DefaultValue = FString::FromInt(ContextDataIndex);
+	 									
+										UEdGraphPin* StructContextPin = GetStructFunction->FindPin(FName("Context"));
+	 									CompilerContext.CopyPinLinksToIntermediate(*ContextPin, *StructContextPin);
+	 									
+	 									UEdGraphPin* ValuePin = GetStructFunction->FindPin(FName("Value"));
+	 									// not sure why this isn't automatically happening:
+	 									ValuePin->PinType.PinCategory = UEdGraphSchema_K2::PC_Struct;
+	 									ValuePin->PinType.PinSubCategoryObject = StructContext.Struct;
+	 					
+	 									CompilerContext.MovePinLinksToIntermediate(*Pin, *ValuePin);
+	 								}
+	 							}
+	 					}
+	 				}
+	 			}
+		}
+		
+	}
+	
+	BreakAllNodeLinks();
+}
+
+UK2Node::ERedirectType UK2Node_GetChooserContextParameters::DoPinsMatchForReconstruction(const UEdGraphPin* NewPin, int32 NewPinIndex, const UEdGraphPin* OldPin, int32 OldPinIndex) const
+{
+	ERedirectType RedirectType = ERedirectType_None;
+
+	// if the pin names do match
+	if (NewPin->PinName.ToString().Equals(OldPin->PinName.ToString(), ESearchCase::CaseSensitive))
+	{
+		// Make sure we're not dealing with a menu node
+		UEdGraph* OuterGraph = GetGraph();
+		if( OuterGraph && OuterGraph->Schema )
+		{
+			const UEdGraphSchema_K2* K2Schema = Cast<const UEdGraphSchema_K2>(GetSchema());
+			if( !K2Schema || K2Schema->IsSelfPin(*NewPin) || K2Schema->ArePinTypesCompatible(OldPin->PinType, NewPin->PinType) )
+			{
+				RedirectType = ERedirectType_Name;
+			}
+			else
+			{
+				RedirectType = ERedirectType_None;
+			}
+		}
+	}
+	else
+	{
+		// try looking for a redirect if it's a K2 node
+		if (UK2Node* Node = Cast<UK2Node>(NewPin->GetOwningNode()))
+		{	
+			// if you don't have matching pin, now check if there is any redirect param set
+			TArray<FString> OldPinNames;
+			GetRedirectPinNames(*OldPin, OldPinNames);
+
+			FName NewPinName;
+			RedirectType = ShouldRedirectParam(OldPinNames, /*out*/ NewPinName, Node);
+
+			// make sure they match
+			if ((RedirectType != ERedirectType_None) && (!NewPin->PinName.ToString().Equals(NewPinName.ToString(), ESearchCase::CaseSensitive)))
+			{
+				RedirectType = ERedirectType_None;
+			}
+		}
+	}
+
+	return RedirectType;
+}
+
+bool UK2Node_GetChooserContextParameters::IsConnectionDisallowed(const UEdGraphPin* MyPin, const UEdGraphPin* OtherPin, FString& OutReason) const
+{
+	return Super::IsConnectionDisallowed(MyPin, OtherPin, OutReason);
+}
+
+void UK2Node_GetChooserContextParameters::GetMenuActions(FBlueprintActionDatabaseRegistrar& ActionRegistrar) const
+{
+	// actions get registered under specific object-keys; the idea is that 
+	// actions might have to be updated (or deleted) if their object-key is  
+	// mutated (or removed)... here we use the node's class (so if the node 
+	// type disappears, then the action should go with it)
+	UClass* ActionKey = GetClass();
+	// to keep from needlessly instantiating a UBlueprintNodeSpawner, first   
+	// check to make sure that the registrar is looking for actions of this type
+	// (could be regenerating actions for a specific asset, and therefore the 
+	// registrar would only accept actions corresponding to that asset)
+	if (ActionRegistrar.IsOpenForRegistration(ActionKey))
+	{
+		UBlueprintNodeSpawner* NodeSpawner = UBlueprintNodeSpawner::Create(GetClass());
+		check(NodeSpawner != nullptr);
+
+		ActionRegistrar.AddBlueprintAction(ActionKey, NodeSpawner);
+	}
+}
+
+FText UK2Node_GetChooserContextParameters::GetMenuCategory() const
+{
+	return FEditorCategoryUtils::GetCommonCategory(FCommonEditorCategory::Animation);
+}
+
+
+#undef LOCTEXT_NAMESPACE
diff --git a/Engine/Plugins/Chooser/Source/ChooserUncooked/Private/GetChooserContextParametersNode.h b/Engine/Plugins/Chooser/Source/ChooserUncooked/Private/GetChooserContextParametersNode.h
new file mode 100644
--- /dev/null
+++ b/Engine/Plugins/Chooser/Source/ChooserUncooked/Private/GetChooserContextParametersNode.h
@@ -0,0 +1,69 @@
+// Copyright Epic Games, Inc. All Rights Reserved.
+
+
+#pragma once
+
+#include "CoreMinimal.h"
+#include "ChooserSignature.h"
+#include "EdGraph/EdGraphNode.h"
+#include "HAL/Platform.h"
+#include "Internationalization/Text.h"
+#include "K2Node.h"
+#include "UObject/ObjectMacros.h"
+#include "UObject/UObjectGlobals.h"
+
+#include "GetChooserContextParametersNode.generated.h"
+
+
+//----------------------------------------------------------------------------------------------
+// UK2Node_GetChooserContextParameters
+// exposes pins from a chooser context, based on a Chooser Signature
+UCLASS(MinimalAPI)
+class UK2Node_GetChooserContextParameters : public UK2Node
+{
+	GENERATED_UCLASS_BODY()
+	
+	//~ Begin UObject Interface
+	virtual void BeginDestroy() override;
+	virtual void PostEditUndo() override;
+	virtual void PostEditChangeProperty(struct FPropertyChangedEvent& PropertyChangedEvent) override;
+	virtual void PostLoad() override;
+	//~ End UObject Interface
+
+	//~ Begin UEdGraphNode Interface.
+	virtual void AllocateDefaultPins() override;
+	virtual FText GetNodeTitle(ENodeTitleType::Type TitleType) const override;
+	virtual bool ShouldShowNodeProperties() const override { return true; }
+	virtual void PinConnectionListChanged(UEdGraphPin* Pin) override;
+	virtual void PinDefaultValueChanged(UEdGraphPin* Pin) override;
+	virtual void PinTypeChanged(UEdGraphPin* Pin) override;
+	virtual FText GetTooltipText() const override;
+	virtual FText GetPinDisplayName(const UEdGraphPin* Pin) const override;
+	//~ End UEdGraphNode Interface.
+
+	//~ Begin UK2Node Interface.
+	virtual bool IsNodePure() const override { return true; }
+	virtual bool NodeCausesStructuralBlueprintChange() const override { return true; }
+	virtual void PostReconstructNode() override;
+	virtual void ExpandNode(class FKismetCompilerContext& CompilerContext, UEdGraph* SourceGraph) override;
+	virtual ERedirectType DoPinsMatchForReconstruction(const UEdGraphPin* NewPin, int32 NewPinIndex, const UEdGraphPin* OldPin, int32 OldPinIndex) const override;
+	virtual bool IsConnectionDisallowed(const UEdGraphPin* MyPin, const UEdGraphPin* OtherPin, FString& OutReason) const override;
+	virtual void GetMenuActions(FBlueprintActionDatabaseRegistrar& ActionRegistrar) const override;
+	virtual FText GetMenuCategory() const override;
+	virtual int32 GetNodeRefreshPriority() const override { return EBaseNodeRefreshPriority::Low_UsesDependentWildcard; }
+	virtual void PreloadRequiredAssets() override;
+
+	virtual void DestroyNode() override;
+	//~ End UK2Node Interface.
+
+private:
+	void ChooserSignatureChanged();
+	void UnregisterCallback();
+	UChooserSignature* CurrentCallbackObject = nullptr;
+
+	UPROPERTY(EditAnywhere, Category = "ChooserSignature")
+	TObjectPtr<UChooserSignature> ChooserSignature;
+
+	/** Tooltip text for this node. */
+	FText NodeTooltip;
+};
