Fix issues with chooser table editor selection details - garbage collection crashes and multi edit issues

diff --git a/Engine/Plugins/Chooser/Source/ChooserEditor/Private/ChooserTableEditor.cpp b/Engine/Plugins/Chooser/Source/ChooserEditor/Private/ChooserTableEditor.cpp
--- a/Engine/Plugins/Chooser/Source/ChooserEditor/Private/ChooserTableEditor.cpp
+++ b/Engine/Plugins/Chooser/Source/ChooserEditor/Private/ChooserTableEditor.cpp
@@ -202,7 +202,7 @@ FChooserTableEditor::~FChooserTableEditor()
 {
 	if (SelectedColumn)
 	{
-		SelectedColumn->RemoveFromRoot();
+		SelectedColumn->ClearFlags(RF_Standalone);
 		SelectedColumn = nullptr;
 	}
 	for (UObject* SelectedRow : SelectedRowsCache)
@@ -1257,22 +1257,25 @@ void FChooserTableEditor::RefreshRowSelectionDetails()
 	int32 SelectedRowsCacheIndex = 0;
 	for(TSharedPtr<FChooserTableRow>& SelectedItem : SelectedItems)
 	{
-		if (Chooser->ResultsStructs.IsValidIndex(SelectedItem->RowIndex))
+		if (Chooser->ResultsStructs.IsValidIndex(SelectedItem->RowIndex) || SelectedItem->RowIndex == ColumnWidget_SpecialIndex_Fallback)
 		{
 			if (SelectedRowsCache[SelectedRowsCacheIndex] == nullptr)
 			{
-				SelectedRowsCache[SelectedRowsCacheIndex] = NewObject<UChooserRowDetails>();
-				SelectedRowsCache[SelectedRowsCacheIndex]->SetFlags(RF_Standalone | RF_Transactional);
+				SelectedRowsCache[SelectedRowsCacheIndex] = NewObject<UChooserRowDetails>(GetTransientPackage(), FName("SelectedRow",SelectedRowsCacheIndex),RF_Standalone | RF_Transactional | RF_Transient);
 			}
 			TObjectPtr<UChooserRowDetails> Selection = SelectedRowsCache[SelectedRowsCacheIndex];
 			SelectedRowsCacheIndex ++;
 			
+			check(Selection->HasAnyFlags(RF_Standalone));
+			check(!Selection->HasAnyFlags(RF_MirroredGarbage))
+			
 			Selection->Chooser = Chooser;
 			Selection->Row = SelectedItem->RowIndex;
-			
+			Selection->Properties.Reset();
+		
 			if (Chooser->ResultType != EObjectChooserResultType::NoPrimaryResult)
 			{
-				FInstancedStruct& Result = Chooser->ResultsStructs[SelectedItem->RowIndex];
+				FInstancedStruct& Result = SelectedItem->RowIndex == ColumnWidget_SpecialIndex_Fallback? Chooser->FallbackResult : Chooser->ResultsStructs[SelectedItem->RowIndex];;
 				Selection->Properties.AddProperties({ResultPropertyDesc});
 				Selection->Properties.SetValueStruct("Result", FConstStructView(FInstancedStruct::StaticStruct(), reinterpret_cast<uint8*>(&Result)));
 			}
@@ -1291,33 +1294,6 @@ void FChooserTableEditor::RefreshRowSelectionDetails()
 				Selection->Properties.SetValueBool("Disabled", Chooser->DisabledRows[SelectedItem->RowIndex]);
 			}
 
-			SelectedRows.Add(Selection);
-		}
-		else if (SelectedItem->RowIndex == ColumnWidget_SpecialIndex_Fallback)
-		{
-			TObjectPtr<UChooserRowDetails> Selection = NewObject<UChooserRowDetails>();
-			Selection->Chooser = Chooser;
-			Selection->Row = SelectedItem->RowIndex;
-			Selection->SetFlags(RF_Standalone | RF_Transactional);
-
-			if (Chooser->ResultType != EObjectChooserResultType::NoPrimaryResult)
-			{
-				FInstancedStruct& Result = Chooser->FallbackResult;
-				Selection->Properties.AddProperties({ResultPropertyDesc});
-				Selection->Properties.SetValueStruct("Result", FConstStructView(FInstancedStruct::StaticStruct(), reinterpret_cast<uint8*>(&Result)));
-			}
-
-			int ColumnIndex = 0;
-			for (FInstancedStruct& ColumnData : Chooser->ColumnsStructs)
-			{
-				FChooserColumnBase& Column = ColumnData.GetMutable<FChooserColumnBase>();
-				if (Column.HasOutputs())
-				{
-					Column.AddToDetails(Selection->Properties, ColumnIndex, SelectedItem->RowIndex);
-				}
-				ColumnIndex++;
-			}
-
 			SelectedRows.Add(Selection);
 		}
 	}
@@ -1852,8 +1828,7 @@ void FChooserTableEditor::SelectColumn(UChooserTable* ChooserEditor, int Index)
    	{
    		if (SelectedColumn == nullptr)
    		{
-   			SelectedColumn = NewObject<UChooserColumnDetails>();
-			SelectedColumn->AddToRoot();
+   			SelectedColumn = NewObject<UChooserColumnDetails>(GetTransientPackage(), "ColumnSelection", RF_Standalone | RF_Transactional | RF_Transient);
    		}
    
    		SelectedColumn->Chooser = Chooser;
