Fix crash selecting fallback row in chooser tables

diff --git a/Engine/Plugins/Chooser/Source/ChooserEditor/Private/ChooserTableEditor.cpp b/Engine/Plugins/Chooser/Source/ChooserEditor/Private/ChooserTableEditor.cpp
--- a/Engine/Plugins/Chooser/Source/ChooserEditor/Private/ChooserTableEditor.cpp
+++ b/Engine/Plugins/Chooser/Source/ChooserEditor/Private/ChooserTableEditor.cpp
@@ -1284,7 +1284,12 @@ void FChooserTableEditor::RefreshRowSelectionDetails()
 			for (FInstancedStruct& ColumnData : Chooser->ColumnsStructs)
 			{
 				FChooserColumnBase& Column = ColumnData.GetMutable<FChooserColumnBase>();
-				Column.AddToDetails(Selection->Properties, ColumnIndex, SelectedItem->RowIndex);
+
+				if (Column.HasOutputs() || SelectedItem->RowIndex != ColumnWidget_SpecialIndex_Fallback)
+				{
+					// add each column's details, but on the fallback row, only add details for Output columns
+					Column.AddToDetails(Selection->Properties, ColumnIndex, SelectedItem->RowIndex);
+				}
 				ColumnIndex++;
 			}
 
