Fix for binding Enum chooser parameters in details panel

diff --git a/Engine/Plugins/Chooser/Source/ChooserEditor/Private/PropertyAccessChainCustomization.cpp b/Engine/Plugins/Chooser/Source/ChooserEditor/Private/PropertyAccessChainCustomization.cpp
--- a/Engine/Plugins/Chooser/Source/ChooserEditor/Private/PropertyAccessChainCustomization.cpp
+++ b/Engine/Plugins/Chooser/Source/ChooserEditor/Private/PropertyAccessChainCustomization.cpp
@@ -64,8 +64,6 @@ void FPropertyAccessChainCustomization::CustomizeHeader(TSharedRef<IPropertyHand
 		})
 		.OnAddBinding_Lambda([PropertyHandle, this, HasContext, TypeFilter](FName /*InPropertyName*/, const TArray<FBindingChainElement>& InBindingChain)
 		{
-			FChooserPropertyBinding ContextProperty;
-			SPropertyAccessChainWidget::SetPropertyBinding(HasContext, InBindingChain, ContextProperty);
 
 			UObject* TransactionObject = Cast<UObject>(HasContext);
             	
@@ -81,13 +79,13 @@ void FPropertyAccessChainCustomization::CustomizeHeader(TSharedRef<IPropertyHand
 			
 			if (FStructProperty* StructProperty = CastField<FStructProperty>(PropertyHandle->GetProperty()))
 			{
-				if (StructProperty->Struct == FChooserPropertyBinding::StaticStruct())
+				if (StructProperty->Struct->IsChildOf(FChooserPropertyBinding::StaticStruct()))
 				{
 					TArray<void*> PropertyDatas;
 					PropertyHandle->AccessRawData(PropertyDatas);
 					for (void* PropertyData : PropertyDatas)
 					{
-						*static_cast<FChooserPropertyBinding*>(PropertyData) = ContextProperty;
+						SPropertyAccessChainWidget::SetPropertyBinding(HasContext, InBindingChain, *static_cast<FChooserPropertyBinding*>(PropertyData));
 					}
 				}
 			}
