Console variable chooser parameters

diff --git a/Engine/Plugins/Chooser/Source/Chooser/Private/ConsoleVariableParameters.cpp b/Engine/Plugins/Chooser/Source/Chooser/Private/ConsoleVariableParameters.cpp
new file mode 100644
--- /dev/null
+++ b/Engine/Plugins/Chooser/Source/Chooser/Private/ConsoleVariableParameters.cpp
@@ -0,0 +1,98 @@
+// Copyright Epic Games, Inc. All Rights Reserved.
+
+#include "ConsoleVariableParameters.h"
+#include "HAL/IConsoleManager.h"
+#include "Logging/LogMacros.h"
+
+bool FBoolCVarProperty::GetValue(FChooserEvaluationContext& Context, bool& OutResult) const
+{
+	IConsoleVariable* Variable = IConsoleManager::Get().FindConsoleVariable(*VariableName);
+ 	if (Variable)
+ 	{
+ 		OutResult = Variable->GetBool();
+ 		return true;
+ 	}
+ 	
+ 	UE_LOG(LogChooser, Warning, TEXT("Failed to find console variable '%s'."), *VariableName);
+ 	return false;
+}
+
+bool FBoolCVarProperty::SetValue(FChooserEvaluationContext& Context, bool InValue) const
+{
+	IConsoleVariable* Variable = IConsoleManager::Get().FindConsoleVariable(*VariableName);
+	if (Variable)
+	{
+		Variable->Set(InValue);
+		return true;
+	}
+	
+	UE_LOG(LogChooser, Warning, TEXT("Failed to find console variable '%s'."), *VariableName);
+	return false;
+}
+
+void FBoolCVarProperty::GetDisplayName(FText& OutName) const
+{
+	OutName = FText::FromString(VariableName);
+}
+
+bool FFloatCVarProperty::GetValue(FChooserEvaluationContext& Context, double& OutResult) const
+{
+	IConsoleVariable* Variable = IConsoleManager::Get().FindConsoleVariable(*VariableName);
+	if (Variable)
+	{
+		OutResult = Variable->GetFloat();
+		return true;
+	}
+	
+	UE_LOG(LogChooser, Warning, TEXT("Failed to find console variable '%s'."), *VariableName);
+	return false;
+}
+
+bool FFloatCVarProperty::SetValue(FChooserEvaluationContext& Context, double InValue) const
+{
+	IConsoleVariable* Variable = IConsoleManager::Get().FindConsoleVariable(*VariableName);
+	if (Variable)
+	{
+		Variable->Set(static_cast<float>(InValue));
+		return true;
+	}
+	
+	UE_LOG(LogChooser, Warning, TEXT("Failed to find console variable '%s'."), *VariableName);
+	return false;
+}
+
+void FFloatCVarProperty::GetDisplayName(FText& OutName) const
+{
+	OutName = FText::FromString(VariableName);
+}
+
+bool FEnumCVarProperty::GetValue(FChooserEvaluationContext& Context, uint8& OutResult) const
+{
+	IConsoleVariable* Variable = IConsoleManager::Get().FindConsoleVariable(*VariableName);
+	if (Variable)
+	{
+		OutResult = Variable->GetInt();
+		return true;
+	}
+	
+	UE_LOG(LogChooser, Warning, TEXT("Failed to find console variable '%s'."), *VariableName);
+	return false;
+}
+
+bool FEnumCVarProperty::SetValue(FChooserEvaluationContext& Context, uint8 InValue) const
+{
+	IConsoleVariable* Variable = IConsoleManager::Get().FindConsoleVariable(*VariableName);
+	if (Variable)
+	{
+		Variable->Set(InValue);
+		return true;
+	}
+	
+	UE_LOG(LogChooser, Warning, TEXT("Failed to find console variable '%s'."), *VariableName);
+	return false;
+}
+
+void FEnumCVarProperty::GetDisplayName(FText& OutName) const
+{
+	OutName = FText::FromString(VariableName);
+}
\ No newline at end of file
diff --git a/Engine/Plugins/Chooser/Source/Chooser/Private/ConsoleVariableParameters.h b/Engine/Plugins/Chooser/Source/Chooser/Private/ConsoleVariableParameters.h
new file mode 100644
--- /dev/null
+++ b/Engine/Plugins/Chooser/Source/Chooser/Private/ConsoleVariableParameters.h
@@ -0,0 +1,60 @@
+// Copyright Epic Games, Inc. All Rights Reserved.
+#pragma once
+
+#include "CoreMinimal.h"
+#include "IChooserParameterBool.h"
+#include "IChooserParameterFloat.h"
+#include "IChooserParameterEnum.h"
+#include "ConsoleVariableParameters.generated.h"
+
+
+USTRUCT(DisplayName = "Bool CVar Param")
+struct FBoolCVarProperty :  public FChooserParameterBoolBase
+{
+	GENERATED_BODY()
+public:
+
+	UPROPERTY(EditAnywhere, Category="Variable")
+	FString VariableName;
+
+	virtual bool GetValue(FChooserEvaluationContext& Context, bool& OutResult) const override;
+	virtual bool SetValue(FChooserEvaluationContext& Context, bool InValue) const override;
+	virtual void GetDisplayName(FText& OutName) const override;
+};
+
+USTRUCT(DisplayName = "Float CVar Param")
+struct FFloatCVarProperty :  public FChooserParameterFloatBase
+{
+	GENERATED_BODY()
+public:
+
+	UPROPERTY(EditAnywhere, Category="Variable")
+	FString VariableName;
+
+	virtual bool GetValue(FChooserEvaluationContext& Context, double& OutResult) const override;
+	virtual bool SetValue(FChooserEvaluationContext& Context, double InValue) const override;
+	virtual void GetDisplayName(FText& OutName) const override;
+};
+
+USTRUCT(DisplayName = "Enum CVar Param")
+struct FEnumCVarProperty :  public FChooserParameterEnumBase
+{
+	GENERATED_BODY()
+
+public:
+
+	UPROPERTY(EditAnywhere, Category="Variable")
+	FString VariableName;
+
+	UPROPERTY(EditAnywhere, Category="Variable")
+	TObjectPtr<UEnum> Enum;
+
+	virtual bool GetValue(FChooserEvaluationContext& Context, uint8& OutResult) const override;
+	virtual bool SetValue(FChooserEvaluationContext& Context, uint8 InValue) const override;
+
+	virtual void GetDisplayName(FText& OutName) const override;
+
+	#if WITH_EDITOR
+	virtual const UEnum* GetEnum() const override { return Enum; }
+	#endif
+};
\ No newline at end of file
