Add more than one recently selected result tracking to randomization column

diff --git a/Engine/Plugins/Chooser/Source/Chooser/Internal/RandomizeColumn.h b/Engine/Plugins/Chooser/Source/Chooser/Internal/RandomizeColumn.h
--- a/Engine/Plugins/Chooser/Source/Chooser/Internal/RandomizeColumn.h
+++ b/Engine/Plugins/Chooser/Source/Chooser/Internal/RandomizeColumn.h
@@ -19,7 +19,7 @@ public:
 	UPROPERTY(EditAnywhere, Meta = (BindingType = "FChooserRandomizationContext", BindingAllowFunctions = "false", BindingColor = "StructPinTypeColor"), Category = "Binding")
 	FChooserPropertyBinding Binding;
 	
-	UE_API virtual bool GetValue(FChooserEvaluationContext& Context, const FChooserRandomizationContext*& OutResult) const override;
+	UE_API virtual bool GetValue(FChooserEvaluationContext& Context, FChooserRandomizationContext*& OutResult) const override;
 	virtual bool IsBound() const override
 	{
 		return Binding.IsBoundToRoot || !Binding.PropertyBindingChain.IsEmpty();
@@ -38,9 +38,8 @@ struct FRandomizeColumn : public FChooserColumnBase
 	UPROPERTY(EditAnywhere, NoClear, Meta = (ExcludeBaseStruct, BaseStruct = "/Script/Chooser.RandomizeContextProperty", ToolTip="Optional reference to a ChooserRandomizationContext struct. If bound, this is used to store the most recent selection (for each Choosers referencing it - you only need to create one variable per Character or context), for use with RepateProbabilityMultiplier to reduce the chance of selecting the same entry twice."), Category = "Data")
 	FInstancedStruct InputValue;
 	
-	
-	UPROPERTY(EditAnywhere, Category= "Data", meta=(ClampMin="0.0",Tooltip="Multiplies the weight of the previous chosen result (set to 0 to never pick the same result twice in a row)"));
-	float RepeatProbabilityMultiplier = 1.0f;
+	UPROPERTY(EditAnywhere, Category= "Data", meta=(ClampMin="0.0",Tooltip="Multiplies the weight of recently chosen results, where the number of entries in this array indicates how many recently selected choices to remember (A value of zero means to avoid repeating an entry unless it's the only choice). You must bind a RandomizationContext struct for this feature to work."));
+	TArray<float> RepeatProbabilityMultipliers;
 
 	UPROPERTY(EditAnywhere, Category= "Data", meta=(ClampMin="0.0",Tooltip="When columns with scoring are used, randomize will pick from among all rows that have a cost nearly equal to the minumum cost, using this threshold"));
 	float EqualCostThreshold = 0.001f;
@@ -48,6 +47,10 @@ struct FRandomizeColumn : public FChooserColumnBase
 #if WITH_EDITORONLY_DATA
 	UPROPERTY(EditAnywhere, Category=Data, meta=(ToolTip="DefaultRowValue will be assigned to cells when new rows are created"));
 	float DefaultRowValue = 1.0f;
+
+	UE_DEPRECATED(5.7, "Moved to RepeatProbabilityMultipliers array")
+	UPROPERTY()
+	float RepeatProbabilityMultiplier = 1.0f;
 #endif
 	
 	UPROPERTY()
@@ -71,6 +74,16 @@ struct FRandomizeColumn : public FChooserColumnBase
 
 	virtual void PostLoad() override
 	{
+#if WITH_EDITORONLY_DATA
+PRAGMA_DISABLE_DEPRECATION_WARNINGS
+		if (RepeatProbabilityMultiplier < 1.0f && RepeatProbabilityMultipliers.Num() == 0)
+		{
+			RepeatProbabilityMultipliers.Add(RepeatProbabilityMultiplier);
+			RepeatProbabilityMultiplier = 1.0f;
+		}
+PRAGMA_ENABLE_DEPRECATION_WARNINGS
+#endif
+		
 		if (InputValue.IsValid())
 		{
 			InputValue.GetMutable<FChooserParameterBase>().PostLoad();
diff --git a/Engine/Plugins/Chooser/Source/Chooser/Private/RandomizeColumn.cpp b/Engine/Plugins/Chooser/Source/Chooser/Private/RandomizeColumn.cpp
--- a/Engine/Plugins/Chooser/Source/Chooser/Private/RandomizeColumn.cpp
+++ b/Engine/Plugins/Chooser/Source/Chooser/Private/RandomizeColumn.cpp
@@ -8,7 +8,7 @@
 
 #include UE_INLINE_GENERATED_CPP_BY_NAME(RandomizeColumn)
 
-bool FRandomizeContextProperty::GetValue(FChooserEvaluationContext& Context, const FChooserRandomizationContext*& OutResult) const
+bool FRandomizeContextProperty::GetValue(FChooserEvaluationContext& Context, FChooserRandomizationContext*& OutResult) const
 {
 	return Binding.GetValuePtr(Context,OutResult);
 }
@@ -18,12 +18,27 @@ FRandomizeColumn::FRandomizeColumn()
 	InputValue.InitializeAs(FRandomizeContextProperty::StaticStruct());
 }
 
+float GetRepeatProbabilityMultiplier(const TArray<float>& RepeatProbabilityMultipliers, const FRandomizationState* State, uint32 Index)
+{
+	if (State)
+	{
+		int32 RecentlySelectedIndex =  State->RecentlySelectedRows.Find(Index);
+		if (RecentlySelectedIndex != INDEX_NONE)
+		{
+			// RecentlySelectedRows array is a ring buffer.  Remap to where 0 means it was selected	most recently
+			RecentlySelectedIndex = (State->LastSelectedRow  - RecentlySelectedIndex + State->RecentlySelectedRows.Num())  % State->RecentlySelectedRows.Num();
+			return RepeatProbabilityMultipliers[RecentlySelectedIndex];
+		}
+	}
+	return 1.0;
+}
+
 void FRandomizeColumn::Filter(FChooserEvaluationContext& Context, const FChooserIndexArray& IndexListIn, FChooserIndexArray& IndexListOut) const
 {
 	int Count = IndexListIn.Num();
 	int Selection = 0;
 
-	const FChooserRandomizationContext* RandomizationContext = nullptr;
+	FChooserRandomizationContext* RandomizationContext = nullptr;
 	if (InputValue.IsValid())
 	{
 		InputValue.Get<FChooserParameterRandomizeBase>().GetValue(Context,RandomizationContext);
@@ -31,12 +46,15 @@ void FRandomizeColumn::Filter(FChooserEvaluationContext& Context, const FChooser
 	
 	if (Count > 1)
 	{
-		int LastSelectedIndex = -1;
+		FRandomizationState* State  = nullptr;
 		if (RandomizationContext)
 		{
-			if (const FRandomizationState* State = RandomizationContext->StateMap.Find(this))
+			State = &RandomizationContext->StateMap.FindOrAdd(this, FRandomizationState());
+
+			if (State->RecentlySelectedRows.Num() != RepeatProbabilityMultipliers.Num())
 			{
-				LastSelectedIndex = State->LastSelectedRow;
+				State->RecentlySelectedRows.Init(-1, RepeatProbabilityMultipliers.Num());
+				State->LastSelectedRow = 0;
 			}
 		}
 
@@ -66,11 +84,8 @@ void FRandomizeColumn::Filter(FChooserEvaluationContext& Context, const FChooser
 					{
 						RowWeight = RowValues[IndexData.Index];
 					}
-			
-					if (IndexData.Index == LastSelectedIndex)
-					{
-						RowWeight *= RepeatProbabilityMultiplier;
-					}
+
+					RowWeight = RowWeight * GetRepeatProbabilityMultiplier(RepeatProbabilityMultipliers, State, IndexData.Index);
 			
 					TotalWeight += RowWeight;
 					MinCount++;
@@ -101,11 +116,8 @@ void FRandomizeColumn::Filter(FChooserEvaluationContext& Context, const FChooser
 					{
 						RowWeight = RowValues[IndexData.Index];
 					}
-				
-					if (IndexData.Index == LastSelectedIndex)
-					{
-						RowWeight *= RepeatProbabilityMultiplier;
-					}
+
+					RowWeight = RowWeight * GetRepeatProbabilityMultiplier(RepeatProbabilityMultipliers, State, IndexData.Index);
 				
 					Weight += RowWeight;
 					if (Weight > RandomNumber)
@@ -126,11 +138,8 @@ void FRandomizeColumn::Filter(FChooserEvaluationContext& Context, const FChooser
 				{
 					RowWeight = RowValues[IndexData.Index];
 				}
-			
-				if (IndexData.Index == LastSelectedIndex)
-				{
-					RowWeight *= RepeatProbabilityMultiplier;
-				}
+
+				RowWeight = RowWeight * GetRepeatProbabilityMultiplier(RepeatProbabilityMultipliers, State, IndexData.Index);
 			
 				TotalWeight += RowWeight;
 			}
@@ -150,10 +159,7 @@ void FRandomizeColumn::Filter(FChooserEvaluationContext& Context, const FChooser
 					RowWeight = RowValues[IndexData.Index];
 				}
 			
-				if (IndexData.Index == LastSelectedIndex)
-				{
-					RowWeight *= RepeatProbabilityMultiplier;
-				}
+				RowWeight = RowWeight * GetRepeatProbabilityMultiplier(RepeatProbabilityMultipliers, State, IndexData.Index);
 			
 				Weight += RowWeight;
 				if (Weight > RandomNumber)
@@ -174,13 +180,18 @@ void FRandomizeColumn::SetOutputs(FChooserEvaluationContext& Context, int RowInd
 {
 	if (InputValue.IsValid() && RowValues.IsValidIndex((RowIndex)))
 	{
-		const FChooserRandomizationContext* ConstRandomizationContext = nullptr;
-		InputValue.Get<FChooserParameterRandomizeBase>().GetValue(Context,ConstRandomizationContext);
-		if(ConstRandomizationContext)
+		FChooserRandomizationContext* RandomizationContext = nullptr;
+		InputValue.Get<FChooserParameterRandomizeBase>().GetValue(Context,RandomizationContext);
+		if(RandomizationContext)
 		{
-			FChooserRandomizationContext* RandomizationContext = const_cast<FChooserRandomizationContext*>(ConstRandomizationContext);
-			FRandomizationState& State = RandomizationContext->StateMap.FindOrAdd(this, FRandomizationState());
-			State.LastSelectedRow = RowIndex;
+			if (FRandomizationState* State = RandomizationContext->StateMap.Find(this))
+			{
+				if (State->RecentlySelectedRows.Num() > 0)
+				{
+					State->LastSelectedRow = (State->LastSelectedRow + 1) % State->RecentlySelectedRows.Num();
+					State->RecentlySelectedRows[State->LastSelectedRow] = RowIndex;
+				}
+			}
 		}
 	}
 }
diff --git a/Engine/Plugins/Chooser/Source/Chooser/Public/IChooserParameterRandomize.h b/Engine/Plugins/Chooser/Source/Chooser/Public/IChooserParameterRandomize.h
--- a/Engine/Plugins/Chooser/Source/Chooser/Public/IChooserParameterRandomize.h
+++ b/Engine/Plugins/Chooser/Source/Chooser/Public/IChooserParameterRandomize.h
@@ -11,6 +11,7 @@
 
 struct FRandomizationState
 {
+	TArray<int32, TInlineAllocator<4>> RecentlySelectedRows;
 	int LastSelectedRow;
 };
 
@@ -26,6 +27,6 @@ USTRUCT()
 struct FChooserParameterRandomizeBase : public FChooserParameterBase
 {
 	GENERATED_BODY()
-	virtual bool GetValue(FChooserEvaluationContext& Context, const FChooserRandomizationContext*& OutResult) const { return false; }
+	virtual bool GetValue(FChooserEvaluationContext& Context, FChooserRandomizationContext*& OutResult) const { return false; }
 	virtual bool IsBound() const { return false; }
 };
