Add feature to Randomize Column to keep track of multiple recently selected entries

diff --git a/Engine/Plugins/Chooser/Source/Chooser/Internal/RandomizeColumn.h b/Engine/Plugins/Chooser/Source/Chooser/Internal/RandomizeColumn.h
--- a/Engine/Plugins/Chooser/Source/Chooser/Internal/RandomizeColumn.h
+++ b/Engine/Plugins/Chooser/Source/Chooser/Internal/RandomizeColumn.h
@@ -19,7 +19,7 @@ public:
 	UPROPERTY(EditAnywhere, Meta = (BindingType = "FChooserRandomizationContext", BindingAllowFunctions = "false", BindingColor = "StructPinTypeColor"), Category = "Binding")
 	FChooserPropertyBinding Binding;
 	
-	UE_API virtual bool GetValue(FChooserEvaluationContext& Context, const FChooserRandomizationContext*& OutResult) const override;
+	UE_API virtual bool GetValue(FChooserEvaluationContext& Context, FChooserRandomizationContext*& OutResult) const override;
 	virtual bool IsBound() const override
 	{
 		return Binding.IsBoundToRoot || !Binding.PropertyBindingChain.IsEmpty();
@@ -38,9 +38,8 @@ struct FRandomizeColumn : public FChooserColumnBase
 	UPROPERTY(EditAnywhere, NoClear, Meta = (ExcludeBaseStruct, BaseStruct = "/Script/Chooser.RandomizeContextProperty", ToolTip="Optional reference to a ChooserRandomizationContext struct. If bound, this is used to store the most recent selection (for each Choosers referencing it - you only need to create one variable per Character or context), for use with RepateProbabilityMultiplier to reduce the chance of selecting the same entry twice."), Category = "Data")
 	FInstancedStruct InputValue;
 	
-	
-	UPROPERTY(EditAnywhere, Category= "Data", meta=(ClampMin="0.0",Tooltip="Multiplies the weight of the previous chosen result (set to 0 to never pick the same result twice in a row)"));
-	float RepeatProbabilityMultiplier = 1.0f;
+	UPROPERTY(EditAnywhere, Category= "Data", meta=(ClampMin="0.0",Tooltip="Multiplies the weight of recently chosen results, where the number of entries in this array indicates how many recently selected choices to remember (A value of zero means to avoid repeating an entry unless it's the only choice). You must bind a RandomizationContext struct for this feature to work."));
+	TArray<float> RepeatProbabilityMultipliers;
 
 	UPROPERTY(EditAnywhere, Category= "Data", meta=(ClampMin="0.0",Tooltip="When columns with scoring are used, randomize will pick from among all rows that have a cost nearly equal to the minumum cost, using this threshold"));
 	float EqualCostThreshold = 0.001f;
@@ -48,6 +47,9 @@ struct FRandomizeColumn : public FChooserColumnBase
 #if WITH_EDITORONLY_DATA
 	UPROPERTY(EditAnywhere, Category=Data, meta=(ToolTip="DefaultRowValue will be assigned to cells when new rows are created"));
 	float DefaultRowValue = 1.0f;
+
+	UPROPERTY()
+	float RepeatProbabilityMultiplier = 1.0f;
 #endif
 	
 	UPROPERTY()
@@ -71,6 +73,14 @@ struct FRandomizeColumn : public FChooserColumnBase
 
 	virtual void PostLoad() override
 	{
+#if WITH_EDITORONLY_DATA
+		if (RepeatProbabilityMultiplier < 1.0f && RepeatProbabilityMultipliers.Num() == 0)
+		{
+			RepeatProbabilityMultipliers.Add(RepeatProbabilityMultiplier);
+			RepeatProbabilityMultiplier = 1.0f;
+		}
+#endif
+		
 		if (InputValue.IsValid())
 		{
 			InputValue.GetMutable<FChooserParameterBase>().PostLoad();
diff --git a/Engine/Plugins/Chooser/Source/Chooser/Private/RandomizeColumn.cpp b/Engine/Plugins/Chooser/Source/Chooser/Private/RandomizeColumn.cpp
--- a/Engine/Plugins/Chooser/Source/Chooser/Private/RandomizeColumn.cpp
+++ b/Engine/Plugins/Chooser/Source/Chooser/Private/RandomizeColumn.cpp
@@ -8,7 +8,7 @@
 
 #include UE_INLINE_GENERATED_CPP_BY_NAME(RandomizeColumn)
 
-bool FRandomizeContextProperty::GetValue(FChooserEvaluationContext& Context, const FChooserRandomizationContext*& OutResult) const
+bool FRandomizeContextProperty::GetValue(FChooserEvaluationContext& Context, FChooserRandomizationContext*& OutResult) const
 {
 	return Binding.GetValuePtr(Context,OutResult);
 }
@@ -18,12 +18,27 @@ FRandomizeColumn::FRandomizeColumn()
 	InputValue.InitializeAs(FRandomizeContextProperty::StaticStruct());
 }
 
+float GetRepeatProbabilityMultiplier(const TArray<float>& RepeatProbabilityMultipliers, const FRandomizationState* State, uint32 Index)
+{
+	if (State)
+	{
+		int32 RecentlySelectedIndex =  State->RecentlySelectedRows.Find(Index);
+		if (RecentlySelectedIndex != INDEX_NONE)
+		{
+			// RecentlySelectedRows array is a ring buffer.  Remap to where 0 means it was selected	most recently
+			RecentlySelectedIndex = (State->LastSelectedRow  - RecentlySelectedIndex + State->RecentlySelectedRows.Num())  % State->RecentlySelectedRows.Num();
+			return RepeatProbabilityMultipliers[RecentlySelectedIndex];
+		}
+	}
+	return 1.0;
+}
+
 void FRandomizeColumn::Filter(FChooserEvaluationContext& Context, const FChooserIndexArray& IndexListIn, FChooserIndexArray& IndexListOut) const
 {
 	int Count = IndexListIn.Num();
 	int Selection = 0;
 
-	const FChooserRandomizationContext* RandomizationContext = nullptr;
+	FChooserRandomizationContext* RandomizationContext = nullptr;
 	if (InputValue.IsValid())
 	{
 		InputValue.Get<FChooserParameterRandomizeBase>().GetValue(Context,RandomizationContext);
@@ -31,12 +46,15 @@ void FRandomizeColumn::Filter(FChooserEvaluationContext& Context, const FChooser
 	
 	if (Count > 1)
 	{
-		int LastSelectedIndex = -1;
+		FRandomizationState* State  = nullptr;
 		if (RandomizationContext)
 		{
-			if (const FRandomizationState* State = RandomizationContext->StateMap.Find(this))
+			State = &RandomizationContext->StateMap.FindOrAdd(this, FRandomizationState());
+
+			if (State->RecentlySelectedRows.Num() != RepeatProbabilityMultipliers.Num())
 			{
-				LastSelectedIndex = State->LastSelectedRow;
+				State->RecentlySelectedRows.Init(-1, RepeatProbabilityMultipliers.Num());
+				State->LastSelectedRow = 0;
 			}
 		}
 
@@ -66,11 +84,8 @@ void FRandomizeColumn::Filter(FChooserEvaluationContext& Context, const FChooser
 					{
 						RowWeight = RowValues[IndexData.Index];
 					}
-			
-					if (IndexData.Index == LastSelectedIndex)
-					{
-						RowWeight *= RepeatProbabilityMultiplier;
-					}
+
+					RowWeight = RowWeight * GetRepeatProbabilityMultiplier(RepeatProbabilityMultipliers, State, IndexData.Index);
 			
 					TotalWeight += RowWeight;
 					MinCount++;
@@ -83,7 +98,14 @@ void FRandomizeColumn::Filter(FChooserEvaluationContext& Context, const FChooser
 				IndexListOut.Push({LowestCostIndex, LowestCost});
 				return;
 			}
-				
+
+			float AddWeight = 0.0f;
+			if (TotalWeight < UE_SMALL_NUMBER)
+			{
+				// when total weight is 0, and there is more than one entry, fake uniform weights;
+				TotalWeight = 1.0f;
+				AddWeight = 1.0f / static_cast<float>(MinCount);
+			}
 
 			// pick a random float from 0-total weight
 			float RandomNumber = FMath::FRandRange(0.0f, TotalWeight);
@@ -101,11 +123,8 @@ void FRandomizeColumn::Filter(FChooserEvaluationContext& Context, const FChooser
 					{
 						RowWeight = RowValues[IndexData.Index];
 					}
-				
-					if (IndexData.Index == LastSelectedIndex)
-					{
-						RowWeight *= RepeatProbabilityMultiplier;
-					}
+
+					RowWeight = AddWeight + RowWeight * GetRepeatProbabilityMultiplier(RepeatProbabilityMultipliers, State, IndexData.Index);
 				
 					Weight += RowWeight;
 					if (Weight > RandomNumber)
@@ -126,15 +145,20 @@ void FRandomizeColumn::Filter(FChooserEvaluationContext& Context, const FChooser
 				{
 					RowWeight = RowValues[IndexData.Index];
 				}
-			
-				if (IndexData.Index == LastSelectedIndex)
-				{
-					RowWeight *= RepeatProbabilityMultiplier;
-				}
+
+				RowWeight = RowWeight * GetRepeatProbabilityMultiplier(RepeatProbabilityMultipliers, State, IndexData.Index);
 			
 				TotalWeight += RowWeight;
 			}
 
+			float AddWeight = 0.0f;
+			if (TotalWeight < UE_SMALL_NUMBER)
+			{
+				// when total weight is 0, and there is more than one entry, fake uniform weights;
+				TotalWeight = 1.0f;
+				AddWeight = 1.0f / static_cast<float>(IndexListIn.Num());
+			}
+
 			// pick a random float from 0-total weight
 			float RandomNumber = FMath::FRandRange(0.0f, TotalWeight);
 			float Weight = 0;
@@ -150,10 +174,7 @@ void FRandomizeColumn::Filter(FChooserEvaluationContext& Context, const FChooser
 					RowWeight = RowValues[IndexData.Index];
 				}
 			
-				if (IndexData.Index == LastSelectedIndex)
-				{
-					RowWeight *= RepeatProbabilityMultiplier;
-				}
+				RowWeight = AddWeight + RowWeight * GetRepeatProbabilityMultiplier(RepeatProbabilityMultipliers, State, IndexData.Index);
 			
 				Weight += RowWeight;
 				if (Weight > RandomNumber)
@@ -174,13 +195,18 @@ void FRandomizeColumn::SetOutputs(FChooserEvaluationContext& Context, int RowInd
 {
 	if (InputValue.IsValid() && RowValues.IsValidIndex((RowIndex)))
 	{
-		const FChooserRandomizationContext* ConstRandomizationContext = nullptr;
-		InputValue.Get<FChooserParameterRandomizeBase>().GetValue(Context,ConstRandomizationContext);
-		if(ConstRandomizationContext)
+		FChooserRandomizationContext* RandomizationContext = nullptr;
+		InputValue.Get<FChooserParameterRandomizeBase>().GetValue(Context,RandomizationContext);
+		if(RandomizationContext)
 		{
-			FChooserRandomizationContext* RandomizationContext = const_cast<FChooserRandomizationContext*>(ConstRandomizationContext);
-			FRandomizationState& State = RandomizationContext->StateMap.FindOrAdd(this, FRandomizationState());
-			State.LastSelectedRow = RowIndex;
+			if (FRandomizationState* State = RandomizationContext->StateMap.Find(this))
+			{
+				if (State->RecentlySelectedRows.Num() > 0)
+				{
+					State->LastSelectedRow = (State->LastSelectedRow + 1) % State->RecentlySelectedRows.Num();
+					State->RecentlySelectedRows[State->LastSelectedRow] = RowIndex;
+				}
+			}
 		}
 	}
 }
diff --git a/Engine/Plugins/Chooser/Source/Chooser/Public/IChooserParameterRandomize.h b/Engine/Plugins/Chooser/Source/Chooser/Public/IChooserParameterRandomize.h
--- a/Engine/Plugins/Chooser/Source/Chooser/Public/IChooserParameterRandomize.h
+++ b/Engine/Plugins/Chooser/Source/Chooser/Public/IChooserParameterRandomize.h
@@ -11,6 +11,11 @@
 
 struct FRandomizationState
 {
+	// A ring buffer of recently selected rows (size determined by the size of FRandomizeColumn::RepeatProbabilityMultipliers)
+	// Where each entry is a row index into the ChooserTable, which was recently selected.
+	TArray<int32, TInlineAllocator<4>> RecentlySelectedRows;
+
+	// Index into RecentlySelectedRows marking the one most recently selected entry (the next most recently selected entry is at (LastSelectedRowIndex - 1 + Num) % Num
 	int LastSelectedRow;
 };
 
@@ -26,6 +31,6 @@ USTRUCT()
 struct FChooserParameterRandomizeBase : public FChooserParameterBase
 {
 	GENERATED_BODY()
-	virtual bool GetValue(FChooserEvaluationContext& Context, const FChooserRandomizationContext*& OutResult) const { return false; }
+	virtual bool GetValue(FChooserEvaluationContext& Context, FChooserRandomizationContext*& OutResult) const { return false; }
 	virtual bool IsBound() const { return false; }
 };
