Fix for garbage collection related crashes in chooser table editor

diff --git a/Engine/Plugins/Chooser/Source/ChooserEditor/Private/ChooserTableEditor.cpp b/Engine/Plugins/Chooser/Source/ChooserEditor/Private/ChooserTableEditor.cpp
--- a/Engine/Plugins/Chooser/Source/ChooserEditor/Private/ChooserTableEditor.cpp
+++ b/Engine/Plugins/Chooser/Source/ChooserEditor/Private/ChooserTableEditor.cpp
@@ -200,20 +200,6 @@ const FName FChooserTableEditor::ChooserEditorAppIdentifier( TEXT( "ChooserEdito
 
 FChooserTableEditor::~FChooserTableEditor()
 {
-	if (SelectedColumn)
-	{
-		SelectedColumn->ClearFlags(RF_Standalone);
-		SelectedColumn = nullptr;
-	}
-	for (UObject* SelectedRow : SelectedRowsCache)
-	{
-		if (SelectedRow)
-		{
-			SelectedRow->ClearFlags(RF_Standalone);
-		}
-	}
-	SelectedRows.Empty();
-
 	FTSTicker::GetCoreTicker().RemoveTicker(TickDelegateHandle);
 	
 	FCoreUObjectDelegates::OnObjectsReplaced.RemoveAll(this);
@@ -1261,14 +1247,11 @@ void FChooserTableEditor::RefreshRowSelectionDetails()
 		{
 			if (SelectedRowsCache[SelectedRowsCacheIndex] == nullptr)
 			{
-				SelectedRowsCache[SelectedRowsCacheIndex] = NewObject<UChooserRowDetails>(GetTransientPackage(), FName("SelectedRow",SelectedRowsCacheIndex),RF_Standalone | RF_Transactional | RF_Transient);
+				SelectedRowsCache[SelectedRowsCacheIndex] = TStrongObjectPtr(NewObject<UChooserRowDetails>(GetTransientPackage(), NAME_None, RF_Transactional | RF_Transient));
 			}
-			TObjectPtr<UChooserRowDetails> Selection = SelectedRowsCache[SelectedRowsCacheIndex];
+			TObjectPtr<UChooserRowDetails> Selection = SelectedRowsCache[SelectedRowsCacheIndex].Get();
 			SelectedRowsCacheIndex ++;
 			
-			check(Selection->HasAnyFlags(RF_Standalone));
-			check(!Selection->HasAnyFlags(RF_MirroredGarbage))
-			
 			Selection->Chooser = Chooser;
 			Selection->Row = SelectedItem->RowIndex;
 			Selection->Properties.Reset();
@@ -1833,12 +1816,12 @@ void FChooserTableEditor::SelectColumn(UChooserTable* ChooserEditor, int Index)
    	{
    		if (SelectedColumn == nullptr)
    		{
-   			SelectedColumn = NewObject<UChooserColumnDetails>(GetTransientPackage(), "ColumnSelection", RF_Standalone | RF_Transactional | RF_Transient);
+   			SelectedColumn = TStrongObjectPtr(NewObject<UChooserColumnDetails>(GetTransientPackage(), NAME_None, RF_Transactional | RF_Transient));
    		}
    
    		SelectedColumn->Chooser = Chooser;
    		SelectedColumn->Column = Index;
-   		DetailsView->SetObject(SelectedColumn, true);
+   		DetailsView->SetObject(SelectedColumn.Get(), true);
    		CurrentSelectionType = ESelectionType::Column;
    	}
    	else
@@ -1853,7 +1836,7 @@ void FChooserTableEditor::ClearSelectedColumn()
 	if (SelectedColumn != nullptr)
 	{
 		SelectedColumn->Column = -1;
-		if (DetailsView->GetSelectedObjects().Contains(SelectedColumn))
+		if (DetailsView->GetSelectedObjects().Contains(SelectedColumn.Get()))
 		{
 			SelectRootProperties();
 		}
diff --git a/Engine/Plugins/Chooser/Source/ChooserEditor/Private/ChooserTableEditor.h b/Engine/Plugins/Chooser/Source/ChooserEditor/Private/ChooserTableEditor.h
--- a/Engine/Plugins/Chooser/Source/ChooserEditor/Private/ChooserTableEditor.h
+++ b/Engine/Plugins/Chooser/Source/ChooserEditor/Private/ChooserTableEditor.h
@@ -220,10 +220,10 @@ namespace UE::ChooserEditor
 		/** The root chooser asset being edited in this editor */
 		UChooserTable* RootChooser = nullptr;
 
-		UChooserColumnDetails* SelectedColumn = nullptr;
+		TStrongObjectPtr<UChooserColumnDetails> SelectedColumn;
 		TArray<TObjectPtr<UChooserRowDetails>> SelectedRows;
 		// cache of objects to reuse for SelectedRows
-		TArray<TObjectPtr<UChooserRowDetails>> SelectedRowsCache;
+		TArray<TStrongObjectPtr<UChooserRowDetails>> SelectedRowsCache;
 
 		TSharedPtr<SBreadcrumbTrail<UChooserTable*>> BreadcrumbTrail;
 		TRingBuffer<UChooserTable*> History;
