Add AnimInstance threading sync to blueprint chooser evaluations

diff --git a/Engine/Plugins/Chooser/Source/Chooser/Private/ChooserFunctionLibrary.cpp b/Engine/Plugins/Chooser/Source/Chooser/Private/ChooserFunctionLibrary.cpp
--- a/Engine/Plugins/Chooser/Source/Chooser/Private/ChooserFunctionLibrary.cpp
+++ b/Engine/Plugins/Chooser/Source/Chooser/Private/ChooserFunctionLibrary.cpp
@@ -2,9 +2,30 @@
 #include "ChooserFunctionLibrary.h"
 #include "Blueprint/BlueprintExceptionInfo.h"
 #include "Chooser.h"
+#include "ChooserPropertyAccess.h"
+#include "Animation/AnimInstance.h"
 
 #include UE_INLINE_GENERATED_CPP_BY_NAME(ChooserFunctionLibrary)
 
+
+namespace
+{
+	void AnimInstanceSync(UAnimInstance* AnimInstance)
+	{
+		class UFinishFlyingAnimInstanceTasks : public UAnimInstance
+		{
+		public:
+			static void Execute(UAnimInstance* AnimInstance)
+			{
+				check(AnimInstance);
+				static_cast<UFinishFlyingAnimInstanceTasks*>(AnimInstance)->GetProxyOnAnyThread<FAnimInstanceProxy>();
+			}
+		};
+		UFinishFlyingAnimInstanceTasks::Execute(AnimInstance);
+	}
+}
+
+
 UChooserFunctionLibrary::UChooserFunctionLibrary(const FObjectInitializer& ObjectInitializer)
 	: Super(ObjectInitializer)
 {
@@ -128,6 +149,11 @@ DEFINE_FUNCTION(UChooserFunctionLibrary::execGetChooserStructOutput)
 UObject* UChooserFunctionLibrary::EvaluateObjectChooserBase(FChooserEvaluationContext& Context, const FInstancedStruct& ObjectChooser, TSubclassOf<UObject> ObjectClass, bool bResultIsClass)
 {
  	UObject* Result = nullptr;
+	
+	if (UAnimInstance* AnimInstance = Cast<UAnimInstance>(Context.GetFirstObjectParam()))
+	{
+		AnimInstanceSync(AnimInstance);
+	}
 
 	if (const FObjectChooserBase* ObjectChooserPtr = ObjectChooser.GetPtr<FObjectChooserBase>())
 	{
@@ -166,6 +192,11 @@ UObject* UChooserFunctionLibrary::EvaluateObjectChooserBase(FChooserEvaluationCo
 TSoftObjectPtr<UObject> UChooserFunctionLibrary::EvaluateObjectChooserBaseSoft(FChooserEvaluationContext& Context, const FInstancedStruct& ObjectChooser, TSubclassOf<UObject> ObjectClass, bool bResultIsClass)
 {
 	TSoftObjectPtr<UObject> Result;
+	
+	if (UAnimInstance* AnimInstance = Cast<UAnimInstance>(Context.GetFirstObjectParam()))
+	{
+		AnimInstanceSync(AnimInstance);
+	}
 
 	if (const FObjectChooserBase* ObjectChooserPtr = ObjectChooser.GetPtr<FObjectChooserBase>())
 	{
@@ -180,6 +211,11 @@ TSoftObjectPtr<UObject> UChooserFunctionLibrary::EvaluateObjectChooserBaseSoft(F
 TArray<UObject*> UChooserFunctionLibrary::EvaluateObjectChooserBaseMulti(FChooserEvaluationContext& Context, const FInstancedStruct& ObjectChooser, TSubclassOf<UObject> ObjectClass, bool bResultIsClass)
 {
  	TArray<UObject*> Result;
+	
+	if (UAnimInstance* AnimInstance = Cast<UAnimInstance>(Context.GetFirstObjectParam()))
+	{
+		AnimInstanceSync(AnimInstance);
+	}
 
 	if (const FObjectChooserBase* ObjectChooserPtr = ObjectChooser.GetPtr<FObjectChooserBase>())
 	{
@@ -215,6 +251,11 @@ TArray<UObject*> UChooserFunctionLibrary::EvaluateObjectChooserBaseMulti(FChoose
 TArray<TSoftObjectPtr<UObject>> UChooserFunctionLibrary::EvaluateObjectChooserBaseMultiSoft(FChooserEvaluationContext& Context, const FInstancedStruct& ObjectChooser, TSubclassOf<UObject> ObjectClass, bool bResultIsClass)
 {
  	TArray<TSoftObjectPtr<UObject>> Result;
+	
+	if (UAnimInstance* AnimInstance = Cast<UAnimInstance>(Context.GetFirstObjectParam()))
+	{
+		AnimInstanceSync(AnimInstance);
+	}
 
 	if (const FObjectChooserBase* ObjectChooserPtr = ObjectChooser.GetPtr<FObjectChooserBase>())
 	{
