Fix chooser parameter search incorrectly reporting that unindexed assets need to be loaded

diff --git a/Engine/Plugins/Chooser/Source/ChooserEditor/Private/ChooserFindProperties.cpp b/Engine/Plugins/Chooser/Source/ChooserEditor/Private/ChooserFindProperties.cpp
--- a/Engine/Plugins/Chooser/Source/ChooserEditor/Private/ChooserFindProperties.cpp
+++ b/Engine/Plugins/Chooser/Source/ChooserEditor/Private/ChooserFindProperties.cpp
@@ -86,10 +86,7 @@ bool UChooserFindProperties::ShouldFilterOutAsset(const FAssetData& InAssetData,
 			}
 		}, UE::String::EParseTokensOptions::SkipEmpty);
 
-		if(bFoundMatch)
-		{
-			return false;
-		}
+		return !bFoundMatch;
 	}
 	
 	bOutIsOldAsset = true;
