Fix for BP compile errors with EvaluateChooser and EvaluateProxy when using Class outputs, and "All Results" Array output mode

diff --git a/Engine/Plugins/Chooser/Source/ChooserUncooked/Private/EvaluateChooserNode.cpp b/Engine/Plugins/Chooser/Source/ChooserUncooked/Private/EvaluateChooserNode.cpp
--- a/Engine/Plugins/Chooser/Source/ChooserUncooked/Private/EvaluateChooserNode.cpp
+++ b/Engine/Plugins/Chooser/Source/ChooserUncooked/Private/EvaluateChooserNode.cpp
@@ -498,7 +498,8 @@ void UK2Node_EvaluateChooser2::AllocateDefaultPins()
 			if (Chooser->OutputObjectType)
 			{
 				ResultType = Chooser->OutputObjectType;
-				bResultIsClass = Chooser->ResultType == EObjectChooserResultType::ClassResult;
+				// For class results, in Array mode, we can't easily do the CastToClass - we have to just return an Array of Objects
+				bResultIsClass = Chooser->ResultType == EObjectChooserResultType::ClassResult && (Mode != EEvaluateChooserMode::AllResults);
 			}
 
 			if (ResultPin)
@@ -820,9 +821,10 @@ void UK2Node_EvaluateChooser2::ExpandNode(class FKismetCompilerContext& Compiler
 				// this ensures that the function does the right kind of type validation
 				CallFunction->GetSchema()->TrySetDefaultValue(*ResultIsClassPin, Chooser->ResultType == EObjectChooserResultType::ClassResult ? TEXT("true") : TEXT("false"));
 
-				if (Chooser->ResultType == EObjectChooserResultType::ClassResult)
+				// For class results, in Array mode, we can't easily do the CastToClass - we have to just return an Array of Objects
+				if (Chooser->ResultType == EObjectChooserResultType::ClassResult && Mode != EEvaluateChooserMode::AllResults)
 				{
-					// if it's a class we need to add a CastToClass function to cast it from an Object pointer to a Class poointer
+					// if it's a class we need to add a CastToClass function to cast it from an Object pointer to a Class pointer
 					UK2Node_CallFunction* CastToClass = CompilerContext.SpawnIntermediateNode<UK2Node_CallFunction>(this, SourceGraph);
 					CompilerContext.MessageLog.NotifyIntermediateObjectCreation(CastToClass, this);
 
diff --git a/Engine/Plugins/Chooser/Source/ProxyTableUncooked/Private/EvaluateProxyNode.cpp b/Engine/Plugins/Chooser/Source/ProxyTableUncooked/Private/EvaluateProxyNode.cpp
--- a/Engine/Plugins/Chooser/Source/ProxyTableUncooked/Private/EvaluateProxyNode.cpp
+++ b/Engine/Plugins/Chooser/Source/ProxyTableUncooked/Private/EvaluateProxyNode.cpp
@@ -390,6 +390,12 @@ void UK2Node_EvaluateProxy2::AllocateDefaultPins()
 
 	if (Proxy)
 	{
+		// We show input struct contexts pins as "references", because the supplied struct might get modified if it's an In/Out param.
+		// We also show output struct contexts pins as references, to show they are there just for convenience and actually pointing to the very same
+		// struct that the one on the input side.
+		FCreatePinParams StructContextPinCreationParams;
+		StructContextPinCreationParams.bIsReference = true;
+
 		for (FInstancedStruct& ContextDataEntry : Proxy->ContextData)
 		{
 			if (ContextDataEntry.IsValid())
@@ -419,7 +425,7 @@ void UK2Node_EvaluateProxy2::AllocateDefaultPins()
 						}
 						else
 						{
-							CreatePin(EGPD_Input, UEdGraphSchema_K2::PC_Struct, StructContext.Struct, StructContext.Struct.GetFName());
+							CreatePin(EGPD_Input, UEdGraphSchema_K2::PC_Struct, StructContext.Struct, StructContext.Struct.GetFName(), StructContextPinCreationParams);
 						}
 					}
 
@@ -431,7 +437,7 @@ void UK2Node_EvaluateProxy2::AllocateDefaultPins()
 						}
 						else
 						{
-							CreatePin(EGPD_Output, UEdGraphSchema_K2::PC_Struct, StructContext.Struct, StructContext.Struct.GetFName());
+							CreatePin(EGPD_Output, UEdGraphSchema_K2::PC_Struct, StructContext.Struct, StructContext.Struct.GetFName(), StructContextPinCreationParams);
 						}
 					}
 
@@ -461,7 +467,8 @@ void UK2Node_EvaluateProxy2::AllocateDefaultPins()
    	if (Proxy && Proxy->Type)
    	{
    		ResultType = Proxy->Type;
-   		bResultIsClass = Proxy->ResultType == EObjectChooserResultType::ClassResult;
+   		// For class results, in Array mode, we can't easily do the CastToClass - we have to just return an Array of Objects
+   		bResultIsClass = Proxy->ResultType == EObjectChooserResultType::ClassResult && Mode != EEvaluateProxyMode::AllResults;
    	}
 	
 	static const FName ProxyTablePinName = "ProxyTable";
@@ -766,9 +773,10 @@ void UK2Node_EvaluateProxy2::ExpandNode(class FKismetCompilerContext& CompilerCo
 				// this ensures that the function does the right kind of type validation
 				CallFunction->GetSchema()->TrySetDefaultValue(*ResultIsClassPin, Proxy->ResultType == EObjectChooserResultType::ClassResult ? TEXT("true") : TEXT("false"));
 
-				if (Proxy->ResultType == EObjectChooserResultType::ClassResult)
+				// For class results, in Array mode, we can't easily do the CastToClass - we have to just return an Array of Objects
+				if (Proxy->ResultType == EObjectChooserResultType::ClassResult && Mode != EEvaluateProxyMode::AllResults)
 				{
-					// if it's a class we need to add a CastToClass function to cast it from an Object pointer to a Class poointer
+					// if it's a class we need to add a CastToClass function to cast it from an Object pointer to a Class pointer
 					UK2Node_CallFunction* CastToClass = CompilerContext.SpawnIntermediateNode<UK2Node_CallFunction>(this, SourceGraph);
 					CompilerContext.MessageLog.NotifyIntermediateObjectCreation(CastToClass, this);
 
