Chooser - fix for crash when changing column type from Pose Match to a different column and then back to PoseMatch

#[review]-45980261 [REDACTED_HANDLE]

diff --git a/Engine/Plugins/Chooser/Source/ChooserEditor/Private/ChooserDetails.cpp b/Engine/Plugins/Chooser/Source/ChooserEditor/Private/ChooserDetails.cpp
--- a/Engine/Plugins/Chooser/Source/ChooserEditor/Private/ChooserDetails.cpp
+++ b/Engine/Plugins/Chooser/Source/ChooserEditor/Private/ChooserDetails.cpp
@@ -83,14 +83,15 @@ void FChooserColumnDetails::CustomizeDetails(IDetailLayoutBuilder& DetailBuilder
 	
 		TSharedPtr<IPropertyHandle> ColumnsArrayProperty = ChooserProperty->GetChildHandle("ColumnsStructs");
 		TSharedPtr<IPropertyHandle> CurrentColumnProperty = ColumnsArrayProperty->AsArray()->GetElement(Column->Column);
+
 		IDetailPropertyRow& NewColumnProperty = PropertiesCategory.AddProperty(CurrentColumnProperty);
-		NewColumnProperty.DisplayName(LOCTEXT("Selected Column","Selected Column"));
-		NewColumnProperty.ShowPropertyButtons(false); // hide array add button
-		NewColumnProperty.ShouldAutoExpand(true);
+		// hide array add button
+		NewColumnProperty.ShowPropertyButtons(false);
+		// removing the column type, and just showing the column instance struct properties
+		NewColumnProperty.CustomWidget(true);
 	}
 }
 
-
 }
 
 #undef LOCTEXT_NAMESPACE
