Fix for nested choosers not respecting root chooser Result Type

diff --git a/Engine/Plugins/Chooser/Source/ChooserEditor/Private/ChooserTableEditor.cpp b/Engine/Plugins/Chooser/Source/ChooserEditor/Private/ChooserTableEditor.cpp
--- a/Engine/Plugins/Chooser/Source/ChooserEditor/Private/ChooserTableEditor.cpp
+++ b/Engine/Plugins/Chooser/Source/ChooserEditor/Private/ChooserTableEditor.cpp
@@ -1558,7 +1558,7 @@ TSharedRef<SWidget>	FChooserTableEditor::MakeCreateRowMenu()
 					FNewMenuDelegate::CreateLambda([this](FMenuBuilder& MenuBuilder)
 					{
 						UChooserTable* Chooser = GetChooser();
-						MakeCreateResultMenu(MenuBuilder, Chooser->ResultType, FCreateStructDelegate::CreateLambda([this](UScriptStruct* Type)
+						MakeCreateResultMenu(MenuBuilder, Chooser->GetContextOwner()->ResultType, FCreateStructDelegate::CreateLambda([this](UScriptStruct* Type)
 						{
 							UChooserTable* Chooser = GetChooser();
 							const FScopedTransaction Transaction(LOCTEXT("Add Fallback Row Transaction", "Add Fallback Row"));
@@ -1597,7 +1597,7 @@ TSharedRef<SWidget>	FChooserTableEditor::MakeCreateRowMenu()
 	}
 	else
 	{
-		MakeCreateResultMenu(MenuBuilder, Chooser->ResultType, FCreateStructDelegate::CreateLambda([this](UScriptStruct* Type)
+		MakeCreateResultMenu(MenuBuilder, Chooser->GetContextOwner()->ResultType, FCreateStructDelegate::CreateLambda([this](UScriptStruct* Type)
 		{
 			UChooserTable* Chooser = GetChooser();
 			const FScopedTransaction Transaction(LOCTEXT("Add Row Transaction", "Add Row"));
