Allow creating ChooserSignature assets

diff --git a/Engine/Plugins/Chooser/Source/ChooserEditor/Private/AssetDefinition_ChooserSignature.cpp b/Engine/Plugins/Chooser/Source/ChooserEditor/Private/AssetDefinition_ChooserSignature.cpp
new file mode 100644
--- /dev/null
+++ b/Engine/Plugins/Chooser/Source/ChooserEditor/Private/AssetDefinition_ChooserSignature.cpp
@@ -0,0 +1,30 @@
+// Copyright Epic Games, Inc. All Rights Reserved.
+
+#include "AssetDefinition_ChooserSignature.h"
+#include "Styling/SlateStyleRegistry.h"
+#include "Toolkits/SimpleAssetEditor.h"
+#include "ChooserEditorStyle.h"
+
+#include UE_INLINE_GENERATED_CPP_BY_NAME(AssetDefinition_ChooserSignature)
+
+#define LOCTEXT_NAMESPACE "UAssetDefinition_ChooserSignature" 
+
+const FSlateBrush* UAssetDefinition_ChooserSignature::GetThumbnailBrush(const FAssetData& InAssetData, const FName InClassName) const
+{
+	return UE::ChooserEditor::FChooserEditorStyle::Get().GetBrush("ChooserEditor.ChooserTableIconLarge");
+}
+
+const FSlateBrush* UAssetDefinition_ChooserSignature::GetIconBrush(const FAssetData& InAssetData, const FName InClassName) const
+{
+	return UE::ChooserEditor::FChooserEditorStyle::Get().GetBrush("ChooserEditor.ChooserTableIconSmall");
+}
+
+EAssetCommandResult UAssetDefinition_ChooserSignature::OpenAssets(const FAssetOpenArgs& OpenArgs) const
+{
+	TArray<UObject*> Objects = OpenArgs.LoadObjects<UObject>();
+	FSimpleAssetEditor::CreateEditor(OpenArgs.GetToolkitMode(), OpenArgs.ToolkitHost, Objects);
+
+	return EAssetCommandResult::Handled;
+}
+
+#undef LOCTEXT_NAMESPACE
diff --git a/Engine/Plugins/Chooser/Source/ChooserEditor/Private/AssetDefinition_ChooserSignature.h b/Engine/Plugins/Chooser/Source/ChooserEditor/Private/AssetDefinition_ChooserSignature.h
new file mode 100644
--- /dev/null
+++ b/Engine/Plugins/Chooser/Source/ChooserEditor/Private/AssetDefinition_ChooserSignature.h
@@ -0,0 +1,29 @@
+// Copyright Epic Games, Inc. All Rights Reserved.
+
+#pragma once
+
+#include "Chooser.h"
+#include "AssetDefinitionDefault.h"
+
+#include "AssetDefinition_ChooserSignature.generated.h"
+
+UCLASS()
+class UAssetDefinition_ChooserSignature : public UAssetDefinitionDefault
+{
+	GENERATED_BODY()
+
+public:
+	// UAssetDefinition Begin
+	virtual FText GetAssetDisplayName() const override { return NSLOCTEXT("AssetDefinition_ChooserSignature", "ChooserSignatureDisplayName", "Chooser Signature"); }
+	virtual FLinearColor GetAssetColor() const override { return FLinearColor(FColor(128,128,64)); }
+	virtual const FSlateBrush* GetIconBrush(const FAssetData& InAssetData, const FName InClassName) const override;
+	virtual const FSlateBrush* GetThumbnailBrush(const FAssetData& InAssetData, const FName InClassName) const override;
+	virtual TSoftClassPtr<UObject> GetAssetClass() const override { return UChooserSignature::StaticClass(); }
+	virtual TConstArrayView<FAssetCategoryPath> GetAssetCategories() const override
+	{
+		static const auto Categories = { FAssetCategoryPath(EAssetCategoryPaths::Misc) , FAssetCategoryPath(EAssetCategoryPaths::Animation) };
+		return Categories;
+	}
+	virtual EAssetCommandResult OpenAssets(const FAssetOpenArgs& OpenArgs) const override;
+	// UAssetDefinition End
+};
diff --git a/Engine/Plugins/Chooser/Source/ChooserEditor/Private/ChooserFactory.cpp b/Engine/Plugins/Chooser/Source/ChooserEditor/Private/ChooserFactory.cpp
--- a/Engine/Plugins/Chooser/Source/ChooserEditor/Private/ChooserFactory.cpp
+++ b/Engine/Plugins/Chooser/Source/ChooserEditor/Private/ChooserFactory.cpp
@@ -7,168 +7,16 @@
 #include "Framework/Application/SlateApplication.h"
 #include "Widgets/Input/SButton.h"
 #include "Widgets/Layout/SUniformGridPanel.h"
-#include "Widgets/SCompoundWidget.h"
 #include "Widgets/SWindow.h"
 #include "Editor.h"
 #include "ObjectChooserClassFilter.h"
 #include "StructViewerModule.h"
 #include "PropertyEditorModule.h"
 #include "Modules/ModuleManager.h"
+#include "SChooserCreateDialog.h"
 
 #include UE_INLINE_GENERATED_CPP_BY_NAME(ChooserFactory)
 
-#define LOCTEXT_NAMESPACE "CreateChooser"
-
-/*------------------------------------------------------------------------------
-	Dialog to configure creation properties
-------------------------------------------------------------------------------*/
-class SChooserCreateDialog : public SCompoundWidget
-{
-public:
-	SLATE_BEGIN_ARGS( SChooserCreateDialog ){}
-
-	SLATE_END_ARGS()
-
-	/** Constructs this widget with InArgs */
-	void Construct( const FArguments& InArgs )
-	{
-		bOkClicked = false;
-		
-		FPropertyEditorModule& PropertyEditorModule = FModuleManager::GetModuleChecked<FPropertyEditorModule>( "PropertyEditor" );
-		FDetailsViewArgs DetailsViewArgs;
-		// DetailsViewArgs.NotifyHook = this;
-		DetailsViewArgs.NameAreaSettings = FDetailsViewArgs::HideNameArea;
-		DetailsViewArgs.bAllowSearch = false;
-		DetailsViewArgs.bShowOptions = false;
-		DetailsViewArgs.bAllowFavoriteSystem = false;
-		DetailsViewArgs.bShowLooseProperties = true;
-		DetailsView = PropertyEditorModule.CreateDetailView( DetailsViewArgs );
-		DetailsView->SetObject(ChooserFactory.Get());
-
-		ChildSlot
-		[
-			SNew(SVerticalBox)
-			+SVerticalBox::Slot().FillHeight(1)
-			[
-				DetailsView.ToSharedRef()
-			]
-
-			// Ok/Cancel buttons
-			+SVerticalBox::Slot()
-			.AutoHeight()
-			.HAlign(HAlign_Right)
-			.VAlign(VAlign_Bottom)
-			.Padding(10.0f)
-			[
-				SNew(SUniformGridPanel)
-				.SlotPadding(FAppStyle::GetMargin("StandardDialog.SlotPadding"))
-				.MinDesiredSlotWidth(FAppStyle::GetFloat("StandardDialog.MinDesiredSlotWidth"))
-				.MinDesiredSlotHeight(FAppStyle::GetFloat("StandardDialog.MinDesiredSlotHeight"))
-				+SUniformGridPanel::Slot(0,0)
-				[
-					SNew(SButton)
-					.ToolTipText(LOCTEXT("CreateAnimBlueprintCreate_Tooltip", "Create a new Chooser Table Asset.."))
-					.IsEnabled_Lambda([this]()
-					{
-						return true;
-					})
-					.HAlign(HAlign_Center)
-					.ContentPadding( FAppStyle::GetMargin("StandardDialog.ContentPadding") )
-					.OnClicked(this, &SChooserCreateDialog::OkClicked)
-					.IsEnabled_Lambda([this]() { return ChooserFactory->ChooserInitializer.IsValid(); } )
-					.Text(LOCTEXT("Create", "Create"))
-				]
-				+SUniformGridPanel::Slot(1,0)
-				[
-					SNew(SButton)
-					.HAlign(HAlign_Center)
-					.ContentPadding( FAppStyle::GetMargin("StandardDialog.ContentPadding") )
-					.OnClicked(this, &SChooserCreateDialog::CancelClicked)
-					.Text(LOCTEXT("Cancel", "Cancel"))
-				]
-			]
-		];
-	}
-	
-	/** Sets properties for the supplied AnimBlueprintFactory */
-	bool ConfigureProperties(TWeakObjectPtr<UChooserTableFactory> InChooserFactory)
-	{
-		ChooserFactory = InChooserFactory;
-		if (!UChooserEditorSettings::Get().DefaultCreateType.IsEmpty())
-		{
-			FTopLevelAssetPath StructPath(UChooserEditorSettings::Get().DefaultCreateType);
-			if (UScriptStruct* DefaultInitializer = FindObject<UScriptStruct>(StructPath))
-			{
-				ChooserFactory->ChooserInitializer.InitializeAs(DefaultInitializer);
-			}
-		}
-		
-		DetailsView->SetObject(ChooserFactory.Get());
-		
-		const float AppScale = FSlateApplication::Get().GetApplicationScale();
-
-		Window = SNew(SWindow)
-		.Title( LOCTEXT("Create Chooser Options", "Create Chooser Table") )
-		.SizingRule(ESizingRule::FixedSize)
-		.ClientSize(FVector2f(AppScale*550, AppScale*300))
-		.SupportsMinimize(false)
-		.SupportsMaximize(false)
-		[
-			AsShared()
-		];
-
-		GEditor->EditorAddModalWindow(Window.ToSharedRef());
-		ChooserFactory.Reset();
-
-		return bOkClicked;
-	}
-
-private:
-
-	/** Handler for when ok is clicked */
-	FReply OkClicked()
-	{
-		CloseDialog(true);
-
-		return FReply::Handled();
-	}
-
-	void CloseDialog(bool bWasPicked=false)
-	{
-		bOkClicked = bWasPicked;
-		if (Window.IsValid())
-		{
-			Window->RequestDestroyWindow();
-		}
-	}
-
-	/** Handler for when cancel is clicked */
-	FReply CancelClicked()
-	{
-		CloseDialog();
-		return FReply::Handled();
-	}
-
-	FReply OnKeyDown( const FGeometry& MyGeometry, const FKeyEvent& InKeyEvent )
-	{
-		if (InKeyEvent.GetKey() == EKeys::Escape)
-		{
-			CloseDialog();
-			return FReply::Handled();
-		}
-		return SWidget::OnKeyDown(MyGeometry, InKeyEvent);
-	}
-
-private:
-	/** The factory for which we are setting up properties */
-	TWeakObjectPtr<UChooserTableFactory> ChooserFactory;
-	TSharedPtr<SWindow> Window;
-	TSharedPtr<IDetailsView> DetailsView;
-
-	/** True if Ok was clicked */
-	bool bOkClicked = false;
-};
-
 UChooserTableFactory::UChooserTableFactory()
 {
 	bCreateNew = true;
@@ -184,13 +32,11 @@ bool UChooserTableFactory::ConfigureProperties()
 
 UObject* UChooserTableFactory::FactoryCreateNew(UClass* Class, UObject* InParent, FName Name, EObjectFlags Flags, UObject* Context, FFeedbackContext* Warn, FName CallingContext)
 {
-	UChooserTable* NewChooser =  NewObject<UChooserTable>(InParent, Class, Name, Flags | RF_Public | RF_Standalone | RF_Transactional | RF_LoadCompleted);
+	UChooserTable* NewChooser =  NewObject<UChooserTable>(InParent, Class, Name, Flags);
 	if (ChooserInitializer.IsValid())
 	{
-		ChooserInitializer.Get<FChooserInitializer>().Initialize(NewChooser);
+		ChooserInitializer.Get<FChooserInitializer>().InitializeSignature(NewChooser);
 		UChooserEditorSettings::Get().DefaultCreateType = ChooserInitializer.GetScriptStruct()->GetStructPathName().ToString();
 	}
 	return NewChooser;
-}
-
-#undef LOCTEXT_NAMESPACE
+}
\ No newline at end of file
diff --git a/Engine/Plugins/Chooser/Source/ChooserEditor/Private/ChooserFactory.h b/Engine/Plugins/Chooser/Source/ChooserEditor/Private/ChooserFactory.h
--- a/Engine/Plugins/Chooser/Source/ChooserEditor/Private/ChooserFactory.h
+++ b/Engine/Plugins/Chooser/Source/ChooserEditor/Private/ChooserFactory.h
@@ -3,13 +3,12 @@
 #pragma once
 
 #include "CoreMinimal.h"
-#include "Factories/Factory.h"
-#include "StructUtils/InstancedStruct.h"
+#include "ChooserSignatureFactory.h"
 
 #include "ChooserFactory.generated.h"
 
 UCLASS(MinimalAPI)
-class UChooserTableFactory : public UFactory
+class UChooserTableFactory : public UChooserSignatureFactory
 {
 	GENERATED_BODY()
 
@@ -23,7 +22,4 @@ public:
 		return FactoryCreateNew(Class, InParent, Name, Flags, Context, Warn, NAME_None);
 	}
 	virtual FName GetNewAssetThumbnailOverride() const override { return "ChooserEditor.ChooserTableIconLarge"; }
-
-	UPROPERTY(EditAnywhere, DisplayName = "Chooser Type", NoClear, Category="Chooser Setup", meta = (ExcludeBaseStruct, BaseStruct = "/Script/ChooserEditor.ChooserInitializer"))
-	FInstancedStruct ChooserInitializer;
 };
\ No newline at end of file
diff --git a/Engine/Plugins/Chooser/Source/ChooserEditor/Private/ChooserInitializer.cpp b/Engine/Plugins/Chooser/Source/ChooserEditor/Private/ChooserInitializer.cpp
--- a/Engine/Plugins/Chooser/Source/ChooserEditor/Private/ChooserInitializer.cpp
+++ b/Engine/Plugins/Chooser/Source/ChooserEditor/Private/ChooserInitializer.cpp
@@ -13,14 +13,14 @@
 #define LOCTEXT_NAMESPACE "ChooserInitializer"
 
 
-void FGenericChooserInitializer::Initialize(UChooserTable* Chooser) const
+void FGenericChooserInitializer::InitializeSignature(UChooserSignature* Chooser) const
 {
 	Chooser->ContextData = ContextData;
 	Chooser->ResultType = ResultType;
 	Chooser->OutputObjectType = OutputObjectType;
 }
 
-void FChooserPlayerInitializer::Initialize(UChooserTable* Chooser) const
+void FChooserPlayerInitializer::InitializeSignature(UChooserSignature* Chooser) const
 {
 	Chooser->ContextData.SetNum(2);
 	Chooser->ContextData[0].InitializeAs(FContextObjectTypeClass::StaticStruct());
@@ -43,7 +43,7 @@ void FChooserPlayerInitializer::Initialize(UChooserTable* Chooser) const
 	Chooser->OutputObjectType = UAnimationAsset::StaticClass();
 }
 
-void FNoPrimaryResultChooserInitializer::Initialize(UChooserTable* Chooser) const
+void FNoPrimaryResultChooserInitializer::InitializeSignature(UChooserSignature* Chooser) const
 {
 	Chooser->ContextData = ContextData;
 	Chooser->ResultType = EObjectChooserResultType::NoPrimaryResult;
diff --git a/Engine/Plugins/Chooser/Source/ChooserEditor/Private/ChooserSignatureFactory.cpp b/Engine/Plugins/Chooser/Source/ChooserEditor/Private/ChooserSignatureFactory.cpp
new file mode 100644
--- /dev/null
+++ b/Engine/Plugins/Chooser/Source/ChooserEditor/Private/ChooserSignatureFactory.cpp
@@ -0,0 +1,42 @@
+﻿// Copyright Epic Games, Inc. All Rights Reserved.
+
+#include "ChooserSignatureFactory.h"
+#include "Chooser.h"
+#include "ChooserInitializer.h"
+#include "ChooserEditorSettings.h"
+#include "Framework/Application/SlateApplication.h"
+#include "Widgets/Input/SButton.h"
+#include "Widgets/Layout/SUniformGridPanel.h"
+#include "Widgets/SCompoundWidget.h"
+#include "Widgets/SWindow.h"
+#include "Editor.h"
+#include "ObjectChooserClassFilter.h"
+#include "StructViewerModule.h"
+#include "Modules/ModuleManager.h"
+#include "SChooserCreateDialog.h"
+
+#include UE_INLINE_GENERATED_CPP_BY_NAME(ChooserSignatureFactory)
+
+UChooserSignatureFactory::UChooserSignatureFactory()
+{
+	bCreateNew = true;
+	bEditAfterNew = true;
+	SupportedClass = UChooserSignature::StaticClass();
+}
+
+bool UChooserSignatureFactory::ConfigureProperties()
+{
+	TSharedRef<SChooserCreateDialog> Dialog = SNew(SChooserCreateDialog);
+	return Dialog->ConfigureProperties(this);
+}
+
+UObject* UChooserSignatureFactory::FactoryCreateNew(UClass* Class, UObject* InParent, FName Name, EObjectFlags Flags, UObject* Context, FFeedbackContext* Warn, FName CallingContext)
+{
+	UChooserSignature* NewChooserSignature =  NewObject<UChooserSignature>(InParent, Class, Name, Flags);
+	if (ChooserInitializer.IsValid())
+	{
+		ChooserInitializer.Get<FChooserInitializer>().InitializeSignature(NewChooserSignature);
+		UChooserEditorSettings::Get().DefaultCreateType = ChooserInitializer.GetScriptStruct()->GetStructPathName().ToString();
+	}
+	return NewChooserSignature;
+}
diff --git a/Engine/Plugins/Chooser/Source/ChooserEditor/Private/ChooserSignatureFactory.h b/Engine/Plugins/Chooser/Source/ChooserEditor/Private/ChooserSignatureFactory.h
new file mode 100644
--- /dev/null
+++ b/Engine/Plugins/Chooser/Source/ChooserEditor/Private/ChooserSignatureFactory.h
@@ -0,0 +1,29 @@
+﻿// Copyright Epic Games, Inc. All Rights Reserved.
+
+#pragma once
+
+#include "CoreMinimal.h"
+#include "Factories/Factory.h"
+#include "StructUtils/InstancedStruct.h"
+
+#include "ChooserSignatureFactory.generated.h"
+
+UCLASS(MinimalAPI)
+class UChooserSignatureFactory : public UFactory
+{
+	GENERATED_BODY()
+
+public:
+	UChooserSignatureFactory();
+	
+	virtual bool ConfigureProperties() override;
+	virtual UObject* FactoryCreateNew(UClass* Class, UObject* InParent, FName Name, EObjectFlags Flags, UObject* Context, FFeedbackContext* Warn, FName CallingContext) override;
+	virtual UObject* FactoryCreateNew(UClass* Class, UObject* InParent, FName Name, EObjectFlags Flags, UObject* Context, FFeedbackContext* Warn) override
+	{
+		return FactoryCreateNew(Class, InParent, Name, Flags, Context, Warn, NAME_None);
+	}
+	virtual FName GetNewAssetThumbnailOverride() const override { return "ChooserEditor.ChooserTableIconLarge"; }
+
+	UPROPERTY(EditAnywhere, DisplayName = "Chooser Type", NoClear, Category="Chooser Setup", meta = (ExcludeBaseStruct, BaseStruct = "/Script/ChooserEditor.ChooserInitializer"))
+	FInstancedStruct ChooserInitializer;
+};
\ No newline at end of file
diff --git a/Engine/Plugins/Chooser/Source/ChooserEditor/Private/SChooserCreateDialog.cpp b/Engine/Plugins/Chooser/Source/ChooserEditor/Private/SChooserCreateDialog.cpp
new file mode 100644
--- /dev/null
+++ b/Engine/Plugins/Chooser/Source/ChooserEditor/Private/SChooserCreateDialog.cpp
@@ -0,0 +1,149 @@
+﻿// Copyright Epic Games, Inc. All Rights Reserved.
+
+#include "SChooserCreateDialog.h"
+#include "Chooser.h"
+#include "ChooserInitializer.h"
+#include "ChooserEditorSettings.h"
+#include "Framework/Application/SlateApplication.h"
+#include "Widgets/Input/SButton.h"
+#include "Widgets/Layout/SUniformGridPanel.h"
+#include "Widgets/SCompoundWidget.h"
+#include "Widgets/SWindow.h"
+#include "Editor.h"
+#include "ObjectChooserClassFilter.h"
+#include "StructViewerModule.h"
+#include "PropertyEditorModule.h"
+#include "Modules/ModuleManager.h"
+
+#define LOCTEXT_NAMESPACE "CreateChooserDialog"
+
+/** Constructs this widget with InArgs */
+void SChooserCreateDialog::Construct( const FArguments& InArgs )
+{
+	bOkClicked = false;
+	
+	FPropertyEditorModule& PropertyEditorModule = FModuleManager::GetModuleChecked<FPropertyEditorModule>( "PropertyEditor" );
+	FDetailsViewArgs DetailsViewArgs;
+	// DetailsViewArgs.NotifyHook = this;
+	DetailsViewArgs.NameAreaSettings = FDetailsViewArgs::HideNameArea;
+	DetailsViewArgs.bAllowSearch = false;
+	DetailsViewArgs.bShowOptions = false;
+	DetailsViewArgs.bAllowFavoriteSystem = false;
+	DetailsViewArgs.bShowLooseProperties = true;
+	DetailsView = PropertyEditorModule.CreateDetailView( DetailsViewArgs );
+	DetailsView->SetObject(ChooserFactory.Get());
+
+	ChildSlot
+	[
+		SNew(SVerticalBox)
+		+SVerticalBox::Slot().FillHeight(1)
+		[
+			DetailsView.ToSharedRef()
+		]
+
+		// Ok/Cancel buttons
+		+SVerticalBox::Slot()
+		.AutoHeight()
+		.HAlign(HAlign_Right)
+		.VAlign(VAlign_Bottom)
+		.Padding(10.0f)
+		[
+			SNew(SUniformGridPanel)
+			.SlotPadding(FAppStyle::GetMargin("StandardDialog.SlotPadding"))
+			.MinDesiredSlotWidth(FAppStyle::GetFloat("StandardDialog.MinDesiredSlotWidth"))
+			.MinDesiredSlotHeight(FAppStyle::GetFloat("StandardDialog.MinDesiredSlotHeight"))
+			+SUniformGridPanel::Slot(0,0)
+			[
+				SNew(SButton)
+				.ToolTipText(LOCTEXT("CreateAnimBlueprintCreate_Tooltip", "Create a new Chooser Table Asset.."))
+				.IsEnabled_Lambda([this]()
+				{
+					return true;
+				})
+				.HAlign(HAlign_Center)
+				.ContentPadding( FAppStyle::GetMargin("StandardDialog.ContentPadding") )
+				.OnClicked(this, &SChooserCreateDialog::OkClicked)
+				.IsEnabled_Lambda([this]() { return ChooserFactory->ChooserInitializer.IsValid(); } )
+				.Text(LOCTEXT("Create", "Create"))
+			]
+			+SUniformGridPanel::Slot(1,0)
+			[
+				SNew(SButton)
+				.HAlign(HAlign_Center)
+				.ContentPadding( FAppStyle::GetMargin("StandardDialog.ContentPadding") )
+				.OnClicked(this, &SChooserCreateDialog::CancelClicked)
+				.Text(LOCTEXT("Cancel", "Cancel"))
+			]
+		]
+	];
+}
+
+/** Sets properties for the supplied AnimBlueprintFactory */
+bool SChooserCreateDialog::ConfigureProperties(TWeakObjectPtr<UChooserSignatureFactory> InChooserFactory)
+{
+	ChooserFactory = InChooserFactory;
+	if (!UChooserEditorSettings::Get().DefaultCreateType.IsEmpty())
+	{
+		FTopLevelAssetPath StructPath(UChooserEditorSettings::Get().DefaultCreateType);
+		if (UScriptStruct* DefaultInitializer = FindObject<UScriptStruct>(StructPath))
+		{
+			ChooserFactory->ChooserInitializer.InitializeAs(DefaultInitializer);
+		}
+	}
+	
+	DetailsView->SetObject(ChooserFactory.Get());
+	
+	const float AppScale = FSlateApplication::Get().GetApplicationScale();
+
+	Window = SNew(SWindow)
+	.Title( LOCTEXT("Create Chooser Options", "Create Chooser Table") )
+	.SizingRule(ESizingRule::FixedSize)
+	.ClientSize(FVector2f(AppScale*550, AppScale*300))
+	.SupportsMinimize(false)
+	.SupportsMaximize(false)
+	[
+		AsShared()
+	];
+
+	GEditor->EditorAddModalWindow(Window.ToSharedRef());
+	ChooserFactory.Reset();
+
+	return bOkClicked;
+}
+
+
+/** Handler for when ok is clicked */
+FReply SChooserCreateDialog::OkClicked()
+{
+	CloseDialog(true);
+
+	return FReply::Handled();
+}
+
+void SChooserCreateDialog::CloseDialog(bool bWasPicked)
+{
+	bOkClicked = bWasPicked;
+	if (Window.IsValid())
+	{
+		Window->RequestDestroyWindow();
+	}
+}
+
+/** Handler for when cancel is clicked */
+FReply SChooserCreateDialog::CancelClicked()
+{
+	CloseDialog();
+	return FReply::Handled();
+}
+
+FReply SChooserCreateDialog::OnKeyDown( const FGeometry& MyGeometry, const FKeyEvent& InKeyEvent )
+{
+	if (InKeyEvent.GetKey() == EKeys::Escape)
+	{
+		CloseDialog();
+		return FReply::Handled();
+	}
+	return SWidget::OnKeyDown(MyGeometry, InKeyEvent);
+}
+
+#undef LOCTEXT_NAMESPACE
diff --git a/Engine/Plugins/Chooser/Source/ChooserEditor/Private/SChooserCreateDialog.h b/Engine/Plugins/Chooser/Source/ChooserEditor/Private/SChooserCreateDialog.h
new file mode 100644
--- /dev/null
+++ b/Engine/Plugins/Chooser/Source/ChooserEditor/Private/SChooserCreateDialog.h
@@ -0,0 +1,46 @@
+﻿// Copyright Epic Games, Inc. All Rights Reserved.
+
+#pragma once
+
+#include "CoreMinimal.h"
+#include "Widgets/SCompoundWidget.h"
+#include "StructUtils/InstancedStruct.h"
+#include "ChooserSignatureFactory.h"
+#include "PropertyEditorModule.h"
+
+/*------------------------------------------------------------------------------
+	Dialog to configure creation properties
+------------------------------------------------------------------------------*/
+class SChooserCreateDialog : public SCompoundWidget
+{
+public:
+	SLATE_BEGIN_ARGS( SChooserCreateDialog ){}
+
+	SLATE_END_ARGS()
+
+	/** Constructs this widget with InArgs */
+	void Construct( const FArguments& InArgs );
+	
+	/** Sets properties for the supplied AnimBlueprintFactory */
+	bool ConfigureProperties(TWeakObjectPtr<UChooserSignatureFactory> InChooserFactory);
+
+private:
+
+	/** Handler for when ok is clicked */
+	FReply OkClicked();
+	
+	void CloseDialog(bool bWasPicked=false);
+	
+	/** Handler for when cancel is clicked */
+	FReply CancelClicked();
+	FReply OnKeyDown( const FGeometry& MyGeometry, const FKeyEvent& InKeyEvent );
+
+private:
+	/** The factory for which we are setting up properties */
+	TWeakObjectPtr<UChooserSignatureFactory> ChooserFactory;
+	TSharedPtr<SWindow> Window;
+	TSharedPtr<IDetailsView> DetailsView;
+
+	/** True if Ok was clicked */
+	bool bOkClicked = false;
+};
\ No newline at end of file
diff --git a/Engine/Plugins/Chooser/Source/ChooserEditor/Public/ChooserInitializer.h b/Engine/Plugins/Chooser/Source/ChooserEditor/Public/ChooserInitializer.h
--- a/Engine/Plugins/Chooser/Source/ChooserEditor/Public/ChooserInitializer.h
+++ b/Engine/Plugins/Chooser/Source/ChooserEditor/Public/ChooserInitializer.h
@@ -12,6 +12,16 @@ struct FChooserInitializer
 {
 	GENERATED_BODY()
 	virtual ~FChooserInitializer() {}
+	virtual void InitializeSignature(UChooserSignature* ChooserSignature) const
+	{
+		// fallback codepath in case there are any old Initializers out there that haven't updated to InitalizeSignature
+		if (UChooserTable* Chooser = Cast<UChooserTable>(ChooserSignature))
+		{
+			Initialize(Chooser);
+		}
+	};
+
+	// Deprecated, implement InitializeSignature
 	virtual void Initialize(UChooserTable* Chooser) const {};
 };
 
@@ -21,7 +31,7 @@ struct FGenericChooserInitializer : public FChooserInitializer
 {
 	GENERATED_BODY()
 
-	virtual void Initialize(UChooserTable* Chooser) const override;
+	virtual void InitializeSignature(UChooserSignature* ChooserSignature) const override;
 	
 	// The kind of output this chooser has (Object or Class)
 	UPROPERTY(EditAnywhere, DisplayName = "Result Type", Category="Result")
@@ -41,7 +51,7 @@ USTRUCT(DisplayName="Animation Chooser", meta = (ToolTip="A ChooserTable for use
 struct FChooserPlayerInitializer : public FChooserInitializer
 {
 	GENERATED_BODY()
-	virtual void Initialize(UChooserTable* Chooser) const override;
+	virtual void InitializeSignature(UChooserSignature* Chooser) const override;
 
 	UPROPERTY(EditAnywhere, Category="AnimClass")
 	TObjectPtr<UClass> AnimClass;
@@ -52,7 +62,7 @@ struct FNoPrimaryResultChooserInitializer : public FChooserInitializer
 {
 	GENERATED_BODY()
 
-	virtual void Initialize(UChooserTable* Chooser) const override;
+	virtual void InitializeSignature(UChooserSignature* Chooser) const override;
 	
 	// Parameter Objects or Structs from which the chooser can read or write properties 
 	UPROPERTY(EditAnywhere, DisplayName = "Parameters", NoClear, Meta = (ExpandByDefault, ExcludeBaseStruct, BaseStruct = "/Script/Chooser.ContextObjectTypeBase"), Category = "Parameters")
