Add gameplay tags output column for chooser tables

diff --git a/Engine/Plugins/Chooser/Source/Chooser/Internal/GameplayTagColumn.h b/Engine/Plugins/Chooser/Source/Chooser/Internal/GameplayTagColumn.h
--- a/Engine/Plugins/Chooser/Source/Chooser/Internal/GameplayTagColumn.h
+++ b/Engine/Plugins/Chooser/Source/Chooser/Internal/GameplayTagColumn.h
@@ -22,12 +22,13 @@ struct FGameplayTagContextProperty :  public FChooserParameterGameplayTagBase
 public:
 	UPROPERTY()
 	TArray<FName> PropertyBindingChain_DEPRECATED;
-	
+
 	UPROPERTY(EditAnywhere, NoClear, Meta = (BindingType = "FGameplayTagContainer", BindingColor = "StructPinTypeColor"), Category = "Binding")
 	FChooserPropertyBinding Binding;
-	
+
 	UE_API virtual bool GetValue(FChooserEvaluationContext& Context, const FGameplayTagContainer*& OutResult) const override;
-	
+	UE_API virtual bool SetValue(FChooserEvaluationContext& Context, const FGameplayTagContainer& Value) const override;
+
 	virtual void PostLoad() override
 	{
 		if (PropertyBindingChain_DEPRECATED.Num() > 0)
diff --git a/Engine/Plugins/Chooser/Source/Chooser/Internal/OutputGameplayTagColumn.h b/Engine/Plugins/Chooser/Source/Chooser/Internal/OutputGameplayTagColumn.h
new file mode 100644
--- /dev/null
+++ b/Engine/Plugins/Chooser/Source/Chooser/Internal/OutputGameplayTagColumn.h
@@ -0,0 +1,60 @@
+﻿// Copyright Epic Games, Inc. All Rights Reserved.
+#pragma once
+
+#include "IChooserColumn.h"
+#include "IChooserParameterGameplayTag.h"
+#include "ChooserPropertyAccess.h"
+#include "GameplayTagContainer.h"
+#include "StructUtils/InstancedStruct.h"
+#include "Serialization/MemoryReader.h"
+#include "OutputGameplayTagColumn.generated.h"
+
+#define UE_API CHOOSER_API
+
+USTRUCT(DisplayName = "Output Gameplay Tag", Meta = (Category = "Output", Tooltip = "A column which writes a Gameplay Tag Container value."))
+struct FOutputGameplayTagColumn : public FChooserColumnBase
+{
+	GENERATED_BODY()
+	public:
+	UE_API FOutputGameplayTagColumn();
+
+	virtual bool HasFilters() const override { return false; }
+	virtual bool HasOutputs() const override { return true; }
+	UE_API virtual void SetOutputs(FChooserEvaluationContext& Context, int RowIndex) const override;
+
+	UPROPERTY(EditAnywhere,  NoClear, Meta = (ExcludeBaseStruct, BaseStruct = "/Script/Chooser.FChooserParameterGameplayTagBase", ToolTip="The Gameplay Tag Container property this column will write to"), Category = "Hidden")
+	FInstancedStruct InputValue;
+
+#if WITH_EDITORONLY_DATA
+	UPROPERTY(EditAnywhere, Category=Data, meta=(ToolTip="DefaultRowValue will be assigned to cells when new rows are created"));
+	FGameplayTagContainer DefaultRowValue;
+#endif
+
+	UPROPERTY()
+	TArray<FGameplayTagContainer> RowValues;
+
+#if WITH_EDITOR
+	mutable FGameplayTagContainer TestValue;
+	UE_API virtual void AddToDetails(FInstancedPropertyBag& PropertyBag, int32 ColumnIndex, int32 RowIndex) override;
+	UE_API virtual void SetFromDetails(FInstancedPropertyBag& PropertyBag, int32 ColumnIndex, int32 RowIndex) override;
+	virtual void CopyFallback(FChooserColumnBase& SourceColumn) override { FallbackValue = static_cast<FOutputGameplayTagColumn&>(SourceColumn).FallbackValue; }
+#endif
+
+	// FallbackValue will be used as the output value if the all rows in the chooser fail, and the FallbackResult from the chooser is used.
+	UPROPERTY();
+   	FGameplayTagContainer FallbackValue;
+
+	FGameplayTagContainer& GetValueForIndex(int32 Index)
+	{
+		return Index == ChooserColumn_SpecialIndex_Fallback ? FallbackValue : RowValues[Index];
+	}
+
+	const FGameplayTagContainer& GetValueForIndex(int32 Index) const
+	{
+		return Index == ChooserColumn_SpecialIndex_Fallback ? FallbackValue : RowValues[Index];
+	}
+
+	CHOOSER_COLUMN_BOILERPLATE(FChooserParameterGameplayTagBase);
+};
+
+#undef UE_API
diff --git a/Engine/Plugins/Chooser/Source/Chooser/Private/GameplayTagColumn.cpp b/Engine/Plugins/Chooser/Source/Chooser/Private/GameplayTagColumn.cpp
--- a/Engine/Plugins/Chooser/Source/Chooser/Private/GameplayTagColumn.cpp
+++ b/Engine/Plugins/Chooser/Source/Chooser/Private/GameplayTagColumn.cpp
@@ -15,6 +15,11 @@ bool FGameplayTagContextProperty::GetValue(FChooserEvaluationContext& Context, c
 	return Binding.GetValuePtr(Context, OutResult);
 }
 
+bool FGameplayTagContextProperty::SetValue(FChooserEvaluationContext& Context, const FGameplayTagContainer& Value) const
+{
+	return Binding.SetValue(Context, Value);
+}
+
 FGameplayTagColumn::FGameplayTagColumn()
 {
 	InputValue.InitializeAs(FGameplayTagContextProperty::StaticStruct());
diff --git a/Engine/Plugins/Chooser/Source/Chooser/Private/OutputGameplayTagColumn.cpp b/Engine/Plugins/Chooser/Source/Chooser/Private/OutputGameplayTagColumn.cpp
new file mode 100644
--- /dev/null
+++ b/Engine/Plugins/Chooser/Source/Chooser/Private/OutputGameplayTagColumn.cpp
@@ -0,0 +1,66 @@
+﻿// Copyright Epic Games, Inc. All Rights Reserved.
+#include "OutputGameplayTagColumn.h"
+
+#include "ChooserIndexArray.h"
+#include "ChooserTrace.h"
+#include "GameplayTagColumn.h"
+#if WITH_EDITOR
+#include "StructUtils/PropertyBag.h"
+#endif
+
+#include UE_INLINE_GENERATED_CPP_BY_NAME(OutputGameplayTagColumn)
+
+FOutputGameplayTagColumn::FOutputGameplayTagColumn()
+{
+	InputValue.InitializeAs(FGameplayTagContextProperty::StaticStruct());
+}
+
+void FOutputGameplayTagColumn::SetOutputs(FChooserEvaluationContext& Context, int RowIndex) const
+{
+	if (InputValue.IsValid())
+	{
+		if (RowIndex == ChooserColumn_SpecialIndex_Fallback || RowValues.IsValidIndex(RowIndex))
+		{
+			const FGameplayTagContainer& OutputTags = GetValueForIndex(RowIndex);
+			InputValue.Get<FChooserParameterGameplayTagBase>().SetValue(Context, OutputTags);
+#if WITH_EDITOR
+			if (Context.DebuggingInfo.bCurrentDebugTarget)
+			{
+				TestValue = OutputTags;
+			}
+#endif
+		}
+		else
+		{
+#if CHOOSER_DEBUGGING_ENABLED
+			UE_ASSET_LOG(LogChooser, Error, Context.DebuggingInfo.CurrentChooser, TEXT("Invalid index %d passed to FOutputGameplayTagColumn::SetOutputs"), RowIndex);
+#else
+			UE_LOG(LogChooser, Error, TEXT("Invalid index %d passed to FOutputGameplayTagColumn::SetOutputs"), RowIndex);
+#endif
+		}
+	}
+}
+
+#if WITH_EDITOR
+	void FOutputGameplayTagColumn::AddToDetails(FInstancedPropertyBag& PropertyBag, int32 ColumnIndex, int32 RowIndex)
+	{
+		FText DisplayName;
+		InputValue.Get<FChooserParameterBase>().GetDisplayName(DisplayName);
+		FName PropertyName("RowData", ColumnIndex);
+		FPropertyBagPropertyDesc PropertyDesc(PropertyName, EPropertyBagPropertyType::Struct, FGameplayTagContainer::StaticStruct());
+		PropertyDesc.MetaData.Add(FPropertyBagPropertyDescMetaData("DisplayName", DisplayName.ToString()));
+		PropertyBag.AddProperties({PropertyDesc});
+		PropertyBag.SetValueStruct(PropertyName, RowValues[RowIndex]);
+	}
+
+	void FOutputGameplayTagColumn::SetFromDetails(FInstancedPropertyBag& PropertyBag, int32 ColumnIndex, int32 RowIndex)
+	{
+		FName PropertyName("RowData", ColumnIndex);
+
+		TValueOrError<FStructView, EPropertyBagResult> Result = PropertyBag.GetValueStruct(PropertyName, FGameplayTagContainer::StaticStruct());
+		if (FStructView* StructView = Result.TryGetValue())
+		{
+			RowValues[RowIndex] = StructView->Get<FGameplayTagContainer>();
+		}
+	}
+#endif
diff --git a/Engine/Plugins/Chooser/Source/Chooser/Public/IChooserParameterGameplayTag.h b/Engine/Plugins/Chooser/Source/Chooser/Public/IChooserParameterGameplayTag.h
--- a/Engine/Plugins/Chooser/Source/Chooser/Public/IChooserParameterGameplayTag.h
+++ b/Engine/Plugins/Chooser/Source/Chooser/Public/IChooserParameterGameplayTag.h
@@ -30,6 +30,7 @@ struct FChooserParameterGameplayTagBase : public FChooserParameterBase
 {
 	GENERATED_BODY()
 	virtual bool GetValue(FChooserEvaluationContext& Context, const FGameplayTagContainer*& OutResult) const { return false; }
+	virtual bool SetValue(FChooserEvaluationContext& Context, const FGameplayTagContainer& Value) const { return false; }
 };
 
 USTRUCT()
diff --git a/Engine/Plugins/Chooser/Source/ChooserEditor/Private/GameplayTagColumnEditor.cpp b/Engine/Plugins/Chooser/Source/ChooserEditor/Private/GameplayTagColumnEditor.cpp
--- a/Engine/Plugins/Chooser/Source/ChooserEditor/Private/GameplayTagColumnEditor.cpp
+++ b/Engine/Plugins/Chooser/Source/ChooserEditor/Private/GameplayTagColumnEditor.cpp
@@ -10,6 +10,8 @@
 #include "SGameplayTagWidget.h"
 #include "SSimpleComboButton.h"
 #include "GraphEditorSettings.h"
+#include "OutputGameplayTagColumn.h"
+#include "SGameplayTagContainerCombo.h"
 #include "Widgets/Images/SImage.h"
 
 #define LOCTEXT_NAMESPACE "FGameplayTagColumnEditor"
@@ -20,7 +22,7 @@ namespace UE::ChooserEditor
 TSharedRef<SWidget> CreateGameplayTagColumnWidget(UChooserTable* Chooser, FChooserColumnBase* Column, int Row)
 {
 	FGameplayTagColumn* GameplayTagColumn = static_cast<struct FGameplayTagColumn*>(Column);
-	
+
 	if (Row == ColumnWidget_SpecialIndex_Fallback)
 	{
 		return SNullWidget::NullWidget;
@@ -29,31 +31,24 @@ TSharedRef<SWidget> CreateGameplayTagColumnWidget(UChooserTable* Chooser, FChoos
 	{
 		// create column header widget
 		const FSlateBrush* ColumnIcon = FCoreStyle::Get().GetBrush("Icons.Filter");
-		const FText ColumnTooltip = LOCTEXT("Gameplay Tag Tooltip", "Gameplay Tag: cells pass if the input gameplay tag collection matches the cell data (accoding to comparison settings in the column properties).");
+		const FText ColumnTooltip = LOCTEXT("Gameplay Tag Tooltip", "Gameplay Tag: cells pass if the input gameplay tag collection matches the cell data (according to comparison settings in the column properties).");
 		const FText ColumnName = LOCTEXT("Gameplay Tag","Gameplay Tag");
-		
+
 		TSharedPtr<SWidget> DebugWidget = nullptr;
 		if (Chooser->GetEnableDebugTesting())
 		{
-			DebugWidget = SNew(SSimpleComboButton)
+			DebugWidget = SNew(SGameplayTagContainerCombo)
 				.IsEnabled_Lambda([Chooser]()
 				{
 					 return !Chooser->HasDebugTarget();
 				})
-				.Text_Lambda([GameplayTagColumn]()
+				.TagContainer_Lambda([GameplayTagColumn]()
 				{
-					FText Text = FText::FromString(GameplayTagColumn->TestValue.ToStringSimple(false));
-					if (Text.IsEmpty())
-					{
-						Text = LOCTEXT("None", "None");
-					}
-					return Text;
-				})	
-				.OnGetMenuContent_Lambda([Chooser, GameplayTagColumn, Row]()
+					return GameplayTagColumn->TestValue;
+				})
+				.OnTagContainerChanged_Lambda([GameplayTagColumn](const FGameplayTagContainer& UpdatedTags)
 				{
-					TArray<SGameplayTagWidget::FEditableGameplayTagContainerDatum> EditableContainers;
-					EditableContainers.Emplace(Chooser, &(GameplayTagColumn->TestValue));
-					return TSharedRef<SWidget>(SNew(SGameplayTagWidget, EditableContainers));
+					GameplayTagColumn->TestValue = UpdatedTags;
 				});
 		}
 
@@ -61,33 +56,68 @@ TSharedRef<SWidget> CreateGameplayTagColumnWidget(UChooserTable* Chooser, FChoos
 	}
 
 	// create cell widget
-	return SNew(SSimpleComboButton)
-		.Text_Lambda([GameplayTagColumn, Row]()
+	return SNew(SGameplayTagContainerCombo)
+		.ReadOnly(false)
+		.TagContainer_Lambda([GameplayTagColumn, Row]()
 		{
 			if (Row < GameplayTagColumn->RowValues.Num())
 			{
-				FText Result = FText::FromString(GameplayTagColumn->RowValues[Row].ToStringSimple(false));
-				if (Result.IsEmpty())
-				{
-					Result = LOCTEXT("Any Tag", "[Any]");
-				}
-				return Result;
+				return GameplayTagColumn->RowValues[Row];
 			}
 			else
 			{
-				return FText();
+				return FGameplayTagContainer::EmptyContainer;
 			}
-		})	
-		.OnGetMenuContent_Lambda([Chooser, GameplayTagColumn, Row]()
+		})
+		.OnTagContainerChanged_Lambda([GameplayTagColumn, Row](const FGameplayTagContainer& UpdatedTags)
 		{
 			if (Row < GameplayTagColumn->RowValues.Num())
 			{
-				TArray<SGameplayTagWidget::FEditableGameplayTagContainerDatum> EditableContainers;
-				EditableContainers.Emplace(Chooser, &(GameplayTagColumn->RowValues[Row]));
-				return TSharedRef<SWidget>(SNew(SGameplayTagWidget, EditableContainers));
+				GameplayTagColumn->RowValues[Row] = UpdatedTags;
 			}
+		}
+	);
+}
+
+TSharedRef<SWidget> CreateOutputGameplayTagColumnWidget(UChooserTable* Chooser, FChooserColumnBase* Column, int Row)
+{
+	FOutputGameplayTagColumn* OutputGameplayTagColumn = static_cast<struct FOutputGameplayTagColumn*>(Column);
+
+	if (Row == ColumnWidget_SpecialIndex_Fallback)
+	{
+		return SNullWidget::NullWidget;
+	}
+	else if (Row == ColumnWidget_SpecialIndex_Header)
+	{
+		// create column header widget
+		const FSlateBrush* ColumnIcon = FCoreStyle::Get().GetBrush("Icons.ArrowRight");
+		const FText ColumnTooltip = LOCTEXT("Output Gameplay Tag Tooltip", "Output Gameplay Tag: writes the value from cell in the result row to the bound variable");
+		const FText ColumnName = LOCTEXT("Output Gameplay Tag", "Output Gameplay Tag");
+
+		TSharedPtr<SWidget> DebugWidget = nullptr;
+		return MakeColumnHeaderWidget(Chooser, Column, ColumnName, ColumnTooltip, ColumnIcon, DebugWidget);
+	}
 
-			return SNullWidget::NullWidget;
+	// create cell widget
+	return SNew(SGameplayTagContainerCombo)
+		.ReadOnly(false)
+		.TagContainer_Lambda([OutputGameplayTagColumn, Row]()
+		{
+			if (Row < OutputGameplayTagColumn->RowValues.Num())
+			{
+				return OutputGameplayTagColumn->RowValues[Row];
+			}
+			else
+			{
+				return FGameplayTagContainer::EmptyContainer;
+			}
+		})
+		.OnTagContainerChanged_Lambda([OutputGameplayTagColumn, Row](const FGameplayTagContainer& UpdatedTags)
+		{
+				if (Row < OutputGameplayTagColumn->RowValues.Num())
+				{
+					OutputGameplayTagColumn->RowValues[Row] = UpdatedTags;
+				}
 		}
 	);
 }
@@ -102,13 +132,14 @@ TSharedRef<SWidget> CreateGameplayTagPropertyWidget(bool bReadOnly, UObject* Tra
 	.PropertyBindingValue(&ContextProperty->Binding)
 	.OnValueChanged(ValueChanged);
 }
-	
+
 void RegisterGameplayTagWidgets()
 {
 	FObjectChooserWidgetFactories::RegisterWidgetCreator(FGameplayTagContextProperty::StaticStruct(), CreateGameplayTagPropertyWidget);
 	FObjectChooserWidgetFactories::RegisterColumnWidgetCreator(FGameplayTagColumn::StaticStruct(), CreateGameplayTagColumnWidget);
+	FObjectChooserWidgetFactories::RegisterColumnWidgetCreator(FOutputGameplayTagColumn::StaticStruct(), CreateOutputGameplayTagColumnWidget);
 }
-	
+
 }
 
 #undef LOCTEXT_NAMESPACE
