Remove ; after UPROPERTY

diff --git a/Engine/Plugins/Chooser/Source/Chooser/Internal/OutputGameplayTagColumn.h b/Engine/Plugins/Chooser/Source/Chooser/Internal/OutputGameplayTagColumn.h
--- a/Engine/Plugins/Chooser/Source/Chooser/Internal/OutputGameplayTagColumn.h
+++ b/Engine/Plugins/Chooser/Source/Chooser/Internal/OutputGameplayTagColumn.h
@@ -41,7 +41,7 @@ struct FOutputGameplayTagColumn : public FChooserColumnBase
 #endif
 
 	// FallbackValue will be used as the output value if the all rows in the chooser fail, and the FallbackResult from the chooser is used.
-	UPROPERTY();
+	UPROPERTY()
    	FGameplayTagContainer FallbackValue;
 
 	FGameplayTagContainer& GetValueForIndex(int32 Index)
diff --git a/Engine/Plugins/Chooser/Source/Chooser/Internal/OutputGameplayTagQueryColumn.h b/Engine/Plugins/Chooser/Source/Chooser/Internal/OutputGameplayTagQueryColumn.h
--- a/Engine/Plugins/Chooser/Source/Chooser/Internal/OutputGameplayTagQueryColumn.h
+++ b/Engine/Plugins/Chooser/Source/Chooser/Internal/OutputGameplayTagQueryColumn.h
@@ -55,7 +55,7 @@ struct FOutputGameplayTagQueryColumn : public FChooserColumnBase
 #endif
 
 	// FallbackValue will be used as the output value if the all rows in the chooser fail, and the FallbackResult from the chooser is used.
-	UPROPERTY();
+	UPROPERTY()
    	FGameplayTagQuery FallbackValue;
 
 	FGameplayTagQuery& GetValueForIndex(int32 Index)
