Fix error in UChooserTableFactory

diff --git a/Engine/Plugins/Chooser/Source/ChooserEditor/Private/ChooserFactory.cpp b/Engine/Plugins/Chooser/Source/ChooserEditor/Private/ChooserFactory.cpp
--- a/Engine/Plugins/Chooser/Source/ChooserEditor/Private/ChooserFactory.cpp
+++ b/Engine/Plugins/Chooser/Source/ChooserEditor/Private/ChooserFactory.cpp
@@ -43,7 +43,7 @@ UObject* UChooserTableFactory::FactoryCreateNew(UClass* Class, UObject* InParent
 	}
 	else
 	{
-		NewObject<UChooserTable>(InParent, Class, Name, Flags);
+		NewChooser = NewObject<UChooserTable>(InParent, Class, Name, Flags);
 	}
 	return NewChooser;
 }
\ No newline at end of file
