Add Validation support to ChooserTable assets

diff --git a/Engine/Plugins/Chooser/Source/Chooser/Private/Chooser.cpp b/Engine/Plugins/Chooser/Source/Chooser/Private/Chooser.cpp
--- a/Engine/Plugins/Chooser/Source/Chooser/Private/Chooser.cpp
+++ b/Engine/Plugins/Chooser/Source/Chooser/Private/Chooser.cpp
@@ -292,6 +292,14 @@ void UChooserTable::CookData()
 #endif
 
 #if WITH_EDITOR
+
+EDataValidationResult UChooserTable::IsDataValid(FDataValidationContext& Context) const
+{
+	// todo: implement a const Compile which knows if it passed or failed
+	const_cast<UChooserTable*>(this)->Compile(true);
+	return Super::IsDataValid(Context);
+}
+
 void UChooserTable::PostEditUndo()
 {
 	UObject::PostEditUndo();
diff --git a/Engine/Plugins/Chooser/Source/Chooser/Public/Chooser.h b/Engine/Plugins/Chooser/Source/Chooser/Public/Chooser.h
--- a/Engine/Plugins/Chooser/Source/Chooser/Public/Chooser.h
+++ b/Engine/Plugins/Chooser/Source/Chooser/Public/Chooser.h
@@ -44,6 +44,8 @@ public:
 #endif
 	}
 #if WITH_EDITOR
+	CHOOSER_API	virtual EDataValidationResult IsDataValid(class FDataValidationContext& Context) const override;
+
 	CHOOSER_API virtual void GetAssetRegistryTags(FAssetRegistryTagsContext Context) const override;
 	UE_DEPRECATED(5.4, "Implement the version that takes FAssetRegistryTagsContext instead.")
 	CHOOSER_API virtual void GetAssetRegistryTags(TArray<FAssetRegistryTag>& OutTags) const override;
