Chooser table editor refactoring to a View and ModelView.  Extracting all the necessary editor functionality from the editor class so that it can be reused in the UAF workspace editor

diff --git a/Engine/Plugins/Animation/PoseSearch/Source/Runtime/Private/Chooser/PoseSearchChooserColumn.cpp b/Engine/Plugins/Animation/PoseSearch/Source/Runtime/Private/Chooser/PoseSearchChooserColumn.cpp
--- a/Engine/Plugins/Animation/PoseSearch/Source/Runtime/Private/Chooser/PoseSearchChooserColumn.cpp
+++ b/Engine/Plugins/Animation/PoseSearch/Source/Runtime/Private/Chooser/PoseSearchChooserColumn.cpp
@@ -223,14 +223,14 @@ void FPoseSearchColumn::InsertRows(int Index, int Count)
 	CheckConsistency();
 }
 
-void FPoseSearchColumn::DeleteRows(const TArray<uint32>& RowIndices)
+void FPoseSearchColumn::DeleteRows(TArrayView<int> RowIndices)
 {
 	using namespace UE::PoseSearch;
 
 	check(InternalDatabase);
 	InternalDatabase->Modify();
 	
-	for (uint32 RowIndex : RowIndices)
+	for (int RowIndex : RowIndices)
 	{
 		const int32 DatabaseAssetIndex = GetDatabaseAssetIndex(RowIndex);
 		InternalDatabase->RemoveAnimationAssetAt(DatabaseAssetIndex);
diff --git a/Engine/Plugins/Animation/PoseSearch/Source/Runtime/Public/PoseSearch/Chooser/PoseSearchChooserColumn.h b/Engine/Plugins/Animation/PoseSearch/Source/Runtime/Public/PoseSearch/Chooser/PoseSearchChooserColumn.h
--- a/Engine/Plugins/Animation/PoseSearch/Source/Runtime/Public/PoseSearch/Chooser/PoseSearchChooserColumn.h
+++ b/Engine/Plugins/Animation/PoseSearch/Source/Runtime/Public/PoseSearch/Chooser/PoseSearchChooserColumn.h
@@ -141,7 +141,7 @@ private:
 	virtual FName RowValuesPropertyName() override;
 	virtual void SetNumRows(int32 NumRows) override;
 	virtual void InsertRows(int Index, int Count) override;
-	virtual void DeleteRows(const TArray<uint32>& RowIndices) override;
+	virtual void DeleteRows(TArrayView<int> RowIndices) override;
 	virtual void MoveRow(int SourceRowIndex, int TargetRowIndex) override;
 	virtual void CopyRow(FChooserColumnBase& SourceColumn, int SourceRowIndex, int TargetRowIndex) override;
 	virtual UScriptStruct* GetInputBaseType() const override;
@@ -150,6 +150,11 @@ private:
 	
 	virtual bool AutoPopulates() const override { return true; }
 	virtual void AutoPopulate(int32 RowIndex, UObject* OutputObject) override;
+	
+	virtual bool HasPrimaryInput() override
+	{
+		return false;
+	}
 
 #endif // WITH_EDITOR
 	virtual FChooserParameterBase* GetInputValue() override;
diff --git a/Engine/Plugins/Chooser/Source/Chooser/Private/Chooser.cpp b/Engine/Plugins/Chooser/Source/Chooser/Private/Chooser.cpp
--- a/Engine/Plugins/Chooser/Source/Chooser/Private/Chooser.cpp
+++ b/Engine/Plugins/Chooser/Source/Chooser/Private/Chooser.cpp
@@ -260,7 +260,7 @@ void UChooserTable::RemoveDisabledData()
 	});
 	
 	// remove disabled rows and corresponding row data from columns
-	TArray<uint32> RowsToDelete;
+	TArray<int> RowsToDelete;
 	const int NumResults = ResultsStructs.Num();
 	for (int ResultIndex = NumResults - 1; ResultIndex>=0; ResultIndex--)
 	{
@@ -297,6 +297,7 @@ EDataValidationResult UChooserTable::IsDataValid(FDataValidationContext& Context
 {
 	// todo: implement a const Compile which knows if it passed or failed
 	const_cast<UChooserTable*>(this)->Compile(true);
+	
 	return Super::IsDataValid(Context);
 }
 
diff --git a/Engine/Plugins/Chooser/Source/Chooser/Public/IChooserColumn.h b/Engine/Plugins/Chooser/Source/Chooser/Public/IChooserColumn.h
--- a/Engine/Plugins/Chooser/Source/Chooser/Public/IChooserColumn.h
+++ b/Engine/Plugins/Chooser/Source/Chooser/Public/IChooserColumn.h
@@ -89,7 +89,7 @@ public:
 	
 	virtual FName RowValuesPropertyName() { return FName(); }
 	virtual void SetNumRows(int32 NumRows) {}
-	virtual void DeleteRows(const TArray<uint32> & RowIndices) {}
+	virtual void DeleteRows(TArrayView<int> RowIndices) {}
 	virtual void MoveRow(int SourceIndex, int TargetIndex) {}
 	virtual void InsertRows(int Index, int Count) {}
 	virtual void CopyRow(FChooserColumnBase& SourceColumn, int SourceIndex, int TargetIndex) {}
@@ -151,11 +151,14 @@ enum {  ChooserColumn_SpecialIndex_Fallback = -2 };
 			RowValuesProperty.Insert(DefaultRowValue, Index);\
 		}\
 	}\
-	virtual void DeleteRows(const TArray<uint32> & RowIndices ) override\
+	virtual void DeleteRows(TArrayView<int> RowIndices ) override\
 	{\
-		for(uint32 Index : RowIndices)\
+		for(int32 Index : RowIndices)\
 		{\
-			RowValuesProperty.RemoveAt(Index);\
+			if (RowValuesProperty.IsValidIndex(Index)) \
+			{\
+				RowValuesProperty.RemoveAt(Index);\
+			}\
 		}\
 	}\
 	virtual void MoveRow(int SourceRowIndex, int TargetRowIndex) override\
diff --git a/Engine/Plugins/Chooser/Source/ChooserEditor/Private/ChooserEditorModule.cpp b/Engine/Plugins/Chooser/Source/ChooserEditor/Private/ChooserEditorModule.cpp
--- a/Engine/Plugins/Chooser/Source/ChooserEditor/Private/ChooserEditorModule.cpp
+++ b/Engine/Plugins/Chooser/Source/ChooserEditor/Private/ChooserEditorModule.cpp
@@ -30,6 +30,7 @@
 #include "Kismet2/EnumEditorUtils.h"
 #include "EnumColumn.h"
 #include "UObject/UObjectIterator.h"
+#include "SChooserTableWidget.h"
 
 #define LOCTEXT_NAMESPACE "ChooserEditorModule"
 
@@ -54,7 +55,17 @@ void FEnumChangedListener::PostChange(const UUserDefinedEnum* Changed, FEnumEdit
 	}
 }
 
-void FModule::StartupModule()
+TSharedPtr<IChooserTableViewModel> FChooserEditorModule::CreateChooserTableViewModel(UChooserTable* ChooserTable)
+{
+	return MakeShared<FChooserTableViewModel>(ChooserTable);
+}
+
+TSharedPtr<SWidget> FChooserEditorModule::CreateChooserTableView(TSharedPtr<IChooserTableViewModel> ViewModel, TSharedPtr<FUICommandList> CommandList)
+{
+	return SNew(SChooserTableWidget).ViewModel(ViewModel).Commands(CommandList);
+}
+
+void FChooserEditorModule::StartupModule()
 {
 	FChooserEditorStyle::Initialize();
 	
@@ -92,7 +103,7 @@ void FModule::StartupModule()
 	FEnumEditorUtils::FEnumEditorManager::Get().AddListener(&EnumChanged);
 }
 
-void FModule::ShutdownModule()
+void FChooserEditorModule::ShutdownModule()
 {
 	IModularFeatures::Get().UnregisterModularFeature( IRewindDebuggerExtension::ModularFeatureName, &RewindDebuggerChooser);
 	IModularFeatures::Get().UnregisterModularFeature(RewindDebugger::IRewindDebuggerTrackCreator::ModularFeatureName, &GChoosersTrackCreator);
@@ -106,6 +117,6 @@ void FModule::ShutdownModule()
 
 }
 
-IMPLEMENT_MODULE(UE::ChooserEditor::FModule, ChooserEditor);
+IMPLEMENT_MODULE(UE::ChooserEditor::FChooserEditorModule, ChooserEditor);
 
 #undef LOCTEXT_NAMESPACE
\ No newline at end of file
diff --git a/Engine/Plugins/Chooser/Source/ChooserEditor/Private/ChooserEditorModule.h b/Engine/Plugins/Chooser/Source/ChooserEditor/Private/ChooserEditorModule.h
--- a/Engine/Plugins/Chooser/Source/ChooserEditor/Private/ChooserEditorModule.h
+++ b/Engine/Plugins/Chooser/Source/ChooserEditor/Private/ChooserEditorModule.h
@@ -7,6 +7,10 @@
 #include "Modules/ModuleInterface.h"
 #include "RewindDebuggerChooser.h"
 #include "Kismet2/EnumEditorUtils.h"
+#include "IChooserEditorModule.h"
+
+
+class FUICommandList;
 
 namespace UE::ChooserEditor
 {
@@ -17,12 +21,16 @@ class FEnumChangedListener : public FEnumEditorUtils::FEnumEditorManager::BaseNo
 	virtual void PreChange(const UUserDefinedEnum* Changed, FEnumEditorUtils::EEnumEditorChangeInfo ChangedType) override {};
 };
 
-class FModule : public IModuleInterface
+class FChooserEditorModule : public IChooserEditorModule
 {
+public: 
+	virtual TSharedPtr<IChooserTableViewModel> CreateChooserTableViewModel(UChooserTable* ChooserTable) override;
+	virtual TSharedPtr<SWidget> CreateChooserTableView(TSharedPtr<IChooserTableViewModel> ViewModel, TSharedPtr<FUICommandList> CommandList) override;
+
 private:
 	virtual void StartupModule() override;
 	virtual void ShutdownModule() override;
-	
+
 	FRewindDebuggerChooser RewindDebuggerChooser;
 	FChooserTraceModule ChooserTraceModule;
 
diff --git a/Engine/Plugins/Chooser/Source/ChooserEditor/Private/ChooserEditorStyle.h b/Engine/Plugins/Chooser/Source/ChooserEditor/Private/ChooserEditorStyle.h
deleted file mode 100644
diff --git a/Engine/Plugins/Chooser/Source/ChooserEditor/Private/ChooserEditorWidgets.cpp b/Engine/Plugins/Chooser/Source/ChooserEditor/Private/ChooserEditorWidgets.cpp
--- a/Engine/Plugins/Chooser/Source/ChooserEditor/Private/ChooserEditorWidgets.cpp
+++ b/Engine/Plugins/Chooser/Source/ChooserEditor/Private/ChooserEditorWidgets.cpp
@@ -2,6 +2,7 @@
 
 #include "ChooserEditorWidgets.h"
 
+#include "AssetTypeCategories.h"
 #include "Chooser.h"
 #include "ChooserEditorStyle.h"
 #include "ChooserTableEditor.h"
@@ -31,6 +32,11 @@ TSharedRef<SWidget> CreateAssetWidget(bool bReadOnly, UObject* TransactionObject
 
 	UObject* Asset = DIAsset->Asset;
 	
+	FAssetReferenceFilterContext AssetReferenceFilterContext;
+	AssetReferenceFilterContext.AddReferencingAssets({&TransactionObject, 1});
+	TSharedPtr<IAssetReferenceFilter> AssetReferenceFilter = GEditor ? GEditor->MakeAssetReferenceFilter(AssetReferenceFilterContext) : nullptr;	
+	
+	
 	return SNew(SObjectPropertyEntryBox)
 		.IsEnabled(!bReadOnly)
 		.AllowedClass(ResultBaseClass ? ResultBaseClass : UObject::StaticClass())
@@ -42,6 +48,14 @@ TSharedRef<SWidget> CreateAssetWidget(bool bReadOnly, UObject* TransactionObject
 			TransactionObject->Modify(true);
 			DIAsset->Asset = AssetData.GetAsset();
 			ValueChanged.ExecuteIfBound();
+		})
+		.OnShouldFilterAsset_Lambda([AssetReferenceFilter](const FAssetData& AssetData)
+		{
+			if (AssetReferenceFilter.IsValid())
+			{
+				return !AssetReferenceFilter->PassesFilter(AssetData);
+			}
+			return false;
 		});
 }
 	
@@ -49,6 +63,10 @@ TSharedRef<SWidget> CreateSoftAssetWidget(bool bReadOnly, UObject* TransactionOb
 {
 	FSoftAssetChooser* DIAsset = static_cast<FSoftAssetChooser*>(Value);
 
+	FAssetReferenceFilterContext AssetReferenceFilterContext;
+	AssetReferenceFilterContext.AddReferencingAssets({&TransactionObject, 1});
+	TSharedPtr<IAssetReferenceFilter> AssetReferenceFilter = GEditor ? GEditor->MakeAssetReferenceFilter(AssetReferenceFilterContext) : nullptr;
+	
 	return SNew(SObjectPropertyEntryBox)
 		.IsEnabled(!bReadOnly)
 		.AllowedClass(ResultBaseClass ? ResultBaseClass : UObject::StaticClass())
@@ -59,14 +77,23 @@ TSharedRef<SWidget> CreateSoftAssetWidget(bool bReadOnly, UObject* TransactionOb
 			TransactionObject->Modify(true);
 			DIAsset->Asset = AssetData.GetAsset();
 			ValueChanged.ExecuteIfBound();
+		})
+		.OnShouldFilterAsset_Lambda([AssetReferenceFilter](const FAssetData& AssetData)
+		{
+			if (AssetReferenceFilter.IsValid())
+			{
+				return !AssetReferenceFilter->PassesFilter(AssetData);
+			}
+			return false;
 		});
 }
-	
+
 TSharedRef<SWidget> CreateClassWidget(bool bReadOnly, UObject* TransactionObject, void* Value, UClass* ResultBaseClass, FChooserWidgetValueChanged ValueChanged)
 {
 	FClassChooser* ClassChooser = static_cast<FClassChooser*>(Value);
 
 	UClass* Class = ClassChooser->Class;
+	
 
 	return SNew(SClassPropertyEntryBox)
 			.IsEnabled(!bReadOnly)
@@ -86,14 +113,27 @@ TSharedRef<SWidget> CreateClassWidget(bool bReadOnly, UObject* TransactionObject
 TSharedRef<SWidget> CreateEvaluateChooserWidget(bool bReadOnly, UObject* TransactionObject, void* Value, UClass* ResultBaseClass, FChooserWidgetValueChanged ValueChanged)
 {
 	FEvaluateChooser* EvaluateChooser = static_cast<FEvaluateChooser*>(Value);
+
+	FAssetReferenceFilterContext AssetReferenceFilterContext;
+	AssetReferenceFilterContext.AddReferencingAssets({&TransactionObject, 1});
+	TSharedPtr<IAssetReferenceFilter> AssetReferenceFilter = GEditor ? GEditor->MakeAssetReferenceFilter(AssetReferenceFilterContext) : nullptr;
 	
 	return SNew(SObjectPropertyEntryBox)
 		.IsEnabled(!bReadOnly)
 		.AllowedClass(UChooserTable::StaticClass())
 		.bOnlyRecognizeOnDragEnter(true)
 		.ObjectPath_Lambda([EvaluateChooser](){ return EvaluateChooser->Chooser ? EvaluateChooser->Chooser.GetPath() : "";})
-		.OnShouldFilterAsset_Lambda([ResultBaseClass](const FAssetData& InAssetData)
+		.OnShouldFilterAsset_Lambda([ResultBaseClass, AssetReferenceFilter](const FAssetData& InAssetData)
 		{
+			if (AssetReferenceFilter.IsValid())
+			{
+				if (!AssetReferenceFilter->PassesFilter(InAssetData))
+				{
+					// filter for assets outside the current plugin scope etc.
+					return true;
+				}
+			}
+				
 			if (ResultBaseClass == nullptr)
 		 	{
 		 		return false;
@@ -281,11 +321,9 @@ TSharedRef<SWidget> CreateNestedChooserWidget(bool bReadOnly, UObject* Transacti
 				{
 					UObject* RootChooser = TransactionObject->GetPackage()->FindAssetInPackage();
 					IAssetEditorInstance* Editor = GEditor->GetEditorSubsystem<UAssetEditorSubsystem>()->FindEditorForAsset(RootChooser, false);
-					if (Editor && Editor->GetEditorName() == FChooserTableEditor::EditorName)
+					if (Editor)
 					{
-						// currently nested tables are only editable in the Chooser Table Editor
-						FChooserTableEditor* ChooserEditor = static_cast<FChooserTableEditor*>(Editor);
-						ChooserEditor->PushChooserTableToEdit(NestedChooser->Chooser);
+						Editor->FocusWindow(NestedChooser->Chooser);
 					}
 				}
 				return FReply::Handled();
diff --git a/Engine/Plugins/Chooser/Source/ChooserEditor/Private/ChooserTableEditor.cpp b/Engine/Plugins/Chooser/Source/ChooserEditor/Private/ChooserTableEditor.cpp
--- a/Engine/Plugins/Chooser/Source/ChooserEditor/Private/ChooserTableEditor.cpp
+++ b/Engine/Plugins/Chooser/Source/ChooserEditor/Private/ChooserTableEditor.cpp
@@ -11,7 +11,6 @@
 #include "DetailCategoryBuilder.h"
 #include "Factories.h"
 #include "GraphEditorSettings.h"
-#include "IDetailsView.h"
 #include "IPropertyAccessEditor.h"
 #include "LandscapeRender.h"
 #include "ObjectChooserClassFilter.h"
@@ -52,6 +51,7 @@
 #include "HAL/PlatformApplicationMisc.h"
 #include "Framework/Commands/GenericCommands.h"
 #include "SPositiveActionButton.h"
+#include "SChooserTableWidget.h"
 
 #include UE_INLINE_GENERATED_CPP_BY_NAME(ChooserTableEditor)
 
@@ -59,6 +59,8 @@
 
 namespace UE::ChooserEditor
 {
+	
+constexpr int32 HistorySize = 16;	
 
 const FName FChooserTableEditor::ToolkitFName( TEXT( "ChooserTableEditor" ) );
 const FName FChooserTableEditor::PropertiesTabId( TEXT( "ChooserEditor_Properties" ) );
@@ -66,99 +68,6 @@ const FName FChooserTableEditor::FindReplaceTabId( TEXT( "ChooserEditor_FindRepl
 const FName FChooserTableEditor::TableTabId( TEXT( "ChooserEditor_Table" ) );
 const FName FChooserTableEditor::NestedTablesTreeTabId( TEXT( "ChooserEditor_NestedTables" ) );
 
-constexpr int HistorySize = 16;	
-
-void FChooserTableEditor::AddHistory()
-{
-	// remove anything ahead of this in the history, if we had gone back
-	while (HistoryIndex !=0)
-	{
-		History.PopFront();
-		HistoryIndex--;
-	}
-	
-	if (History.Num() >= HistorySize)
-	{
-		History.Pop();
-	}
-	History.AddFront(GetChooser());
-}
-
-bool FChooserTableEditor::CanNavigateBack() const
-{
-	return HistoryIndex < History.Num() - 1;
-}
-
-void FChooserTableEditor::NavigateBack()
-{
-	if (HistoryIndex < History.Num() - 1)
-	{
-		HistoryIndex++;
-		SetChooserTableToEdit(History[HistoryIndex], false);
-	}
-}
-
-bool FChooserTableEditor::CanNavigateForward() const
-{
-	return HistoryIndex > 0;
-}
-
-void FChooserTableEditor::NavigateForward()
-{
-	if (HistoryIndex > 0)
-	{
-		HistoryIndex--;
-		SetChooserTableToEdit(History[HistoryIndex], false);
-	}
-}
-
-void FChooserTableEditor::SetChooserTableToEdit(UChooserTable* Chooser, bool bApplyToHistory)
-{
-	if (Chooser == GetChooser())
-	{
-		return;
-	}
-	
-	BreadcrumbTrail->ClearCrumbs();
-
-	TArray<UChooserTable*> OuterList;
-	OuterList.Push(Chooser);
-	
-	while(OuterList.Last() != GetRootChooser())
-	{
-		OuterList.Push(Cast<UChooserTable>(OuterList.Last()->GetOuter()));
-	}
-
-	while(!OuterList.IsEmpty())
-	{
-		UChooserTable* Popped = OuterList.Pop();
-		BreadcrumbTrail->PushCrumb(FText::FromString(Popped->GetName()), Popped);
-	}
-	
-	if (bApplyToHistory)
-	{
-		AddHistory();
-	}
-	
-	RefreshAll();
-}
-
-void FChooserTableEditor::PushChooserTableToEdit(UChooserTable* Chooser)
-{
-	BreadcrumbTrail->PushCrumb(FText::FromString(Chooser->GetName()), Chooser);
-	AddHistory();
-	RefreshAll();
-}
-	
-void FChooserTableEditor::PopChooserTableToEdit()
-{
-	if (BreadcrumbTrail->HasCrumbs())
-	{
-		BreadcrumbTrail->PopCrumb();
-		RefreshAll();
-	}
-}
-	
 void FChooserTableEditor::RegisterTabSpawners(const TSharedRef<class FTabManager>& InTabManager)
 {
 	WorkspaceMenuCategory = InTabManager->AddLocalWorkspaceMenuCategory(LOCTEXT("WorkspaceMenu_ChooserTableEditor", "Chooser Table Editor"));
@@ -198,19 +107,18 @@ void FChooserTableEditor::UnregisterTabSpawners(const TSharedRef<class FTabManag
 
 const FName FChooserTableEditor::ChooserEditorAppIdentifier( TEXT( "ChooserEditorApp" ) );
 
+FChooserTableEditor::FChooserTableEditor()
+{
+}
+
 FChooserTableEditor::~FChooserTableEditor()
 {
-	FTSTicker::RemoveTicker(TickDelegateHandle);
-	
 	FCoreUObjectDelegates::OnObjectsReplaced.RemoveAll(this);
-	FCoreUObjectDelegates::OnObjectTransacted.RemoveAll(this);
-
-	DetailsView.Reset();
 }
 
 FName FChooserTableEditor::EditorName = "ChooserTableEditor";
 	
-FName FChooserTableEditor::ContextMenuName("ChooserEditorContextMenu"); // todo: for this to actually be extensible this needs to be somewhere public
+FName FChooserTableEditor::ContextMenuName("ChooserEditorContextMenu");
 	
 FName FChooserTableEditor::GetEditorName() const
 {
@@ -232,13 +140,13 @@ void FChooserTableEditor::MakeDebugTargetMenu(UToolMenu* InToolMenu)
 				FUIAction(
 					FExecuteAction::CreateLambda([this]()
 					{
-						UChooserTable* Chooser = GetRootChooser();
+						UChooserTable* Chooser = ViewModel->GetRootChooser();
 						Chooser->ResetDebugTarget();
 						if (Chooser->GetEnableDebugTesting())
 						{
 							Chooser->SetEnableDebugTesting(false);
 							Chooser->SetDebugTestValuesValid(false);
-							UpdateTableColumns();
+							ViewModel->RefreshAll();
 						}
 					}),
 					FCanExecuteAction()
@@ -255,20 +163,20 @@ void FChooserTableEditor::MakeDebugTargetMenu(UToolMenu* InToolMenu)
 				FUIAction(
 					FExecuteAction::CreateLambda([this]()
 					{
-						UChooserTable* Chooser = GetRootChooser();
+						UChooserTable* Chooser = ViewModel->GetRootChooser();
 						Chooser->ResetDebugTarget();
 						if (!Chooser->GetEnableDebugTesting())
 						{
 							Chooser->SetEnableDebugTesting(true);
 							Chooser->SetDebugTestValuesValid(true);
-							UpdateTableColumns();
+							ViewModel->RefreshAll ();
 						}
 					}),
 					FCanExecuteAction()
 				)
 			));
 
-	const UChooserTable* Chooser = GetChooser();
+	const UChooserTable* Chooser = ViewModel->GetChooser();
 
 	Chooser->IterateRecentContextObjects([this, InToolMenu](const FString& ObjectName)
 		{
@@ -282,13 +190,13 @@ void FChooserTableEditor::MakeDebugTargetMenu(UToolMenu* InToolMenu)
 							FUIAction(
 								FExecuteAction::CreateLambda([this, ObjectName]()
 								{
-									UChooserTable* Chooser = GetRootChooser();
+									UChooserTable* Chooser = ViewModel->GetRootChooser();
 									Chooser->SetDebugTarget(ObjectName);
 									Chooser->SetDebugTestValuesValid(false);
 									if (!Chooser->GetEnableDebugTesting())
 									{
 										Chooser->SetEnableDebugTesting(true);
-										UpdateTableColumns();
+										ViewModel->RefreshAll();
 									}
 								}),
 								FCanExecuteAction()
@@ -298,13 +206,10 @@ void FChooserTableEditor::MakeDebugTargetMenu(UToolMenu* InToolMenu)
 	);
 
 }
-	
-TSharedPtr<SWidget> FChooserTableEditor::GenerateRowContextMenu()
+
+UChooserTable* FChooserTableEditor::GetRootChooser()
 {
-	UToolMenus* ToolMenus = UToolMenus::Get();
-	FToolMenuContext ToolMenuContext;
-	InitToolMenuContext(ToolMenuContext);
-	return ToolMenus->GenerateWidget(ContextMenuName, ToolMenuContext);
+	return ViewModel->GetRootChooser();
 }
 
 void FChooserTableEditor::RegisterToolbar()
@@ -346,7 +251,7 @@ void FChooserTableEditor::RegisterToolbar()
 					InSection.AddEntry(FToolMenuEntry::InitComboButton( "SelectDebugTarget",
 						FToolUIActionChoice(),
 					  FNewToolMenuDelegate::CreateSP(ChooserEditor.Get(), &FChooserTableEditor::MakeDebugTargetMenu),
-						TAttribute<FText>::CreateLambda([Chooser = ChooserEditor->GetRootChooser() ]
+						TAttribute<FText>::CreateLambda([Chooser = ChooserEditor->ViewModel->GetRootChooser() ]
 						{
 							if (Chooser->HasDebugTarget())
 							{
@@ -367,78 +272,7 @@ void FChooserTableEditor::RegisterToolbar()
 
 void FChooserTableEditor::RegisterMenus()
 {
-	UToolMenus* ToolMenus = UToolMenus::Get();
-	const FChooserTableEditorCommands& Commands = FChooserTableEditorCommands::Get();
-
-	// Table Context Menu
-	UToolMenu* ToolMenu;
-	if (ToolMenus->IsMenuRegistered(ContextMenuName))
-	{
-		ToolMenu = ToolMenus->ExtendMenu(ContextMenuName);
-	}
-	else
-	{
-		ToolMenu = UToolMenus::Get()->RegisterMenu(ContextMenuName, NAME_None, EMultiBoxType::Menu);
-	}
-
-	if (ToolMenu)
-	{
-		FToolMenuSection& Section = ToolMenu->AddSection("ChooserEditorContext", TAttribute<FText>());
-	
-		Section.AddEntry(FToolMenuEntry::InitMenuEntry(FGenericCommands::Get().Copy));
-		Section.AddEntry(FToolMenuEntry::InitMenuEntry(FGenericCommands::Get().Cut));
-		Section.AddEntry(FToolMenuEntry::InitMenuEntry(FGenericCommands::Get().Paste));
-		Section.AddEntry(FToolMenuEntry::InitMenuEntry(FGenericCommands::Get().Duplicate));
-		Section.AddEntry(FToolMenuEntry::InitMenuEntry(FGenericCommands::Get().Delete));
-		Section.AddEntry(FToolMenuEntry::InitMenuEntry(Commands.Disable));
-		Section.AddEntry(FToolMenuEntry::InitMenuEntry(Commands.MoveUp));
-		Section.AddEntry(FToolMenuEntry::InitMenuEntry(Commands.MoveDown));
-		Section.AddEntry(FToolMenuEntry::InitMenuEntry(Commands.MoveLeft));
-		Section.AddEntry(FToolMenuEntry::InitMenuEntry(Commands.MoveRight));
-		Section.AddEntry(FToolMenuEntry::InitMenuEntry(Commands.AutoPopulateSelection));
-		
-
-		Section.AddDynamicEntry("ColumnInputType", FNewToolMenuSectionDelegate::CreateLambda([this] (FToolMenuSection& Section)
-		{
-			if (CurrentSelectionType == ESelectionType::Column && SelectedColumn)
-			{
-				Section.AddSubMenu("ParameterType", LOCTEXT("Parameter Type", "Parameter Type"),
-					LOCTEXT("Parameter Type Tooltip", "Change the type of input/output parameter for this column"),
-					FNewToolMenuChoice(FOnGetContent::CreateLambda([this]()
-					{
-						int ColumnIndex = SelectedColumn->Column;
-						UChooserTable* Chooser = SelectedColumn->Chooser;
-						FStructViewerInitializationOptions Options;
-						Options.StructFilter = MakeShared<FStructFilter>(Chooser->ColumnsStructs[ColumnIndex].Get<FChooserColumnBase>().GetInputBaseType());
-						Options.bAllowViewOptions = false;
-						Options.bShowNoneOption = false;
-						Options.NameTypeToDisplay = EStructViewerNameTypeToDisplay::DisplayName;
-					
-						// Add class filter for columns here
-						TSharedRef<SWidget> Widget = FModuleManager::LoadModuleChecked<FStructViewerModule>("StructViewer").CreateStructViewer(Options, FOnStructPicked::CreateLambda([this, ColumnIndex](const UScriptStruct* ChosenStruct)
-						{
-							const FScopedTransaction Transaction(LOCTEXT("SetColumnInputType", "Set Column Input Type"));
-							UChooserTable* ChooserTable = GetChooser();
-							ChooserTable->ColumnsStructs[ColumnIndex].GetMutable<FChooserColumnBase>().SetInputType(ChosenStruct);
-							ChooserTable->Modify(true);
-							UpdateTableColumns();
-							UpdateTableRows();
-							
-							if (SelectedColumn && SelectedColumn->Column == ColumnIndex)
-							{
-								// if this column was selected, reselect to refresh the details widgets
-								SelectColumn(ChooserTable, ColumnIndex);
-							}
-														
-						}));
-				
-						return Widget;
-					}))
-					);
-			}
-		}));
-	}
-
+	ViewModel->RegisterMenus(GetToolkitCommands());
 	
 	struct Local
 	{
@@ -457,7 +291,7 @@ void FChooserTableEditor::RegisterMenus()
 			MenuBuilder.EndSection();
 		}
 	};
-
+	
 	TSharedPtr<FExtender> MenuExtender = MakeShareable(new FExtender);
 
 	// Extend the Edit menu
@@ -473,190 +307,27 @@ void FChooserTableEditor::RegisterMenus()
 void FChooserTableEditor::InitToolMenuContext(FToolMenuContext& MenuContext)
 {
 	FAssetEditorToolkit::InitToolMenuContext(MenuContext);
-
 	UChooserEditorToolMenuContext* Context = NewObject<UChooserEditorToolMenuContext>();
 	Context->ChooserEditor = SharedThis(this);
-	MenuContext.AppendCommandList(GetToolkitCommands());
 	MenuContext.AddObject(Context);
-}
-
-void FChooserTableEditor::SaveAsset_Execute()
-{
-	AutoPopulateAll();
-	FAssetEditorToolkit::SaveAsset_Execute();
-}
-
-
-void FChooserTableEditor::BindCommands()
-{
-	const FChooserTableEditorCommands& Commands = FChooserTableEditorCommands::Get();
-
-	ToolkitCommands->MapAction(
-		Commands.EditChooserSettings,
-		FExecuteAction::CreateSP(this, &FChooserTableEditor::SelectRootProperties));
-	
-	ToolkitCommands->MapAction(
-		Commands.AutoPopulateAll,
-		FExecuteAction::CreateSP(this, &FChooserTableEditor::AutoPopulateAll));
-	
-	ToolkitCommands->MapAction(
-		Commands.RemoveDisabledData,
-		FExecuteAction::CreateSP(this, &FChooserTableEditor::RemoveDisabledData));
-	
-	ToolkitCommands->MapAction(
-		FGenericCommands::Get().Delete,
-		FExecuteAction::CreateSP(this, &FChooserTableEditor::DeleteSelection),
-		FCanExecuteAction::CreateSP(this, &FChooserTableEditor::HasSelection)
-		);
-
-	ToolkitCommands->MapAction(
-		FGenericCommands::Get().Duplicate,
-		FExecuteAction::CreateSP(this, &FChooserTableEditor::DuplicateSelection),
-		FCanExecuteAction::CreateSP(this, &FChooserTableEditor::HasSelection)
-		);
-
-	ToolkitCommands->MapAction(
-		Commands.AutoPopulateSelection,
-		FExecuteAction::CreateSP(this, &FChooserTableEditor::AutoPopulateSelection),
-		FCanExecuteAction::CreateSP(this, &FChooserTableEditor::CanAutoPopulateSelection),
-		FGetActionCheckState(),FIsActionButtonVisible::CreateSP(this, &FChooserTableEditor::HasSelection)
-		);
-	
-	ToolkitCommands->MapAction(
-		Commands.Disable,
-		FExecuteAction::CreateSP(this, &FChooserTableEditor::ToggleDisableSelection),
-		FCanExecuteAction::CreateSP(this, &FChooserTableEditor::HasSelection),
-		FIsActionChecked::CreateSP(this, &FChooserTableEditor::IsSelectionDisabled),
-		FIsActionButtonVisible::CreateSP(this, &FChooserTableEditor::HasSelection)
-		);
-	
-	ToolkitCommands->MapAction(
-		FGenericCommands::Get().Copy,
-		FExecuteAction::CreateSP(this, &FChooserTableEditor::CopySelection),
-		FCanExecuteAction::CreateSP(this, &FChooserTableEditor::HasSelection));
-	
-	ToolkitCommands->MapAction(
-		FGenericCommands::Get().Cut,
-		FExecuteAction::CreateSP(this, &FChooserTableEditor::CutSelection),
-		FCanExecuteAction::CreateSP(this, &FChooserTableEditor::HasSelection));
-	
-	ToolkitCommands->MapAction(
-		FGenericCommands::Get().Paste,
-		FExecuteAction::CreateSP(this, &FChooserTableEditor::Paste),
-		FCanExecuteAction::CreateSP(this, &FChooserTableEditor::CanPaste));
-	
-	ToolkitCommands->MapAction(
-		Commands.MoveUp,
-		FExecuteAction::CreateSP(this, &FChooserTableEditor::MoveRowsUp),
-		FCanExecuteAction::CreateSP(this, &FChooserTableEditor::CanMoveRowsUp),
-		FIsActionChecked(),
-		FIsActionButtonVisible::CreateSP(this, &FChooserTableEditor::HasRowsSelected)
-		);
-	ToolkitCommands->MapAction(
-    		Commands.MoveDown,
-    		FExecuteAction::CreateSP(this, &FChooserTableEditor::MoveRowsDown),
-    		FCanExecuteAction::CreateSP(this, &FChooserTableEditor::CanMoveRowsDown),
-    		FIsActionChecked(),
-    		FIsActionButtonVisible::CreateSP(this, &FChooserTableEditor::HasRowsSelected)
-    		);
-	
-	ToolkitCommands->MapAction(
-			Commands.MoveLeft,
-			FExecuteAction::CreateSP(this, &FChooserTableEditor::MoveColumnLeft),
-			FCanExecuteAction::CreateSP(this, &FChooserTableEditor::CanMoveColumnLeft),
-			FIsActionChecked(),
-			FIsActionButtonVisible::CreateSP(this, &FChooserTableEditor::HasColumnSelected)
-			);
-
-	ToolkitCommands->MapAction(
-			Commands.MoveRight,
-			FExecuteAction::CreateSP(this, &FChooserTableEditor::MoveColumnRight),
-			FCanExecuteAction::CreateSP(this, &FChooserTableEditor::CanMoveColumnRight),
-			FIsActionChecked(),
-			FIsActionButtonVisible::CreateSP(this, &FChooserTableEditor::HasColumnSelected)
-			);
-}
-
-void FChooserTableEditor::OnObjectsTransacted(UObject* Object, const FTransactionObjectEvent& Event)
-{
-	if (UChooserTable* ChooserTable = Cast<UChooserTable>(Object))
-	{
-		// if this is the chooser we're editing
-		if (GetChooser() == ChooserTable)
-		{
-			if (CurrentSelectionType == ESelectionType::Rows)
-			{
-				// refresh details if we have rows selected
-				RefreshRowSelectionDetails();
-			}
-		}
-	}
 	
-	if (UChooserRowDetails* RowDetails = Cast<UChooserRowDetails>(Object))
-	{
-		// if this is for the chooser we're editing
-		if (GetChooser() == RowDetails->Chooser)
-		{
-			if (RowDetails->Chooser->ResultsStructs.IsValidIndex(RowDetails->Row))
-			{
-				// copy all the values over
-				TValueOrError<FStructView, EPropertyBagResult> Result = RowDetails->Properties.GetValueStruct("Result", FInstancedStruct::StaticStruct());
-				if (Result.IsValid())
-				{
-					RowDetails->Chooser->ResultsStructs[RowDetails->Row] = Result.GetValue().Get<FInstancedStruct>();
-				}
-
-				int ColumnIndex = 0;
-				for (FInstancedStruct& ColumnData : RowDetails->Chooser->ColumnsStructs)
-				{
-					FChooserColumnBase& Column = ColumnData.GetMutable<FChooserColumnBase>();
-					Column.SetFromDetails(RowDetails->Properties, ColumnIndex, RowDetails->Row);
-					ColumnIndex++;
-				}
-			
-				TValueOrError<bool, EPropertyBagResult> DisabledResult = RowDetails->Properties.GetValueBool("Disabled");
-				RowDetails->Chooser->DisabledRows[RowDetails->Row] = DisabledResult.GetValue();
-			}
-			else if (RowDetails->Row == ColumnWidget_SpecialIndex_Fallback)
-			{
-				TValueOrError<FStructView, EPropertyBagResult> Result = RowDetails->Properties.GetValueStruct("Result", FInstancedStruct::StaticStruct());
-				if (Result.IsValid())
-				{
-					RowDetails->Chooser->FallbackResult = Result.GetValue().Get<FInstancedStruct>();
-				}
-				
-				int ColumnIndex = 0;
-				for (FInstancedStruct& ColumnData : RowDetails->Chooser->ColumnsStructs)
-				{
-					FChooserColumnBase& Column = ColumnData.GetMutable<FChooserColumnBase>();
-					Column.SetFromDetails(RowDetails->Properties, ColumnIndex, RowDetails->Row);
-					ColumnIndex++;
-				}
-			}
-			QueueRefreshAll();
-		}
-	}
+	MenuContext.AppendCommandList(ToolkitCommands);
 }
 	
-bool FChooserTableEditor::HandleTicker(float DeltaTime)
+void FChooserTableEditor::SaveAsset_Execute()
 {
-	if (bQueueRefreshAll)
-	{
-		RefreshAll();
-		bQueueRefreshAll = false;
-	}
-	return true;
+	ViewModel->AutoPopulateAll();
+	FAssetEditorToolkit::SaveAsset_Execute();
 }
 
 void FChooserTableEditor::InitEditor( const EToolkitMode::Type Mode, const TSharedPtr< class IToolkitHost >& InitToolkitHost, const TArray<UObject*>& ObjectsToEdit, FGetDetailsViewObjects GetDetailsViewObjects )
 {
-	TickDelegateHandle = FTSTicker::GetCoreTicker().AddTicker(FTickerDelegate::CreateRaw(this, &FChooserTableEditor::HandleTicker), 1.f);
-	
 	UChooserTable* Chooser = Cast<UChooserTable>(ObjectsToEdit[0]);
-	RootChooser = Chooser->GetRootChooser();
-	check(RootChooser);
+	
+	ViewModel = MakeShared<FChooserTableViewModel>(Chooser);
 
 	History.Reserve(HistorySize);
+	
 	BreadcrumbTrail = SNew(SBreadcrumbTrail<UChooserTable*>)
 		.ButtonStyle(FAppStyle::Get(), "GraphBreadcrumbButton")
 		.TextStyle(FAppStyle::Get(), "GraphBreadcrumbButtonText")
@@ -664,12 +335,12 @@ void FChooserTableEditor::InitEditor( const EToolkitMode::Type Mode, const TShar
 		.DelimiterImage( FAppStyle::GetBrush("BreadcrumbTrail.Delimiter") )
 		.OnCrumbPushed_Lambda([this](UChooserTable* Table)
 		{
-			RefreshAll();
+			ViewModel->SetChooser(BreadcrumbTrail->PeekCrumb());
 		})
 		.OnCrumbClicked_Lambda([this](UChooserTable* Table)
 		{
 			AddHistory();
-			RefreshAll();
+			ViewModel->SetChooser(BreadcrumbTrail->PeekCrumb());
 		})
 		.GetCrumbMenuContent_Lambda([this](UChooserTable* Item)
 		{
@@ -677,16 +348,25 @@ void FChooserTableEditor::InitEditor( const EToolkitMode::Type Mode, const TShar
 		})
 	;
 
-	BreadcrumbTrail->PushCrumb(FText::FromString(RootChooser->GetName()), RootChooser);
+	BreadcrumbTrail->PushCrumb(FText::FromString(Chooser->GetName()), Chooser);
 	AddHistory();
-	
+
 	FCoreUObjectDelegates::OnObjectsReplaced.AddSP(this, &FChooserTableEditor::OnObjectsReplaced);
 
 	FPropertyEditorModule& PropertyEditorModule = FModuleManager::GetModuleChecked<FPropertyEditorModule>( "PropertyEditor" );
 	FDetailsViewArgs DetailsViewArgs;
-	DetailsViewArgs.NotifyHook = this;
+	DetailsViewArgs.NotifyHook = ViewModel.Get();
 	DetailsViewArgs.NameAreaSettings = FDetailsViewArgs::HideNameArea;
 	DetailsView = PropertyEditorModule.CreateDetailView( DetailsViewArgs );
+
+	ViewModel->SetShowDetailsDelegate(FShowDetails::CreateLambda([DetailsViewWeakPtr = DetailsView.ToWeakPtr()](const TArray<UObject*>& Objects)
+		{
+			if (TSharedPtr<IDetailsView> Details = DetailsViewWeakPtr.Pin())
+			{
+				Details->SetObjects(Objects, true);
+			}
+		}));
+	
 	const TSharedRef<FTabManager::FLayout> StandaloneDefaultLayout = FTabManager::NewLayout( "Standalone_ChooserTableEditor_Layout_v1.6" )
 	->AddArea
 	(
@@ -724,21 +404,17 @@ void FChooserTableEditor::InitEditor( const EToolkitMode::Type Mode, const TShar
 	const bool bCreateDefaultToolbar = true;
 	FAssetEditorToolkit::InitAssetEditor( Mode, InitToolkitHost, FChooserTableEditor::ChooserEditorAppIdentifier, StandaloneDefaultLayout, bCreateDefaultStandaloneMenu, bCreateDefaultToolbar, ObjectsToEdit );
 
-	BindCommands();
-	
 	// todo: should these be triggered once per session by the module?
 	RegisterToolbar();
-	RegisterMenus();
+	ViewModel->RegisterMenus(GetToolkitCommands());
 	
 	RegenerateMenusAndToolbars();
 
-	SelectRootProperties();
+	ViewModel->SelectRootProperties();
 	SetChooserTableToEdit(Chooser);
 		
 	FAnimAssetFindReplaceConfig FindReplaceConfig;
 	FindReplaceConfig.InitialProcessorClass = UChooserFindProperties::StaticClass();
-	
-	FCoreUObjectDelegates::OnObjectTransacted.AddSP(this, &FChooserTableEditor::OnObjectsTransacted);
 }
 
 void FChooserTableEditor::FocusWindow(UObject* ObjectToFocusOn)
@@ -747,8 +423,6 @@ void FChooserTableEditor::FocusWindow(UObject* ObjectToFocusOn)
 	{
 		SetChooserTableToEdit(Chooser);
 	}
-	// refresh, even if we set the same chooser we were already editing. (Rewind Debugger double click enables debug testing, which requires recreating the header widgets)
-	RefreshAll();
 	FAssetEditorToolkit::FocusWindow(ObjectToFocusOn);
 }
 
@@ -762,59 +436,6 @@ FText FChooserTableEditor::GetBaseToolkitName() const
 	return LOCTEXT("AppLabel", "Chooser Table Editor");
 }
 
-void FChooserTableEditor::RefreshAll()
-{
-	if (HeaderRow)
-	{
-		// Cache Selection state
-		ESelectionType CachedSelectionType = CurrentSelectionType;
-		int SelectedColumnIndex = -1;
-		UChooserTable* SelectedChooser = nullptr;
-		TArray<int32> CachedSelectedRows;
-
-		if (CachedSelectionType == ESelectionType::Column)
-		{
-			SelectedColumnIndex = SelectedColumn->Column;
-			SelectedChooser = SelectedColumn->Chooser;
-		}
-		else if (CachedSelectionType == ESelectionType::Rows)
-		{
-			if (!SelectedRows.IsEmpty())
-			{
-				SelectedChooser = SelectedRows[0]->Chooser;
-			}
-			for(const TObjectPtr<UChooserRowDetails>& SelectedRow : SelectedRows)
-			{
-				CachedSelectedRows.Add(SelectedRow->Row);
-			}
-		}
-		
-		UpdateTableColumns();
-		UpdateTableRows();
-
-		// reapply cached selection state
-		if (CachedSelectionType == ESelectionType::Root)
-		{
-			SelectRootProperties();
-		}
-		else if (CachedSelectionType == ESelectionType::Column)
-		{
-			SelectColumn(SelectedChooser, SelectedColumnIndex);
-		}
-		else if (CachedSelectionType == ESelectionType::Rows)
-		{
-			SelectRows(CachedSelectedRows);
-		}
-	}
-	
-	RefreshNestedObjectTree();
-}
-
-void FChooserTableEditor::QueueRefreshAll()
-{
-	bQueueRefreshAll = true;
-}
-
 void FChooserTableEditor::RefreshNestedObjectTree()
 {
 	if (NestedChooserTree.IsValid())
@@ -823,70 +444,16 @@ void FChooserTableEditor::RefreshNestedObjectTree()
 	}
 }
 
-bool FChooserTableEditor::MatchesContext(const FTransactionContext& InContext, const TArray<TPair<UObject*, FTransactionObjectEvent>>& TransactionObjectContexts) const
-{
-	TArray<UObject*> ContainedObjects;
-	GetObjectsWithOuter(RootChooser->GetPackage(), ContainedObjects, true);
-	
-	for(const TPair<UObject*, FTransactionObjectEvent>&  Entry : TransactionObjectContexts)
-	{
-		if (ContainedObjects.Contains(Entry.Key))
-		{
-			return true;
-		}
-	}
-	return false;
-}
-
-void FChooserTableEditor::PostUndo(bool bSuccess)
-{
-	RefreshAll();
-}
-
-void FChooserTableEditor::PostRedo(bool bSuccess)
-{
-	RefreshAll();
-}
-
-
-void FChooserTableEditor::NotifyPreChange(FProperty* PropertyAboutToChange)
-{
-}
-
-void FChooserTableEditor::NotifyPostChange(const FPropertyChangedEvent& PropertyChangedEvent, FProperty* PropertyThatChanged)
-{
-	// Called on details panel edits
-	
-	if (CurrentSelectionType==ESelectionType::Root)
-	{
-		// Editing the root in the details panel can change ContextData that means all wigets need to be refreshed
-		UpdateTableColumns();
-		UpdateTableRows();
-		SelectRootProperties();
-	}
-	if (CurrentSelectionType==ESelectionType::Column)
-	{
-		check(SelectedColumn);
-		int SelectedColumnIndex = SelectedColumn->Column;
-		UChooserTable* SelectedColumnChooser = SelectedColumn->Chooser;
-		// Editing column properties can change the column type, which requires refreshing everything
-		UpdateTableColumns();
-		UpdateTableRows();
-		SelectColumn(SelectedColumnChooser, SelectedColumnIndex);
-	}
-	// editing row data should not require any refreshing
-}
-
 FText FChooserTableEditor::GetToolkitName() const
 {
-	check( RootChooser );
-	return FText::FromString(RootChooser->GetName());
+	check( ViewModel->GetRootChooser() );
+	return FText::FromString(ViewModel->GetRootChooser()->GetName());
 }
 
 FText FChooserTableEditor::GetToolkitToolTipText() const
 {
-	check( RootChooser );
-	return FAssetEditorToolkit::GetToolTipTextForObject(RootChooser);
+	check( ViewModel->GetRootChooser() );
+	return FAssetEditorToolkit::GetToolTipTextForObject(ViewModel->GetRootChooser());
 }
 
 FLinearColor FChooserTableEditor::GetWorldCentricTabColorScale() const
@@ -934,733 +501,29 @@ TSharedRef<SDockTab> FChooserTableEditor::SpawnFindReplaceTab( const FSpawnTabAr
 		PersonaModule.CreateFindReplaceWidget(Config)
 	];
 }
-
-TSharedRef<ITableRow> FChooserTableEditor::GenerateTableRow(TSharedPtr<FChooserTableRow> InItem, const TSharedRef<STableViewBase>& OwnerTable)
+	
+TSharedRef<SDockTab> FChooserTableEditor::SpawnNestedTablesTreeTab( const FSpawnTabArgs& Args )
 {
-	UChooserTable* Chooser = GetChooser();
+	check( Args.GetTabId() == NestedTablesTreeTabId );
 
-	return SNew(SChooserTableRow, OwnerTable)
-		.Entry(InItem).Chooser(Chooser).Editor(this);
+	return SNew(SDockTab)
+		.Label( LOCTEXT("NestedChooserTreeTitle", "Nested Choosers") )
+		[
+			SAssignNew(NestedChooserTree, SNestedChooserTree).ChooserEditor(this)
+		];
 }
 
-void FChooserTableEditor::SelectRootProperties()
+TSharedRef<SDockTab> FChooserTableEditor::SpawnTableTab( const FSpawnTabArgs& Args )
 {
-	if( DetailsView.IsValid() )
-	{
-		// point the details view to the main table
-		DetailsView->SetObject( GetRootChooser() );
-		CurrentSelectionType = ESelectionType::Root;
-	}
-}
+	check( Args.GetTabId() == TableTabId );
 
-void FChooserTableEditor::RemoveDisabledData()
-{
-	UChooserTable* Chooser = GetChooser();
-	const FScopedTransaction Transaction(LOCTEXT("Remove Disabled Data", "Remove Disabled Data"));
+	UChooserTable* Chooser = ViewModel->GetChooser();
 
-	Chooser->Modify(true);
-	Chooser->RemoveDisabledData();
-	RefreshAll();
-}
-
-int FChooserTableEditor::MoveColumn(int SourceIndex, int TargetIndex)
-{
-	UChooserTable* Chooser = GetChooser();
-	
-	TargetIndex = FMath::Clamp(TargetIndex, 0, Chooser->ColumnsStructs.Num());
-
-	if (SourceIndex < 0 || SourceIndex == TargetIndex)
-	{
-		return TargetIndex;
-	}
-
-	const FScopedTransaction Transaction(LOCTEXT("Move Row", "Move Row"));
-
-	Chooser->Modify(true);
-
-	FInstancedStruct ColumnData = Chooser->ColumnsStructs[SourceIndex];
-	Chooser->ColumnsStructs.RemoveAt(SourceIndex);
-		
-	if (SourceIndex < TargetIndex)
-	{
-		TargetIndex--;
-	}
-	
-	if (TargetIndex == Chooser->ColumnsStructs.Num())
-	{
-		if (Chooser->ColumnsStructs.Last().GetPtr<FRandomizeColumn>())
-		{
-			// never drop after a Randomize Column;
-			TargetIndex--;
-		}
-	}
-
-	Chooser->ColumnsStructs.Insert(ColumnData, TargetIndex);
-
-	RefreshAll();
-
-	return TargetIndex;
-}
-
-int FChooserTableEditor::MoveRow(int SourceRowIndex, int TargetRowIndex)
-{
-	UChooserTable* Chooser = GetChooser();
-	TargetRowIndex = FMath::Min(TargetRowIndex,Chooser->ResultsStructs.Num());
-
-	const FScopedTransaction Transaction(LOCTEXT("Move Row", "Move Row"));
-
-	Chooser->Modify(true);
-
-	for (FInstancedStruct& ColStruct : Chooser->ColumnsStructs)
-	{
-		FChooserColumnBase& Column = ColStruct.GetMutable<FChooserColumnBase>();
-		Column.MoveRow(SourceRowIndex, TargetRowIndex);
-	}
-
-	FInstancedStruct Result = Chooser->ResultsStructs[SourceRowIndex];
-	Chooser->ResultsStructs.RemoveAt(SourceRowIndex);
-	bool bDisabled = Chooser->DisabledRows[SourceRowIndex];
-	Chooser->DisabledRows.RemoveAt(SourceRowIndex);
-	if (SourceRowIndex < TargetRowIndex)
-	{
-		TargetRowIndex--;
-	}
-	Chooser->ResultsStructs.Insert(Result, TargetRowIndex);
-	Chooser->DisabledRows.Insert(bDisabled, TargetRowIndex);
-	UpdateTableRows();
-
-	return TargetRowIndex;
-}
-	
-void FChooserTableEditor::SelectRow(int32 RowIndex, bool bClear)
-{
-	if (TSharedPtr<FChooserTableRow>* Row = TableRows.FindByPredicate([RowIndex](const TSharedPtr<FChooserTableRow>& InRow)
-		{
-			return InRow->RowIndex == RowIndex;
-		}))
-	{
-		if (!TableView->IsItemSelected(*Row))
-		{
-			if (bClear)
-			{
-				TableView->ClearSelection();
-			}
-			TableView->SetItemSelection(*Row, true, ESelectInfo::OnMouseClick);
-		}
-	}
-}
-	
-void FChooserTableEditor::SelectRows(const TArrayView<int32>& Rows)
-{
-	
-	TableView->ClearSelection();
-	TArray<TSharedPtr<FChooserTableRow>> RowsToSelect;
-
-	for (TSharedPtr<FChooserTableRow>& Row : TableRows)
-	{
-		if (Rows.Contains(Row->RowIndex))
-		{
-			RowsToSelect.Add(Row);
-		}
-	}
-	
-	TableView->SetItemSelection(RowsToSelect, true);
-}
-	
-void FChooserTableEditor::ClearSelectedRows() 
-{
-	SelectedRows.SetNum(0);
-	TableView->ClearSelection();
-	SelectRootProperties();
-}
-
-bool FChooserTableEditor::IsRowSelected(int32 RowIndex)
-{
-	for(auto& SelectedRow:SelectedRows)
- 	{
- 		if (SelectedRow->Row == RowIndex)
- 		{
- 			return true;
- 		}
- 	}
-	return false;
-}
-	
-bool FChooserTableEditor::IsColumnSelected(int32 ColumnIndex)
-{
-	return  (CurrentSelectionType == ESelectionType::Column && SelectedColumn && SelectedColumn->Column == ColumnIndex);
-}
-
-void FChooserTableEditor::UpdateTableColumns()
-{
-	UChooserTable* Chooser = GetChooser();
-
-	HeaderRow->ClearColumns();
-
-	HeaderRow->AddColumn(SHeaderRow::Column("Handles")
-					.DefaultLabel(FText())
-					.ManualWidth(30));
-
-	if (Chooser->ResultType != EObjectChooserResultType::NoPrimaryResult)
-	{
-		HeaderRow->AddColumn(SHeaderRow::Column("Result")
-						.ManualWidth_Lambda([Chooser]() { return Chooser->EditorResultsColumnWidth; } )
-						.OnWidthChanged_Lambda([Chooser](float NewWidth) { Chooser->EditorResultsColumnWidth = NewWidth; })
-						.HeaderContent()
-						[
-							SNew(SVerticalBox)
-							+ SVerticalBox::Slot()
-							.VAlign(VAlign_Top)
-							[
-								SNew(STextBlock)
-									.Text(LOCTEXT("Result", "Result"))
-									.ToolTipText(LOCTEXT("ResultTooltip", "The Result is the asset which will be returned if a row is selected (or other Chooser to evaluate to get the asset to return"))
-							]
-						]);
-	}
-
-	FName ColumnId("ChooserColumn", 1);
-	int NumColumns = Chooser->ColumnsStructs.Num();	
-	for(int ColumnIndex = 0; ColumnIndex < NumColumns; ColumnIndex++)
-	{
-		FChooserColumnBase& Column = Chooser->ColumnsStructs[ColumnIndex].GetMutable<FChooserColumnBase>();
-
-		TSharedPtr<SWidget> HeaderWidget = FObjectChooserWidgetFactories::CreateColumnWidget(&Column, Chooser->ColumnsStructs[ColumnIndex].GetScriptStruct(), Chooser->GetRootChooser(), -1);
-		if (!HeaderWidget.IsValid())
-		{
-			HeaderWidget = SNullWidget::NullWidget;
-		}
-		
-		HeaderRow->AddColumn(SHeaderRow::FColumn::FArguments()
-			.ColumnId(ColumnId)
-			.ManualWidth(Column.EditorColumnWidth)
-			.ManualWidth_Lambda([&Column]() { return Column.EditorColumnWidth; } )
-			.OnWidthChanged_Lambda([&Column](float NewWidth) { Column.EditorColumnWidth = NewWidth; })
-			.HeaderComboVisibility(EHeaderComboVisibility::Ghosted)
-			.HeaderContent()
-			[
-				SNew(SChooserColumnHandle)
-					.ChooserEditor(this)
-					.ColumnIndex(ColumnIndex)
-					.NoDropAfter(Chooser->ColumnsStructs[ColumnIndex].GetPtr<FRandomizeColumn>() != nullptr)
-				[
-					HeaderWidget.ToSharedRef()
-				]
-			
-			]);
-	
-		ColumnId.SetNumber(ColumnId.GetNumber() + 1);
-	}
-
-	HeaderRow->AddColumn( SHeaderRow::FColumn::FArguments()
-		.ColumnId("Add")
-		.FillWidth(1.0)
-		.HeaderContent( )
-		[
-			SNew(SVerticalBox)
-			 + SVerticalBox::Slot().AutoHeight()
-			 [
-				SNew(SHorizontalBox)
-				+ SHorizontalBox::Slot().MaxWidth(150)
-				[
-					CreateColumnComboButton.ToSharedRef()
-				]
-			]
-		]
-		);
-
-}
-
-void FChooserTableEditor::AddColumn(const UScriptStruct* ColumnType)
-{
-	FSlateApplication::Get().DismissAllMenus();
-	UChooserTable* Chooser = GetChooser();
-	const FScopedTransaction Transaction(LOCTEXT("Add Column Transaction", "Add Column"));
-	Chooser->Modify(true);
-
-	FInstancedStruct NewColumn;
-	NewColumn.InitializeAs(ColumnType);
-	NewColumn.GetMutable<FChooserColumnBase>().Initialize(Chooser);
-	const FChooserColumnBase& NewColumnRef = NewColumn.Get<FChooserColumnBase>();
-	int InsertIndex = 0;
-	if (NewColumnRef.IsRandomizeColumn())
-	{
-		// add randomization column at the end (do nothing if there already is one)
-		InsertIndex = Chooser->ColumnsStructs.Num();
-		if (InsertIndex == 0 || !Chooser->ColumnsStructs[InsertIndex - 1].Get<FChooserColumnBase>().IsRandomizeColumn())
-		{
-			Chooser->ColumnsStructs.Add(NewColumn);
-		}
-	}
-	else if (NewColumnRef.HasOutputs())
-	{
-		// add output columns at the end (but before any randomization column)
-		InsertIndex = Chooser->ColumnsStructs.Num();
-		if (InsertIndex > 0 && Chooser->ColumnsStructs[InsertIndex - 1].Get<FChooserColumnBase>().IsRandomizeColumn())
-		{
-			InsertIndex--;
-		}
-		Chooser->ColumnsStructs.Insert(NewColumn, InsertIndex);
-	}
-	else
-	{
-		// add other columns after the last non-output, non-randomization column
-		while(InsertIndex < Chooser->ColumnsStructs.Num())
-		{
-			const FChooserColumnBase& Column = Chooser->ColumnsStructs[InsertIndex].Get<FChooserColumnBase>();
-			if (Column.HasOutputs() || Column.IsRandomizeColumn())
-			{
-				break;
-			}
-			InsertIndex++;
-		}
-		Chooser->ColumnsStructs.Insert(NewColumn, InsertIndex);
-	}
-
-	UpdateTableColumns();
-	UpdateTableRows();
-
-	SelectColumn(Chooser, InsertIndex);
-}
-
-void FChooserTableEditor::RefreshRowSelectionDetails()
-{
-	SelectedRows.SetNum(0);
-	UChooserTable* Chooser = GetChooser();
-
-	FPropertyBagPropertyDesc ResultPropertyDesc ("Result", EPropertyBagPropertyType::Struct, FInstancedStruct::StaticStruct());
-	ResultPropertyDesc.MetaData.Add({"ExcludeBaseStruct",""});
-	ResultPropertyDesc.MetaData.Add({"BaseStruct","/Script/Chooser.ObjectChooserBase"});
-	
-	// Get the list of objects to edit the details of
-	TArray<TSharedPtr<FChooserTableRow>> SelectedItems = TableView->GetSelectedItems();
-
-	if (SelectedRowsCache.Num() < SelectedItems.Num())
-	{
-		SelectedRowsCache.SetNum(SelectedItems.Num());
-	}
-	
-	int32 SelectedRowsCacheIndex = 0;
-	for(TSharedPtr<FChooserTableRow>& SelectedItem : SelectedItems)
-	{
-		if (Chooser->ResultsStructs.IsValidIndex(SelectedItem->RowIndex) || SelectedItem->RowIndex == ColumnWidget_SpecialIndex_Fallback)
-		{
-			if (SelectedRowsCache[SelectedRowsCacheIndex] == nullptr)
-			{
-				SelectedRowsCache[SelectedRowsCacheIndex] = TStrongObjectPtr(NewObject<UChooserRowDetails>(GetTransientPackage(), NAME_None, RF_Transactional | RF_Transient));
-			}
-			TObjectPtr<UChooserRowDetails> Selection = SelectedRowsCache[SelectedRowsCacheIndex].Get();
-			SelectedRowsCacheIndex ++;
-			
-			Selection->Chooser = Chooser;
-			Selection->Row = SelectedItem->RowIndex;
-			Selection->Properties.Reset();
-		
-			if (Chooser->ResultType != EObjectChooserResultType::NoPrimaryResult)
-			{
-				FInstancedStruct& Result = SelectedItem->RowIndex == ColumnWidget_SpecialIndex_Fallback? Chooser->FallbackResult : Chooser->ResultsStructs[SelectedItem->RowIndex];;
-				Selection->Properties.AddProperties({ResultPropertyDesc});
-				Selection->Properties.SetValueStruct("Result", FConstStructView(FInstancedStruct::StaticStruct(), reinterpret_cast<uint8*>(&Result)));
-			}
-
-			int ColumnIndex = 0;
-			for (FInstancedStruct& ColumnData : Chooser->ColumnsStructs)
-			{
-				FChooserColumnBase& Column = ColumnData.GetMutable<FChooserColumnBase>();
-
-				if (Column.HasOutputs() || SelectedItem->RowIndex != ColumnWidget_SpecialIndex_Fallback)
-				{
-					// add each column's details, but on the fallback row, only add details for Output columns
-					Column.AddToDetails(Selection->Properties, ColumnIndex, SelectedItem->RowIndex);
-				}
-				ColumnIndex++;
-			}
-
-			if (Chooser->DisabledRows.IsValidIndex(SelectedItem->RowIndex))
-			{
-				Selection->Properties.AddProperty("Disabled", EPropertyBagPropertyType::Bool);
-				Selection->Properties.SetValueBool("Disabled", Chooser->DisabledRows[SelectedItem->RowIndex]);
-			}
-
-			SelectedRows.Add(Selection);
-		}
-	}
-	
-	TArray<UObject*> DetailsObjects;
-	for(auto& Item : SelectedRows)
-	{
-		DetailsObjects.Add(Item.Get());
-	}
-
-	if( DetailsView.IsValid() )
-	{
-		// Make sure details window is pointing to our object
-		DetailsView->SetObjects( DetailsObjects );
-		DetailsView->ForceRefresh();
-	}
-}
-
-void FChooserTableEditor::MakeChoosersMenuRecursive(UObject* Outer, FMenuBuilder& MenuBuilder, const FString& Indent = "") 
-{
-	TArray<UObject*> ChildObjects;
-	GetObjectsWithOuter(Outer, ChildObjects, false);
-
-	FString SubIndent = Indent + "    ";
-	for (UObject* Object : ChildObjects)
-	{
-		if (UChooserTable* Chooser = Cast<UChooserTable>(Object))
-		{
-			if (Chooser == RootChooser || Chooser->GetRootChooser()->NestedObjects.Contains(Chooser))
-			{
-				MenuBuilder.AddMenuEntry( FText::FromString(Indent + Chooser->GetName()), LOCTEXT("Edit Chooser ToolTip", "Browse to this Nested Chooser Table"), FSlateIcon(),
-					FUIAction(FExecuteAction::CreateLambda([this, Chooser]()
-					{
-						SetChooserTableToEdit(Chooser);
-					})));
-
-				MakeChoosersMenuRecursive(Chooser, MenuBuilder, SubIndent);
-			}
-		}
-	}
-}
-	
-TSharedRef<SWidget> FChooserTableEditor::MakeChoosersMenu(UObject* RootObject)
-{
-	FMenuBuilder MenuBuilder(true, nullptr);
-
-	MakeChoosersMenuRecursive(RootObject, MenuBuilder);
-
-	return MenuBuilder.MakeWidget();
-}
-
-TSharedRef<SDockTab> FChooserTableEditor::SpawnNestedTablesTreeTab( const FSpawnTabArgs& Args )
-{
-	check( Args.GetTabId() == NestedTablesTreeTabId );
-
-	return SNew(SDockTab)
-		.Label( LOCTEXT("NestedChooserTreeTitle", "Nested Choosers") )
-		[
-			SAssignNew(NestedChooserTree, SNestedChooserTree).ChooserEditor(this)
-		];
-}
-
-DECLARE_DELEGATE_OneParam(FCreateStructDelegate, UScriptStruct*);
-
-struct FColumnTypeInfoStruct
-{
-	int SortOrder = 100;
-	FString Category;
-	UScriptStruct* Type;
-	bool operator < (const FColumnTypeInfoStruct& Other) const
-	{
-		if (Category == Other.Category)
-		{
-			return Type->GetDisplayNameText().ToString() < Other.Type->GetDisplayNameText().ToString();
-		}
-		else if (SortOrder == Other.SortOrder)
-		{
-			return Category < Other.Category;
-		}
-		else
-		{
-			return SortOrder < Other.SortOrder;
-		}
-	}
-};
-
-TSharedRef<SWidget>	FChooserTableEditor::MakeCreateColumnMenu()
-{
-	FMenuBuilder MenuBuilder(true, nullptr);
-	static TArray<FColumnTypeInfoStruct> ColumnTypes;
-
-	if (ColumnTypes.IsEmpty())
-	{
-		UScriptStruct* BaseType = FChooserColumnBase::StaticStruct();
-		for (TObjectIterator<UScriptStruct> StructIt; StructIt; ++StructIt)
-		{
-			if (*StructIt != BaseType && StructIt->IsChildOf(BaseType))
-			{
-				if (!StructIt->HasMetaData("Hidden"))
-				{
-					FColumnTypeInfoStruct Info;
-					Info.Type = *StructIt;
-					Info.Category = StructIt->HasMetaData("Category") ? StructIt->GetMetaData("Category") : "Other";
-
-					if (Info.Category == "Filter")
-					{
-						Info.SortOrder = 1;
-					}
-					else if (Info.Category == "Scoring")
-                   	{
-                   		Info.SortOrder = 2;
-                   	}
-					else if (Info.Category == "Output")
-                   	{
-                   		Info.SortOrder = 3;
-                   	}
-					else if (Info.Category == "Random")
-					{
-						Info.SortOrder = 4;
-					}
-
-					ColumnTypes.Add(Info);
-				}
-			}
-		}
-		ColumnTypes.Sort();
-	}
-
-	FString Section = "";
-	for(FColumnTypeInfoStruct& Type : ColumnTypes)
-	{
-		if (Section != Type.Category)
-		{
-			if (Section != "")
-			{
-				MenuBuilder.EndSection();
-			}
-			Section = Type.Category;
-			MenuBuilder.BeginSection(FName(Section), FText::FromString(Section));
-		}
-		
-		MenuBuilder.AddMenuEntry(Type.Type->GetDisplayNameText(), Type.Type->GetToolTipText(), FSlateIcon(),
-			FUIAction(
-				FExecuteAction::CreateLambda([this, Type]()
-				{
-					AddColumn(Type.Type);
-				}))
-			);
-							
-	}
-	return MenuBuilder.MakeWidget();
-}
-
-struct FResultTypeInfoStruct
-{
-	bool ObjectOnly = false;
-	bool ClassOnly = false;
-	FString Category;
-	UScriptStruct* Type = nullptr;
-	
-	bool operator < (const FResultTypeInfoStruct& Other) const
-	{
-		if (Category == Other.Category)
-		{
-			return Type->GetDisplayNameText().ToString() < Other.Type->GetDisplayNameText().ToString();
-		}
-		else
-		{
-			return Category < Other.Category;
-		}
-	}
-};
-
-void MakeCreateResultMenu(FMenuBuilder& MenuBuilder, EObjectChooserResultType ChooserResultType, FCreateStructDelegate CreateStruct)
-{
-	static TArray<FResultTypeInfoStruct> ResultTypes;
-
-	if (ResultTypes.IsEmpty())
-	{
-		UScriptStruct* BaseType = FObjectChooserBase::StaticStruct();
-		for (TObjectIterator<UScriptStruct> StructIt; StructIt; ++StructIt)
-		{
-			if (*StructIt != BaseType && StructIt->IsChildOf(BaseType))
-			{
-				if (!StructIt->HasMetaData("Hidden"))
-				{
-					FResultTypeInfoStruct Info;
-					Info.Type = *StructIt;
-					Info.Category = StructIt->HasMetaData("Category") ? StructIt->GetMetaData("Category") : "Other";
-
-					if (StructIt->HasMetaData("ResultType"))
-					{
-						FString ResultTypeString = StructIt->GetMetaData("ResultType");
-						Info.ClassOnly = ResultTypeString == "Class";
-						Info.ObjectOnly = ResultTypeString == "Object";
-					}
-					
-					ResultTypes.Add(Info);
-				}
-			}
-		}
-		ResultTypes.Sort();
-	}
-
-	FString Section = "";
-	for(FResultTypeInfoStruct& Type : ResultTypes)
-	{
-		if (Section != Type.Category)
-		{
-			if (Section != "")
-			{
-				MenuBuilder.EndSection();
-			}
-			Section = Type.Category;
-			MenuBuilder.BeginSection(FName(Section), FText::FromString(Section));
-		}
-		
-		MenuBuilder.AddMenuEntry(Type.Type->GetDisplayNameText(), Type.Type->GetToolTipText(), FSlateIcon(),
-			FUIAction(
-				FExecuteAction::CreateLambda([Type, CreateStruct]()
-				{
-					CreateStruct.Execute(Type.Type);
-				}),
-				FCanExecuteAction::CreateLambda([Type, ChooserResultType]()
-				{
-					if (Type.ClassOnly && ChooserResultType == EObjectChooserResultType::ObjectResult)
-					{
-						return false;
-					}
-					if (Type.ObjectOnly && ChooserResultType == EObjectChooserResultType::ClassResult)
-					{
-						return false;
-					}
-					return true;
-				})
-				)
-			);
-							
-	}
-}
-
-TSharedRef<SWidget>	FChooserTableEditor::MakeCreateRowMenu()
-{
-	FMenuBuilder MenuBuilder(true, nullptr);
-
-    UChooserTable* Chooser = GetChooser();
-	if (!Chooser->FallbackResult.IsValid())
-	{
-		if (Chooser->ResultType == EObjectChooserResultType::NoPrimaryResult)
-		{
-			MenuBuilder.AddMenuEntry(
-				LOCTEXT("Add Fallback Output", "Add Fallback Output"),
-				LOCTEXT("Add Fallback Output Tooltip", "Add a Fallback row to the chooser, which will be used in the case where no other rows passed all filter columns"),
-				FSlateIcon(),
-				FUIAction(FExecuteAction::CreateLambda([this]()
-				{
-					UChooserTable* Chooser = GetChooser();
-					const FScopedTransaction Transaction(LOCTEXT("Add Fallback Row Transaction", "Add Fallback Row"));
-					Chooser->Modify(true);
-					
-					// Just construct a dummy result to make sure all rows always have "valid results"
-					// You can't just leave a null result otherwise rows don't apply their output.
-					Chooser->FallbackResult.InitializeAs(FClassChooser::StaticStruct());
-					Chooser->FallbackResult.GetMutable<FClassChooser>().Class = UClass::StaticClass();
-					
-					UpdateTableRows();
-				}))
-			);
-		}
-		else
-		{
-			MenuBuilder.AddSubMenu(
-				LOCTEXT("Add Fallback", "Add Fallback Result"),
-					LOCTEXT("Add Fallback Tooltip", "Add a Fallback row to the chooser, which will be used in the case where no other rows passed all filter columns"),
-					FNewMenuDelegate::CreateLambda([this](FMenuBuilder& MenuBuilder)
-					{
-						UChooserTable* Chooser = GetChooser();
-						MakeCreateResultMenu(MenuBuilder, Chooser->GetContextOwner()->ResultType, FCreateStructDelegate::CreateLambda([this](UScriptStruct* Type)
-						{
-							UChooserTable* Chooser = GetChooser();
-							const FScopedTransaction Transaction(LOCTEXT("Add Fallback Row Transaction", "Add Fallback Row"));
-							Chooser->Modify(true);
-							
-							Chooser->FallbackResult.InitializeAs(Type);
-							
-							UpdateTableRows();
-						}));
-					})
-				);
-		}
-	}
-	
-	if (Chooser->ResultType == EObjectChooserResultType::NoPrimaryResult)
-	{
-		MenuBuilder.AddMenuEntry(
-			LOCTEXT("Add Output Row", "Add Output Row"),
-			LOCTEXT("Add Output Row Tooltip", "Add a regular row to the chooser"),
-			FSlateIcon(),
-			FUIAction(FExecuteAction::CreateLambda([this]()
-			{
-				UChooserTable* Chooser = GetChooser();
-				const FScopedTransaction Transaction(LOCTEXT("Add Row Transaction", "Add Row"));
-				Chooser->Modify(true);
-
-				// Just construct a dummy result to make sure all rows always have "valid results"
-				// You can't just leave a null result otherwise rows don't apply their output.
-				FInstancedStruct& NewResult = Chooser->ResultsStructs.AddDefaulted_GetRef();
-				NewResult.InitializeAs(FClassChooser::StaticStruct());
-				NewResult.GetMutable<FClassChooser>().Class = UClass::StaticClass();
-				
-				UpdateTableRows();
-			}))
-		);
-	}
-	else
-	{
-		MakeCreateResultMenu(MenuBuilder, Chooser->GetContextOwner()->ResultType, FCreateStructDelegate::CreateLambda([this](UScriptStruct* Type)
-		{
-			UChooserTable* Chooser = GetChooser();
-			const FScopedTransaction Transaction(LOCTEXT("Add Row Transaction", "Add Row"));
-			Chooser->Modify(true);
-			
-			FInstancedStruct& NewResult = Chooser->ResultsStructs.AddDefaulted_GetRef();
-			NewResult.InitializeAs(Type);
-			
-			UpdateTableRows();
-		}));
-	}
-
-	return MenuBuilder.MakeWidget();
-}
-
-TSharedRef<SDockTab> FChooserTableEditor::SpawnTableTab( const FSpawnTabArgs& Args )
-{
-	check( Args.GetTabId() == TableTabId );
-
-	UChooserTable* Chooser = GetChooser();
-
-	// + button to create new columns
-	
-	CreateColumnComboButton = SNew(SPositiveActionButton)
-		.Text(LOCTEXT("Add Column", "Add Column"))
-		.OnGetMenuContent_Lambda([this]()
-		{
-			return MakeCreateColumnMenu();
-		});
-
-	CreateRowComboButton = SNew(SPositiveActionButton)
-		.Text(LOCTEXT("Add Row", "Add Row"))
-		.OnGetMenuContent(this, &FChooserTableEditor::MakeCreateRowMenu);
-			
-	HeaderRow = SNew(SHeaderRow);
-
-	UpdateTableRows();
-	UpdateTableColumns();
-
-	TableView = SNew(SListView<TSharedPtr<FChooserTableRow>>)
-				.OnKeyDownHandler_Lambda([this](const FGeometry&, const FKeyEvent& Event)
-				{
-					if (GetToolkitCommands()->ProcessCommandBindings(Event))
-					{
-						return FReply::Handled();
-					}
-					return FReply::Unhandled();
-				})
-    			.ListItemsSource(&TableRows)
-				.OnContextMenuOpening_Raw(this, &FChooserTableEditor::GenerateRowContextMenu)
-				.OnSelectionChanged_Lambda([this](TSharedPtr<FChooserTableRow>,  ESelectInfo::Type SelectInfo)
-				{
-					// deselect any selected column
-					ClearSelectedColumn();
-
-					CurrentSelectionType = ESelectionType::Rows;
-
-					RefreshRowSelectionDetails();
-				})
-    			.OnGenerateRow_Raw(this, &FChooserTableEditor::GenerateTableRow)
-				.HeaderRow(HeaderRow);
+	TSharedRef<SWidget> ChooserTableView = SNew(SChooserTableWidget)
+											.Commands(GetToolkitCommands())
+											.ViewModel(ViewModel);
 
+	ViewModel->RefreshAll();
 
 	TSharedRef<SComboButton> EditChooserTableButton = SNew(SComboButton)
 		.ButtonStyle(FAppStyle::Get(), "GraphBreadcrumbButton");
@@ -1669,7 +532,7 @@ TSharedRef<SDockTab> FChooserTableEditor::SpawnTableTab( const FSpawnTabArgs& Ar
     		FOnGetContent::CreateLambda(
     			[this]()
                 		{
-							return MakeChoosersMenu(GetRootChooser()->GetPackage());
+							return MakeChoosersMenu(ViewModel->GetRootChooser()->GetPackage());
                 		})
     		);
 
@@ -1686,7 +549,7 @@ TSharedRef<SDockTab> FChooserTableEditor::SpawnTableTab( const FSpawnTabArgs& Ar
 				[
 					SNew(SButton)
 						.ButtonStyle(FAppStyle::Get(), "GraphBreadcrumbButton")
-						.IsEnabled_Raw(this, &FChooserTableEditor::CanNavigateBack)
+						.IsEnabled(this, &FChooserTableEditor::CanNavigateBack)
 						.OnClicked_Lambda([this]()
 						{
 							NavigateBack();
@@ -1702,7 +565,7 @@ TSharedRef<SDockTab> FChooserTableEditor::SpawnTableTab( const FSpawnTabArgs& Ar
 				[
 					SNew(SButton)
 						.ButtonStyle(FAppStyle::Get(), "GraphBreadcrumbButton")
-						.IsEnabled_Raw(this, &FChooserTableEditor::CanNavigateForward)
+						.IsEnabled(this, &FChooserTableEditor::CanNavigateForward)
 						.OnClicked_Lambda([this]()
 						{
 							NavigateForward();
@@ -1730,952 +593,179 @@ TSharedRef<SDockTab> FChooserTableEditor::SpawnTableTab( const FSpawnTabArgs& Ar
 				SNew(SScrollBox).Orientation(Orient_Horizontal)
 				+ SScrollBox::Slot()
 				[
-					TableView.ToSharedRef()
-				]
-			]
-		];
-}
-
-void FChooserTableEditor::UpdateTableRows()
-{
-	UChooserTable* Chooser = GetChooser();
-	int32 NewNum = Chooser->ResultsStructs.Num();
-	Chooser->DisabledRows.SetNum(NewNum);
-
-	// Sync the TableRows array which drives the ui table to match the number of results.
-	TableRows.SetNum(0, EAllowShrinking::No);
-	for(int i =0; i < NewNum; i++)
-	{
-		TableRows.Add(MakeShared<FChooserTableRow>(i));
-	}
-
-	// Add one at the end, for the Fallback result
-	if (Chooser->FallbackResult.IsValid())
-	{
-		TableRows.Add(MakeShared<FChooserTableRow>(SChooserTableRow::SpecialIndex_Fallback));
-	}
-	
-	// Add one at the end, for the "Add Row" control
-	TableRows.Add(MakeShared<FChooserTableRow>(SChooserTableRow::SpecialIndex_AddRow));
-
-	// Make sure each column has the same number of row datas as there are results
-	for(FInstancedStruct& ColumnData : Chooser->ColumnsStructs)
-	{
-		FChooserColumnBase& Column = ColumnData.GetMutable<FChooserColumnBase>();
-		Column.SetNumRows(NewNum);
-	}
-
-	if (TableView.IsValid())
-	{
-		TableView->RebuildList();
-	}
-}
-
-void FChooserTableEditor::OnObjectsReplaced(const TMap<UObject*, UObject*>& ReplacementMap)
-{
-	bool bChangedAny = false;
-
-	UObject* ReplacedObject = ReplacementMap.FindRef(RootChooser);
-
-	if (ReplacedObject && ReplacedObject != RootChooser)
-	{
-		RootChooser = Cast<UChooserTable>(ReplacedObject);
-		SetChooserTableToEdit(RootChooser);
-		SelectRootProperties();
-	}
-}
-
-FString FChooserTableEditor::GetWorldCentricTabPrefix() const
-{
-	return LOCTEXT("WorldCentricTabPrefix", "Chooser Table Asset ").ToString();
-}
-
-TSharedRef<FChooserTableEditor> FChooserTableEditor::CreateEditor( const EToolkitMode::Type Mode, const TSharedPtr< IToolkitHost >& InitToolkitHost, UObject* ObjectToEdit, FGetDetailsViewObjects GetDetailsViewObjects )
-{
-	TSharedRef< FChooserTableEditor > NewEditor( new FChooserTableEditor() );
-
-	TArray<UObject*> ObjectsToEdit;
-	ObjectsToEdit.Add( ObjectToEdit );
-	NewEditor->InitEditor( Mode, InitToolkitHost, ObjectsToEdit, GetDetailsViewObjects );
-
-	return NewEditor;
-}
-
-TSharedRef<FChooserTableEditor> FChooserTableEditor::CreateEditor( const EToolkitMode::Type Mode, const TSharedPtr< IToolkitHost >& InitToolkitHost, const TArray<UObject*>& ObjectsToEdit, FGetDetailsViewObjects GetDetailsViewObjects )
-{
-	TSharedRef< FChooserTableEditor > NewEditor( new FChooserTableEditor() );
-	NewEditor->InitEditor( Mode, InitToolkitHost, ObjectsToEdit, GetDetailsViewObjects );
-	return NewEditor;
-}
-	
-void FChooserTableEditor::SelectColumn(UChooserTable* ChooserEditor, int Index)
-{
-	ClearSelectedRows();
-	
-	UChooserTable* Chooser = GetChooser();
-   	if (Index < Chooser->ColumnsStructs.Num())
-   	{
-   		if (SelectedColumn == nullptr)
-   		{
-   			SelectedColumn = TStrongObjectPtr(NewObject<UChooserColumnDetails>(GetTransientPackage(), NAME_None, RF_Transactional | RF_Transient));
-   		}
-   
-   		SelectedColumn->Chooser = Chooser;
-   		SelectedColumn->Column = Index;
-   		DetailsView->SetObject(SelectedColumn.Get(), true);
-   		CurrentSelectionType = ESelectionType::Column;
-   	}
-   	else
-   	{
-   		SelectRootProperties();
-   	}
-}
-	
-void FChooserTableEditor::ClearSelectedColumn()
-{
-	UChooserTable* Chooser = GetChooser();
-	if (SelectedColumn != nullptr)
-	{
-		SelectedColumn->Column = -1;
-		if (DetailsView->GetSelectedObjects().Contains(SelectedColumn.Get()))
-		{
-			SelectRootProperties();
-		}
-	}
-}
-	
-void FChooserTableEditor::DeleteColumn(int Index)
-{
-	const FScopedTransaction Transaction(LOCTEXT("Delete Column Transaction", "Delete Column"));
-	ClearSelectedColumn();
-	SelectRootProperties();
-	UChooserTable* Chooser = GetChooser();
-
-	if (Index < Chooser->ColumnsStructs.Num())
-	{
-		Chooser->Modify(true);
-		Chooser->ColumnsStructs.RemoveAt(Index);
-		UpdateTableColumns();
-	}
-}
-	
-int32 FChooserTableEditor::DeleteSelectedRows(int32 RowIndexToRemember)
-{
-	const FScopedTransaction Transaction(LOCTEXT("Delete Row Transaction", "Delete Row"));
-	return DeleteSelectedRowsInternal(RowIndexToRemember);
-}
-	
-int32 FChooserTableEditor::DeleteSelectedRowsInternal(int32 RowIndexToRemember)
-{
-	UChooserTable* Chooser = GetChooser();
-	Chooser->Modify(true);
-	// delete selected rows.
-	TArray<uint32> RowsToDelete;
-	for(auto& SelectedRow:SelectedRows)
-	{
-		if (SelectedRow->Row == ColumnWidget_SpecialIndex_Fallback)
-		{
-			Chooser->FallbackResult.Reset();
-		}
-		else
-		{
-			RowsToDelete.Add(SelectedRow->Row);
-		}
-	}
-
-	SelectedRows.SetNum(0);
-	SelectRootProperties();
-
-	// sort indices in reverse
-	RowsToDelete.Sort([](int32 A, int32 B){ return A>B; });
-	for(int32 RowIndex : RowsToDelete)
-	{
-		if (RowIndex <= RowIndexToRemember)
-		{
-			RowIndexToRemember--;
-		}
-		Chooser->ResultsStructs.RemoveAt(RowIndex);
-		Chooser->DisabledRows.RemoveAt(RowIndex);
-	}
-
-	for(FInstancedStruct& ColumnData : Chooser->ColumnsStructs)
-	{
-		FChooserColumnBase& Column = ColumnData.GetMutable<FChooserColumnBase>();
-		Column.DeleteRows(RowsToDelete);
-	}
-	UpdateTableRows();
-
-	return RowIndexToRemember;
-}
-
-void FChooserTableEditor::MoveRows(int TargetIndex)
-{
-	const FScopedTransaction Transaction(LOCTEXT("Move Row(s)", "Move Row(s)"));
-	UChooserTable* RowCopy = CopySelectionInternal();
-	TargetIndex = DeleteSelectedRowsInternal(TargetIndex);
-	PasteInternal(RowCopy, TargetIndex);
-}
-
-void FChooserTableEditor::AutoPopulateColumn(FChooserColumnBase& Column)
-{
-	if (UChooserTable* Chooser = GetChooser())
-	{
-		const int RowCount = Chooser->ResultsStructs.Num();
-		if (Column.AutoPopulates())
-		{
-			for (int i = 0; i < RowCount; ++i)
-			{
-				UObject* ReferencedObject = Chooser->ResultsStructs[i].IsValid() ? Chooser->ResultsStructs[i].Get<FObjectChooserBase>().GetReferencedObject() : nullptr;
-				Column.AutoPopulate(i, ReferencedObject);
-			}
-
-			UObject* ReferencedObject = Chooser->FallbackResult.IsValid() ? Chooser->FallbackResult.Get<FObjectChooserBase>().GetReferencedObject() : nullptr;
-			Column.AutoPopulate(ColumnWidget_SpecialIndex_Fallback, ReferencedObject);
-		}
-	}
-}
-
-void FChooserTableEditor::AutoPopulateRow(int Index)
-{
-	if (UChooserTable* Chooser = GetChooser())
-	{
-		if (Chooser->ResultsStructs.IsValidIndex(Index) && Chooser->ResultsStructs[Index].IsValid())
-		{
-			if (UObject* ReferencedObject = Chooser->ResultsStructs[Index].Get<FObjectChooserBase>().GetReferencedObject()) 
-			{
-				for (FInstancedStruct& ColumnData : Chooser->ColumnsStructs)
-				{
-					FChooserColumnBase& Column = ColumnData.GetMutable<FChooserColumnBase>();
-					Column.AutoPopulate(Index, ReferencedObject);
-				}
-			}
-		}
-		else if (Index == ColumnWidget_SpecialIndex_Fallback && Chooser->FallbackResult.IsValid())
-		{
-			if (UObject* ReferencedObject = Chooser->FallbackResult.Get<FObjectChooserBase>().GetReferencedObject())
-			{
-				for (FInstancedStruct& ColumnData : Chooser->ColumnsStructs)
-				{
-					FChooserColumnBase& Column = ColumnData.GetMutable<FChooserColumnBase>();
-					Column.AutoPopulate(Index, ReferencedObject);
-				}
-			}
-		}
-	}
-}
-	
-bool FChooserTableEditor::CanAutoPopulateSelection()
-{
-	if (UChooserTable* Chooser = GetChooser())
-	{
-		if (CurrentSelectionType == ESelectionType::Column && SelectedColumn)
-		{
-			// when a column is selected, return true if that column supports auto populate
-			if (Chooser->ColumnsStructs.IsValidIndex(SelectedColumn->Column))
-			{
-				return Chooser->ColumnsStructs[SelectedColumn->Column].Get<FChooserColumnBase>().AutoPopulates();
-			}	
-		}
-		else
-		{
-			if (SelectedRows.IsEmpty())
-			{
-				return false;
-			}
-			
-			// when rows are selected, return true if any column supports auto populate
-			for (FInstancedStruct& ColumnData : Chooser->ColumnsStructs)
-			{
-				const FChooserColumnBase& Column = ColumnData.Get<FChooserColumnBase>();
-				if (Column.AutoPopulates())
-				{
-					return true;
-				}
-			}
-		}
-	}
-	
-	return false;
-}
-	
-void FChooserTableEditor::AutoPopulateSelection()
-{
-	if (UChooserTable* Chooser = GetChooser())
-	{
-		const FScopedTransaction Transaction(LOCTEXT("Auto Populate Selection", "Auto Populate Selection"));
-		Chooser->Modify();
-		if (HasColumnSelected())
-		{
-			if (Chooser->ColumnsStructs.IsValidIndex(SelectedColumn->Column))
-			{
-				AutoPopulateColumn(Chooser->ColumnsStructs[SelectedColumn->Column].GetMutable<FChooserColumnBase>());
-			}
-		}
-		else if (HasRowsSelected())
-		{
-			for(UChooserRowDetails* RowDetails : SelectedRows)
-			{
-				AutoPopulateRow(RowDetails->Row);
-			}
-		}
-	}
-}
-
-void FChooserTableEditor::AutoPopulateAll()
-{
-	if (UChooserTable* Chooser = GetChooser())
-	{
-		const FScopedTransaction Transaction(LOCTEXT("Auto Populate Chooser", "Auto Populate All"));
-		Chooser->Modify();
-		for (FInstancedStruct& ColumnData : Chooser->ColumnsStructs)
-		{
-			FChooserColumnBase& Column = ColumnData.GetMutable<FChooserColumnBase>();
-			AutoPopulateColumn(Column);
-		}
-	}
-}
-
-bool FChooserTableEditor::HasSelection() const
-{
-	return HasRowsSelected() || HasColumnSelected();
-}
-
-bool FChooserTableEditor::HasRowsSelected() const
-{
-	return CurrentSelectionType == ESelectionType::Rows && !SelectedRows.IsEmpty();
-}
-	
-bool FChooserTableEditor::HasColumnSelected() const
-{
-	const UChooserTable* Chooser = GetChooser();
-	return CurrentSelectionType == ESelectionType::Column && SelectedColumn && Chooser->ColumnsStructs.IsValidIndex(SelectedColumn->Column);
-}
-
-bool FChooserTableEditor::IsSelectionDisabled() const
-{
-	if (const UChooserTable* Chooser = GetChooser())
-	{
-		if (HasColumnSelected())
-		{
-			if ( Chooser->ColumnsStructs.IsValidIndex(SelectedColumn->Column))
-			{
-				const FChooserColumnBase& Column = Chooser->ColumnsStructs[SelectedColumn->Column].Get<FChooserColumnBase>();
-				return Column.bDisabled;
-			}
-		}
-		else if (HasRowsSelected())
-		{
-			bool bSomethingEnabled = false;
-			for(auto& Row : SelectedRows)
-			{
-				if (!Chooser->IsRowDisabled(Row->Row))
-				{
-					bSomethingEnabled = true;
-					break;
-				}
-			}
-			return !bSomethingEnabled;
-		}
-	}
-	return false;
-}
-
-void FChooserTableEditor::ToggleDisableSelection()
-{
-	bool bDisabled = IsSelectionDisabled();
-	if (UChooserTable* Chooser = GetChooser())
-	{
-		if (HasColumnSelected())
-		{
-			if (Chooser->ColumnsStructs.IsValidIndex(SelectedColumn->Column))
-			{
-				FChooserColumnBase& Column = Chooser->ColumnsStructs[SelectedColumn->Column].GetMutable<FChooserColumnBase>();
-				Column.bDisabled = !Column.bDisabled;
-			}
-		}
-		else if (HasRowsSelected())
-		{
-			for (auto& Row : SelectedRows)
-			{
-				if (Chooser->DisabledRows.IsValidIndex(Row->Row))
-				{
-					Chooser->DisabledRows[Row->Row] = !bDisabled;
-				}
-			}
-			RefreshRowSelectionDetails();
-		}
-	}
-}
-
-void FChooserTableEditor::DeleteSelection()
-{
-	if (HasColumnSelected())
-	{
-		DeleteColumn(SelectedColumn->Column);
-	}
-	else if (HasRowsSelected())
-	{
-		DeleteSelectedRows();
-	}
-}
-
-void FChooserTableEditor::DuplicateSelection()
-{
-	if (HasRowsSelected())
-	{
-		const FScopedTransaction Transaction(LOCTEXT("Duplicate Row(s)", "Duplicate Row(s)"));
-		UChooserTable* RowCopy = CopySelectionInternal();
-		int MaxSelectedRow = -1;
-		for(UChooserRowDetails* SelectedRow : SelectedRows)
-		{
-			MaxSelectedRow = FMath::Max(SelectedRow->Row, MaxSelectedRow);
-		}
-		PasteInternal(RowCopy, MaxSelectedRow + 1);
-	}
-	else if(HasColumnSelected())
-	{
-		const FScopedTransaction Transaction(LOCTEXT("Duplicate Column", "Duplicate Column"));
-		UChooserTable* Chooser = GetChooser();
-		Chooser->Modify();
-		FInstancedStruct Column = Chooser->ColumnsStructs[SelectedColumn->Column];
-		Chooser->ColumnsStructs.Insert(Column,SelectedColumn->Column);
-		RefreshAll();
-	}
-}
-
-bool FChooserTableEditor::HasFallbackSelected()
-{
-	for(UChooserRowDetails* SelectedRow : SelectedRows)
-	{
-		if (SelectedRow->Row == ColumnWidget_SpecialIndex_Fallback)
-		{
-			return true;
-		}
-	}
-
-	return false;
+					ChooserTableView
+				]
+			]
+		];
 }
 
-bool FChooserTableEditor::CanMoveRowsUp()
+void FChooserTableEditor::OnObjectsReplaced(const TMap<UObject*, UObject*>& ReplacementMap)
 {
-	if (HasRowsSelected())
-	{
-		UChooserTable* Chooser = GetChooser();
+	bool bChangedAny = false;
 
-		int MinSelectedRow = Chooser->ResultsStructs.Num();
-		for(UChooserRowDetails* SelectedRow : SelectedRows)
-		{
-			if (SelectedRow->Row != ColumnWidget_SpecialIndex_Fallback)
-			{
-				MinSelectedRow = FMath::Min(SelectedRow->Row, MinSelectedRow);
-			}
-		}
+	UChooserTable* RootChooser = ViewModel->GetRootChooser();
+
+	UObject* ReplacedObject = ReplacementMap.FindRef(RootChooser);
 
-		return MinSelectedRow > 0;
+	if (ReplacedObject && ReplacedObject != RootChooser)
+	{
+		RootChooser = Cast<UChooserTable>(ReplacedObject);
+		SetChooserTableToEdit(RootChooser);
+		ViewModel->SelectRootProperties();
 	}
-	return false;
 }
 
-void FChooserTableEditor::MoveRowsUp()
+FString FChooserTableEditor::GetWorldCentricTabPrefix() const
 {
-	if (HasRowsSelected())
-	{
-		UChooserTable* Chooser = GetChooser();
-		int MinSelectedRow = Chooser->ResultsStructs.Num();
-		for(UChooserRowDetails* SelectedRow : SelectedRows)
-		{
-			if (SelectedRow->Row != ColumnWidget_SpecialIndex_Fallback)
-			{
-				MinSelectedRow = FMath::Min(SelectedRow->Row, MinSelectedRow);
-			}
-		}
-		MoveRows(MinSelectedRow - 1);
-	}
+	return LOCTEXT("WorldCentricTabPrefix", "Chooser Table Asset ").ToString();
 }
 
-bool FChooserTableEditor::CanMoveRowsDown()
+TSharedRef<FChooserTableEditor> FChooserTableEditor::CreateEditor( const EToolkitMode::Type Mode, const TSharedPtr< IToolkitHost >& InitToolkitHost, UObject* ObjectToEdit, FGetDetailsViewObjects GetDetailsViewObjects )
 {
-	if (HasRowsSelected())
-	{
-		UChooserTable* Chooser = GetChooser();
-		
-		int MaxSelectedRow = -1;
-		for(UChooserRowDetails* SelectedRow : SelectedRows)
-		{
-			MaxSelectedRow = FMath::Max(SelectedRow->Row, MaxSelectedRow);
-		}
+	TSharedRef< FChooserTableEditor > NewEditor( new FChooserTableEditor() );
 
-		return MaxSelectedRow < Chooser->ResultsStructs.Num() - 1;
-	}
-	return false;
+	TArray<UObject*> ObjectsToEdit;
+	ObjectsToEdit.Add( ObjectToEdit );
+	NewEditor->InitEditor( Mode, InitToolkitHost, ObjectsToEdit, GetDetailsViewObjects );
+
+	return NewEditor;
 }
 
-void FChooserTableEditor::MoveRowsDown()
+TSharedRef<FChooserTableEditor> FChooserTableEditor::CreateEditor( const EToolkitMode::Type Mode, const TSharedPtr< IToolkitHost >& InitToolkitHost, const TArray<UObject*>& ObjectsToEdit, FGetDetailsViewObjects GetDetailsViewObjects )
 {
-	if (HasRowsSelected())
-   	{
-   		UChooserTable* Chooser = GetChooser();
-   		int MaxSelectedRow = -1;
-   		for(UChooserRowDetails* SelectedRow : SelectedRows)
-   		{
-   			MaxSelectedRow = FMath::Max(SelectedRow->Row, MaxSelectedRow);
-   		}
-   		MoveRows(MaxSelectedRow + 2);
-   	}
+	TSharedRef< FChooserTableEditor > NewEditor( new FChooserTableEditor() );
+	NewEditor->InitEditor( Mode, InitToolkitHost, ObjectsToEdit, GetDetailsViewObjects );
+	return NewEditor;
 }
 
-bool FChooserTableEditor::CanMoveColumnLeft()
+
+void FChooserTableEditor::AddHistory()
 {
-	if (HasColumnSelected())
+	// remove anything ahead of this in the history, if we had gone back
+	while (HistoryIndex !=0)
 	{
-		UChooserTable* Chooser = GetChooser();
-		
-		if (Chooser->ColumnsStructs[SelectedColumn->Column].GetPtr<FRandomizeColumn>())
-		{
-			return false;
-		}
-		
-		return SelectedColumn->Column > 0;
+		History.PopFront();
+		HistoryIndex--;
 	}
-	return false;
-}
-
-void FChooserTableEditor::MoveColumnLeft()
-{
-	if (CanMoveColumnLeft())
+	
+	if (History.Num() >= HistorySize)
 	{
-		SelectColumn(GetChooser(), MoveColumn(SelectedColumn->Column, SelectedColumn->Column - 1));
+		History.Pop();
 	}
+	History.AddFront(ViewModel->GetChooser());
 }
 
-bool FChooserTableEditor::CanMoveColumnRight()
+bool FChooserTableEditor::CanNavigateBack() const
 {
-	if (HasColumnSelected())
-   	{
-		UChooserTable* Chooser = GetChooser();
-		
-		if (Chooser->ColumnsStructs[SelectedColumn->Column].GetPtr<FRandomizeColumn>())
-		{
-			return false;
-		}
-		int NumColumns = Chooser->ColumnsStructs.Num();
-		if (NumColumns > 0 && Chooser->ColumnsStructs.Last().GetPtr<FRandomizeColumn>())
-		{
-			NumColumns--;
-		}
-   		return (SelectedColumn->Column < NumColumns-1);
-   	}
-   	return false;
+	return HistoryIndex < History.Num() - 1;
 }
 
-void FChooserTableEditor::MoveColumnRight()
+void FChooserTableEditor::NavigateBack()
 {
-	if (CanMoveColumnRight())
+	if (HistoryIndex < History.Num() - 1)
 	{
-		SelectColumn(GetChooser(), MoveColumn(SelectedColumn->Column, SelectedColumn->Column + 2));
+		HistoryIndex++;
+		SetChooserTableToEdit(History[HistoryIndex], false);
 	}
 }
 
-UChooserTable* DuplicateNestedChooser(UChooserTable* Chooser, UChooserTable* NewOuter)
+bool FChooserTableEditor::CanNavigateForward() const
 {
-	UChooserTable* RootTable = NewOuter->GetRootChooser();
-	if (TObjectPtr<UObject>* FoundTable = RootTable->NestedObjects.FindByPredicate([Chooser](UObject* Object)
-		{
-			if (Cast<UChooserTable>(Object))
-			{
-				return Object->GetName() == Chooser->GetName();
-			}
-			return false;
-		}))
-	{
-		// we already duplicated this table
-		return Cast<UChooserTable>(*FoundTable);
-	}
-	
-	UChooserTable* NewTable = NewObject<UChooserTable>(NewOuter, Chooser->GetFName(), RF_Transactional);
-	NewTable->ColumnsStructs = Chooser->ColumnsStructs;
-	NewTable->ResultsStructs = Chooser->ResultsStructs;
-	NewTable->RootChooser = RootTable;
-	RootTable->AddNestedObject(NewTable);
-
-	for (FInstancedStruct& ResultData : NewTable->ResultsStructs)
-	{
-		if (FNestedChooser* NestedChooser = ResultData.GetMutablePtr<FNestedChooser>())
-		{
-			NestedChooser->Chooser = DuplicateNestedChooser(NestedChooser->Chooser, NewOuter);
-		}
-	}
-
-	return NewTable;
+	return HistoryIndex > 0;
 }
-	
-UChooserTable* FChooserTableEditor::CopySelectionInternal()
-{
-	UChooserTable* CopyData = NewObject<UChooserTable>(GetTransientPackage());
-
-	UChooserTable* Chooser = GetChooser();
 
-	// copy context data from root table
-	CopyData->OutputObjectType = RootChooser->OutputObjectType;
-	CopyData->ResultType = RootChooser->ResultType;
-	CopyData->ContextData = RootChooser->ContextData;
-
-	if (CurrentSelectionType == ESelectionType::Column)
-	{
-		// add selected column including all the cell data
-		CopyData->ColumnsStructs.Add(Chooser->ColumnsStructs[SelectedColumn->Column]);
-	}
-	else if (CurrentSelectionType == ESelectionType::Rows)
+void FChooserTableEditor::NavigateForward()
+{
+	if (HistoryIndex > 0)
 	{
-		TArray<UChooserRowDetails*> SelectedRowsCopy = SelectedRows;
-		Algo::Sort(SelectedRowsCopy, [](UChooserRowDetails *A, UChooserRowDetails *B) { return A->Row<B->Row;});
-		
-		if (SelectedRowsCopy.Num() > 0 && SelectedRowsCopy[0]->Row == ColumnWidget_SpecialIndex_Fallback)
-		{
-			SelectedRowsCopy.RemoveAt(0);
-			
-			CopyData->FallbackResult = Chooser->FallbackResult;
-			if (FNestedChooser* CopiedNestedChooser = CopyData->FallbackResult.GetMutablePtr<FNestedChooser>())
-			{
-				// if the fallback result was a nested chooser, duplicate it
-				CopiedNestedChooser->Chooser = DuplicateNestedChooser(CopiedNestedChooser->Chooser, CopyData);
-			}
-		}
-		
-		CopyData->ResultsStructs.SetNum(SelectedRowsCopy.Num());
-		CopyData->DisabledRows.SetNum(SelectedRowsCopy.Num());
-
-		// add the selected results and column data
-		
-		for (int RowIndex = 0; RowIndex < SelectedRowsCopy.Num(); RowIndex++)
-		{
-			CopyData->ResultsStructs[RowIndex] = Chooser->ResultsStructs[SelectedRowsCopy[RowIndex]->Row];
-			if (FNestedChooser* CopiedNestedChooser = CopyData->ResultsStructs[RowIndex].GetMutablePtr<FNestedChooser>())
-			{
-				if (CopiedNestedChooser->Chooser)
-				{
-					// if the result for this row was a nested chooser (with a valid chooser assigned), duplicate it
-					CopiedNestedChooser->Chooser = DuplicateNestedChooser(CopiedNestedChooser->Chooser, CopyData);
-				}
-			}
-
-			CopyData->DisabledRows[RowIndex] = Chooser->DisabledRows[SelectedRowsCopy[RowIndex]->Row];
-		}
-
-		// copy all columns
-		CopyData->ColumnsStructs = Chooser->ColumnsStructs;
-		for (FInstancedStruct& ColumnData : CopyData->ColumnsStructs)
-		{
-			FChooserColumnBase& Column = ColumnData.GetMutable<FChooserColumnBase>();
-			Column.Initialize(CopyData);
-		}
-
-		// clear all column's cell data
-		for (FInstancedStruct& ColumnData : CopyData->ColumnsStructs)
-		{
-			FChooserColumnBase& Column = ColumnData.GetMutable<FChooserColumnBase>();
-			Column.SetNumRows(0);
-			Column.SetNumRows(SelectedRowsCopy.Num());
-		}
-
-		for (int RowIndex = 0; RowIndex < SelectedRowsCopy.Num(); RowIndex++)
-		{
-			for (int ColumnIndex = 0; ColumnIndex < CopyData->ColumnsStructs.Num(); ColumnIndex++)
-			{
-				FChooserColumnBase& SourceColumn = Chooser->ColumnsStructs[ColumnIndex].GetMutable<FChooserColumnBase>();
-				FChooserColumnBase& TargetColumn = CopyData->ColumnsStructs[ColumnIndex].GetMutable<FChooserColumnBase>();
-				TargetColumn.CopyRow(SourceColumn, SelectedRowsCopy[RowIndex]->Row, RowIndex);
-			}
-		}
+		HistoryIndex--;
+		SetChooserTableToEdit(History[HistoryIndex], false);
 	}
-	
-	return CopyData;
 }
 
-void FChooserTableEditor::CopySelection()
+void FChooserTableEditor::SetChooserTableToEdit(UChooserTable* Chooser, bool bApplyToHistory)
 {
-	UChooserTable* CopyData = CopySelectionInternal();
-
-	// Clear the mark state for saving.
-	UnMarkAllObjects(EObjectMark(OBJECTMARK_TagExp | OBJECTMARK_TagImp));
-	
-	// Export the clipboard to text.
-	FStringOutputDevice Archive;
-	const FExportObjectInnerContext Context;
-	UExporter::ExportToOutputDevice(&Context,CopyData, nullptr, Archive, TEXT("copy"), 0, PPF_ExportsNotFullyQualified | PPF_Copy | PPF_Delimited, false, CopyData->GetOuter());
-	FPlatformApplicationMisc::ClipboardCopy(*Archive);
-}
+	if (Chooser == ViewModel->GetChooser())
+	{
+		return;
+	}
 	
-void FChooserTableEditor::CutSelection()
-{
-	CopySelection();
-	DeleteSelection();
-}
+	BreadcrumbTrail->ClearCrumbs();
 
-class FChooserClipboardFactory : public FCustomizableTextObjectFactory
-{
-public:
+	TArray<UChooserTable*> OuterList;
+	OuterList.Push(Chooser);
 	
-	FChooserClipboardFactory()
-		: FCustomizableTextObjectFactory(GWarn)
-		, ClipboardContent(nullptr) 
+	while(OuterList.Last() != ViewModel->GetRootChooser())
 	{
+		OuterList.Push(Cast<UChooserTable>(OuterList.Last()->GetOuter()));
 	}
 
-	UChooserTable* ClipboardContent;
-	
-protected:
-
-	virtual bool CanCreateClass(UClass* InObjectClass, bool& bOmitSubObjs) const override
+	while(!OuterList.IsEmpty())
 	{
-		if (InObjectClass->IsChildOf(UChooserTable::StaticClass()))
-		{
-			return true;
-		}
-		return false;
+		UChooserTable* Popped = OuterList.Pop();
+		BreadcrumbTrail->PushCrumb(FText::FromString(Popped->GetName()), Popped);
 	}
 	
-	virtual void ProcessConstructedObject(UObject* CreatedObject) override
+	if (bApplyToHistory)
 	{
-		if (CreatedObject->IsA<UChooserTable>())
-		{
-			ClipboardContent = CastChecked<UChooserTable>(CreatedObject);
-		}
+		AddHistory();
 	}
-};
+	
+	ViewModel->SetChooser(Chooser);
+}
 
-static FString GetColumnName(FChooserColumnBase& Column)
+void FChooserTableEditor::PushChooserTableToEdit(UChooserTable* Chooser)
 {
-	if (FChooserParameterBase* InputValue = Column.GetInputValue())
-	{
-		return InputValue->GetDebugName();
-	}
-	return FString();
+	BreadcrumbTrail->PushCrumb(FText::FromString(Chooser->GetName()), Chooser);
+	AddHistory();
+	ViewModel->SetChooser(Chooser);
 }
 	
-void FChooserTableEditor::PasteInternal(UChooserTable* PastedContent, int PasteRowIndex)
+void FChooserTableEditor::PopChooserTableToEdit()
 {
-	UChooserTable* Chooser = GetChooser();
-	Chooser->Modify();
-	
-	if (PastedContent->ResultsStructs.IsEmpty() && !PastedContent->FallbackResult.IsValid())
+	if (BreadcrumbTrail->HasCrumbs())
 	{
-		// pasting a column
-		int InsertColumnIndex = Chooser->ColumnsStructs.Num();
-		if (CurrentSelectionType == ESelectionType::Column && SelectedColumn)
-		{
-			InsertColumnIndex = FMath::Min(InsertColumnIndex, SelectedColumn->Column + 1);
-		}
-		
-		if (Chooser->ColumnsStructs.Num() > 0 && Chooser->ColumnsStructs.Num() == InsertColumnIndex)
-		{
-			// if were inserting at the end, there is a randomize column, insert new columns before it
-			if (Chooser->ColumnsStructs.Last().GetPtr<FRandomizeColumn>())
-			{
-				InsertColumnIndex--;
-			}
-		}
-		Chooser->ColumnsStructs.Insert(PastedContent->ColumnsStructs, InsertColumnIndex);
-		SelectColumn(Chooser, InsertColumnIndex);
+		BreadcrumbTrail->PopCrumb();
+		ViewModel->SetChooser(BreadcrumbTrail->PeekCrumb());
 	}
-	else
-	{
-		// pasting rows
-		int RowsToPaste = PastedContent->ResultsStructs.Num();
-
-		// figure out where to start inserting
-		int InsertIndex = Chooser->ResultsStructs.Num();
-
-		if (PasteRowIndex >=0)
-		{
-			InsertIndex = PasteRowIndex;
-		}
-		else
-		{
-			if (SelectedRows.Num() > 0)
-			{
-				InsertIndex = SelectedRows[0]->Row;
-				for(int SelectedRowIndex = 1; SelectedRowIndex < SelectedRows.Num(); SelectedRowIndex ++)
-				{
-					InsertIndex = FMath::Max(InsertIndex, SelectedRows[SelectedRowIndex]->Row);
-				}
-				if (InsertIndex == ColumnWidget_SpecialIndex_Fallback)
-				{
-					// if the only row selected was the fallback, reset insert index to the last row
-					InsertIndex = Chooser->ResultsStructs.Num();
-				}
-				else
-				{
-					InsertIndex++;
-				}
-			}
-		}
-
-		if (PastedContent->ResultsStructs.Num() > 0)
-		{
-			Chooser->ResultsStructs.Insert(PastedContent->ResultsStructs, InsertIndex);
-			Chooser->DisabledRows.Insert(PastedContent->DisabledRows, InsertIndex);
-
-			// Make sure each column has the same number of row datas as there are results
-			for(FInstancedStruct& ColumnData : Chooser->ColumnsStructs)
-			{
-				FChooserColumnBase& Column = ColumnData.GetMutable<FChooserColumnBase>();
-				Column.InsertRows(InsertIndex, RowsToPaste);
-			}
-		}
-		if (PastedContent->FallbackResult.IsValid())
-		{
-			// paste fallback result if copy data has one
-			Chooser->FallbackResult = PastedContent->FallbackResult;
-			if (FNestedChooser* NestedChooser = Chooser->FallbackResult.GetMutablePtr<FNestedChooser>())
-			{
-				// duplicate the nested chooser if the fallback result refers to a nested chooser
-				NestedChooser->Chooser = DuplicateNestedChooser(NestedChooser->Chooser, Chooser);
-			}
-		}
-		
-		if (!PastedContent->NestedObjects.IsEmpty())
-		{
-			// if there were nested choosers in the copy buffer we have to remap or paste them here
-			for (int ResultIndex = InsertIndex; ResultIndex < PastedContent->ResultsStructs.Num() + InsertIndex; ResultIndex++)
-			{
-				if (FNestedChooser* NestedChooser = Chooser->ResultsStructs[ResultIndex].GetMutablePtr<FNestedChooser>())
-				{
-					NestedChooser->Chooser = DuplicateNestedChooser(NestedChooser->Chooser, Chooser);
-				}
-			}
-		}
-
-		// try to also paste column data from columns in the paste buffer which match the columns in the current chooser
-		// -- matching by column type and input value name
+}
 
-		// keep track of target columns that have already been matched, to avoid matching multiple source columns with the same target column
-		TArray<bool> MatchedTargetColumns;
-		MatchedTargetColumns.SetNum(Chooser->ColumnsStructs.Num());
+	
+void FChooserTableEditor::MakeChoosersMenuRecursive(UObject* Outer, FMenuBuilder& MenuBuilder, const FString& Indent = "") 
+{
+	TArray<UObject*> ChildObjects;
+	GetObjectsWithOuter(Outer, ChildObjects, false);
 
-		// keep track of which source columns were matched, so we can add new columns for the unmatched ones after
-		TArray<bool> MatchedSourceColumns;
-		MatchedSourceColumns.SetNum(PastedContent->ColumnsStructs.Num());
-		
-		for(int SourceColumnIndex = 0; SourceColumnIndex < PastedContent->ColumnsStructs.Num(); SourceColumnIndex++)
+	FString SubIndent = Indent + "    ";
+	for (UObject* Object : ChildObjects)
+	{
+		if (UChooserTable* Chooser = Cast<UChooserTable>(Object))
 		{
-			FInstancedStruct& PastedColumnData = PastedContent->ColumnsStructs[SourceColumnIndex];
-			FChooserColumnBase& PastedColumn = PastedColumnData.GetMutable<FChooserColumnBase>();
-			FString PastedColumnName = GetColumnName(PastedColumn);
-			for(int TargetColumnIndex = 0; TargetColumnIndex < Chooser->ColumnsStructs.Num(); TargetColumnIndex++)
+			if (Chooser == ViewModel->GetRootChooser() || Chooser->GetRootChooser()->NestedObjects.Contains(Chooser))
 			{
-				if (!MatchedTargetColumns[TargetColumnIndex])
-				{
-					FInstancedStruct& ColumnData = Chooser->ColumnsStructs[TargetColumnIndex];
-					if (ColumnData.GetScriptStruct() == PastedColumnData.GetScriptStruct())
+				MenuBuilder.AddMenuEntry( FText::FromString(Indent + Chooser->GetName()), LOCTEXT("Edit Chooser ToolTip", "Browse to this Nested Chooser Table"), FSlateIcon(),
+					FUIAction(FExecuteAction::CreateLambda([this, Chooser]()
 					{
-						FChooserColumnBase& Column = ColumnData.GetMutable<FChooserColumnBase>();
-						FString ColumnName = GetColumnName(Column);
-
-						if (ColumnName == PastedColumnName)
-						{
-							MatchedTargetColumns[TargetColumnIndex] = true;
-							MatchedSourceColumns[SourceColumnIndex] = true;
-
-							// found a match, copy the data over
-							for (int i = 0; i < RowsToPaste; i++)
-							{
-								Column.CopyRow(PastedColumn, i, InsertIndex + i);
-							}
-
-							if (PastedContent->FallbackResult.IsValid())
-							{
-								// if the fallback row was copied, paste the fallback data for columns
-								Column.CopyFallback(PastedColumn);
-							}
-							break;
-						}
-					}
-				}
-			}
-		}
-		
-		// add new columns for any source columns that were unmatched
-		
-		int InsertColumnIndex = Chooser->ColumnsStructs.Num();
-		if (Chooser->ColumnsStructs.Num() > 0)
-		{
-			// if there is a randomize column, insert new columns before it
-			if (Chooser->ColumnsStructs.Last().GetPtr<FRandomizeColumn>())
-			{
-				InsertColumnIndex--;
-			}
-		}
+						SetChooserTableToEdit(Chooser);
+					})));
 
-		for(int SourceColumnIndex = 0; SourceColumnIndex < PastedContent->ColumnsStructs.Num(); SourceColumnIndex++)
-		{
-			if (!MatchedSourceColumns[SourceColumnIndex])
-			{
-				FInstancedStruct& PastedColumnData = PastedContent->ColumnsStructs[SourceColumnIndex];
-				FChooserColumnBase& PastedColumn = PastedColumnData.GetMutable<FChooserColumnBase>();
-				// if we couldn't find a match, paste a new column
-				Chooser->ColumnsStructs.Insert(PastedColumnData, InsertColumnIndex);
-				FChooserColumnBase& Column = Chooser->ColumnsStructs[InsertColumnIndex].GetMutable<FChooserColumnBase>();
-				InsertColumnIndex++;
-				Column.SetNumRows(0);
-				Column.SetNumRows(Chooser->ResultsStructs.Num());
-				for (int i = 0; i < RowsToPaste; i++)
-				{
-					Column.CopyRow(PastedColumn, i, InsertIndex + i);
-				}
+				MakeChoosersMenuRecursive(Chooser, MenuBuilder, SubIndent);
 			}
 		}
-
-		ClearSelectedRows();
-		RefreshAll();
-
-		TArray<int32, TInlineAllocator<256>> SelectedIndices;
-		SelectedIndices.Reserve(RowsToPaste+1);
-		
-		// select the inserted rows
-		for (int i = 0; i < RowsToPaste; i++)
-		{
-			SelectedIndices.Add(InsertIndex + i);
-		}
-		if(PastedContent->FallbackResult.IsValid())
-		{
-			SelectedIndices.Add(ColumnWidget_SpecialIndex_Fallback);
-		}
-
-		SelectRows(SelectedIndices);
 	}
 }
 	
-void FChooserTableEditor::Paste()
+TSharedRef<SWidget> FChooserTableEditor::MakeChoosersMenu(UObject* RootObject)
 {
-	// Get the text from the clipboard.
-	FString ClipboardText;
-	FPlatformApplicationMisc::ClipboardPaste(ClipboardText);
-	FChooserClipboardFactory Factory;
+	FMenuBuilder MenuBuilder(true, nullptr);
 
-	UChooserTable* Chooser = GetChooser();
+	MakeChoosersMenuRecursive(RootObject, MenuBuilder);
 
-	if (Factory.CanCreateObjectsFromText(ClipboardText))
-	{
-		Factory.ProcessBuffer((UObject*)GetTransientPackage(), RF_Transactional, ClipboardText);
-		if (UChooserTable* PastedContent = Factory.ClipboardContent)
-		{
-			FScopedTransaction Transaction(LOCTEXT("Paste Chooser Data", "Paste Chooser Data"));
-			PasteInternal(PastedContent);
-		}
-	}
+	return MenuBuilder.MakeWidget();
 }
 
 	
-bool FChooserTableEditor::CanPaste() const
-{
-	// Get the text from the clipboard.
-	FString ClipboardText;
-	FPlatformApplicationMisc::ClipboardPaste(ClipboardText);
-	FChooserClipboardFactory Factory;
-	return Factory.CanCreateObjectsFromText(ClipboardText); 
-}
-
+	
 void FChooserTableEditor::RegisterWidgets()
 {
 	FObjectChooserWidgetFactories::RegisterWidgetCreator(FAssetChooser::StaticStruct(), CreateAssetWidget);
@@ -2690,6 +780,7 @@ void FChooserTableEditor::RegisterWidgets()
 	PropertyModule.RegisterCustomClassLayout("ChooserRowDetails", FOnGetDetailCustomizationInstance::CreateStatic(&FChooserRowDetails::MakeInstance));	
 	PropertyModule.RegisterCustomClassLayout("ChooserColumnDetails", FOnGetDetailCustomizationInstance::CreateStatic(&FChooserColumnDetails::MakeInstance));	
 }
+
 }
 
 #undef LOCTEXT_NAMESPACE
diff --git a/Engine/Plugins/Chooser/Source/ChooserEditor/Private/ChooserTableEditor.h b/Engine/Plugins/Chooser/Source/ChooserEditor/Private/ChooserTableEditor.h
--- a/Engine/Plugins/Chooser/Source/ChooserEditor/Private/ChooserTableEditor.h
+++ b/Engine/Plugins/Chooser/Source/ChooserEditor/Private/ChooserTableEditor.h
@@ -13,7 +13,7 @@
 #include "Toolkits/IToolkitHost.h"
 #include "Widgets/Navigation/SBreadcrumbTrail.h"
 #include "Widgets/Views/SListView.h"
-#include "Widgets/Views/STableViewBase.h"
+#include "ChooserTableViewModel.h"
 #include "ChooserTableEditor.generated.h"
 
 class SComboButton;
@@ -22,24 +22,13 @@ class SEditableText;
 class IDetailsView;
 class UChooserRowDetails;
 
-// Class used for chooser editor details customization
-UCLASS()
-class UChooserColumnDetails : public UObject
-{
-	GENERATED_BODY()
-public:
-	UPROPERTY(EditAnywhere, Instanced, Category="Hidden")
-	TObjectPtr<UChooserTable> Chooser;
-	int Column = -1;
-};
-
-
 namespace UE::ChooserEditor
 {
+	class SChooserTableWidget;
 	class SNestedChooserTree;
 	struct FChooserTableRow;
-	
-	class FChooserTableEditor : public FAssetEditorToolkit, public FSelfRegisteringEditorUndoClient, public FNotifyHook
+
+	class FChooserTableEditor : public FAssetEditorToolkit
 	{
 	public:
 		
@@ -61,6 +50,8 @@ namespace UE::ChooserEditor
 
 		virtual void FocusWindow(UObject* ObjectToFocusOn) override;
 
+		FChooserTableEditor();
+		
 		/** Destructor */
 		virtual ~FChooserTableEditor();
 
@@ -78,105 +69,27 @@ namespace UE::ChooserEditor
 		virtual void InitToolMenuContext(FToolMenuContext& MenuContext) override;
 		virtual void SaveAsset_Execute() override;
 
-		/** FEditorUndoClient Interface */
-		virtual bool MatchesContext( const FTransactionContext& InContext, const TArray<TPair<UObject*, FTransactionObjectEvent>>& TransactionObjectContexts ) const override;
-		virtual void PostUndo(bool bSuccess) override;
-		virtual void PostRedo(bool bSuccess) override;
-		
-		/** Begin FNotifyHook Interface */
-		virtual void NotifyPreChange( FProperty* PropertyAboutToChange ) override;
-		virtual void NotifyPostChange( const FPropertyChangedEvent& PropertyChangedEvent, FProperty* PropertyThatChanged) override;
-
-		const UChooserTable* GetRootChooser() const { return RootChooser; }
-		UChooserTable* GetRootChooser() { return RootChooser; }
-		UChooserTable* GetChooser() { return BreadcrumbTrail->PeekCrumb(); }
-		const UChooserTable* GetChooser() const { return BreadcrumbTrail->PeekCrumb(); }
 
-		void PushChooserTableToEdit(UChooserTable* Chooser);
-		void PopChooserTableToEdit();
-		void RefreshAll();
-		void QueueRefreshAll();
-		void RefreshNestedObjectTree();
-	
 		/** Used to show or hide certain properties */
 		void SetPropertyVisibilityDelegate(FIsPropertyVisible InVisibilityDelegate);
 		/** Can be used to disable the details view making it read-only */
 		void SetPropertyEditingEnabledDelegate(FIsPropertyEditingEnabled InPropertyEditingDelegate);
-
-		bool HasSelection() const;
-		bool HasRowsSelected() const;
-		bool HasColumnSelected() const;
-		bool IsSelectionDisabled() const;
-		void ToggleDisableSelection();
-		void DeleteSelection();
-		void DuplicateSelection();
-		bool HasFallbackSelected();
-
-		bool CanMoveRowsUp();
-		void MoveRowsUp();
-		bool CanMoveRowsDown();
-		void MoveRowsDown();
-		bool CanMoveColumnLeft();
-		void MoveColumnLeft();
-		bool CanMoveColumnRight();
-		void MoveColumnRight();
-		
-		void CopySelection();
-        void CutSelection();
-        		
-
-		void AutoPopulateColumn(FChooserColumnBase& Column);
-		void AutoPopulateRow(int Index);
-		bool CanAutoPopulateSelection();
-		void AutoPopulateSelection();
-		void AutoPopulateAll();
-		
-		void UpdateTableRows();
-		void SelectColumn(UChooserTable* Chooser, int Index);
-		void ClearSelectedColumn();
-		void DeleteColumn(int Index);
-		void AddColumn(const UScriptStruct* ColumnType);
-		void RefreshRowSelectionDetails();
-		TSharedRef<SWidget> MakeChoosersMenu(UObject* RootObject);
-		void MakeChoosersMenuRecursive(UObject* Outer, FMenuBuilder& MenuBuilder, const FString& Indent);
-		int32 DeleteSelectedRows(int32 RowIndexToRemember = 0);
-		int32 DeleteSelectedRowsInternal(int32 RowIndexToRemember = 0);
-		void MoveRows(int TargetIndex);
-		int MoveRow(int SourceRowIndex, int TargetIndex);
-		int MoveColumn(int SourceIndex, int TargetIndex);
-		void SelectRows(const TArrayView<int32>& Rows);
-		void SelectRow(int32 RowIndex, bool bClear = true);
-		void ClearSelectedRows(); 
-		bool IsRowSelected(int32 RowIndex);
-		bool IsColumnSelected(int32 ColumnIndex);
 		
+		UChooserTable* GetRootChooser();
+		void PushChooserTableToEdit(UChooserTable* Chooser);
 		void SetChooserTableToEdit(UChooserTable* Chooser, bool bApplyToHistory = true);
-
-		enum class ESelectionType
-		{
-			Root, Rows, Column
-		};
-		
-		ESelectionType GetCurrentSelectionType() const { return CurrentSelectionType; }
-
-		UChooserTable* CopySelectionInternal();
-		void PasteInternal(UChooserTable* PasteObject, int PasteRowIndex=-1);
-
-		bool TableHasFocus() const { return TableView->HasKeyboardFocus(); }
 	private:
+		void AddHistory();
+		bool CanNavigateBack() const;
+		void NavigateBack();
+		bool CanNavigateForward() const;
+		void NavigateForward();
+		void PopChooserTableToEdit();
 		
-		void SelectRootProperties();
-		void RemoveDisabledData();
 		void RegisterToolbar();
 		void RegisterMenus();
 
-
-		void Paste();
-		bool CanPaste() const;
-		void BindCommands();
-		void OnObjectsTransacted(UObject* Object, const FTransactionObjectEvent& Event);
 		void MakeDebugTargetMenu(UToolMenu* InToolMenu);
-		TSharedPtr<SWidget> GenerateRowContextMenu();
 	
 		/** Create the properties tab and its content */
 		TSharedRef<SDockTab> SpawnPropertiesTab( const FSpawnTabArgs& Args );
@@ -187,21 +100,13 @@ namespace UE::ChooserEditor
 		/** Create the nested tables list tab and its content */
 		TSharedRef<SDockTab> SpawnNestedTablesTreeTab( const FSpawnTabArgs& Args );
 	
-		TSharedRef<ITableRow> GenerateTableRow(TSharedPtr<FChooserTableRow> InItem, const TSharedRef<STableViewBase>& OwnerTable);
-
 		/** Called when objects need to be swapped out for new versions, like after a blueprint recompile. */
 		void OnObjectsReplaced(const TMap<UObject*, UObject*>& ReplacementMap);
 
 		/** Details view */
 		TSharedPtr< class IDetailsView > DetailsView;
 
-		
-		bool bQueueRefreshAll = false;
-
-		FTSTicker::FDelegateHandle TickDelegateHandle;
-		
-		bool HandleTicker(float DeltaTime);
-		
+	
 		/** App Identifier. */
 		static const FName ChooserEditorAppIdentifier;
 
@@ -211,40 +116,22 @@ namespace UE::ChooserEditor
 		static const FName TableTabId;
 		static const FName NestedTablesTreeTabId;
 		
-		void AddHistory();
-		bool CanNavigateBack() const;
-		void NavigateBack();
-		bool CanNavigateForward() const;
-		void NavigateForward();
-
-		/** The root chooser asset being edited in this editor */
-		UChooserTable* RootChooser = nullptr;
-
-		TStrongObjectPtr<UChooserColumnDetails> SelectedColumn;
-		TArray<TObjectPtr<UChooserRowDetails>> SelectedRows;
-		// cache of objects to reuse for SelectedRows
-		TArray<TStrongObjectPtr<UChooserRowDetails>> SelectedRowsCache;
-
-		TSharedPtr<SBreadcrumbTrail<UChooserTable*>> BreadcrumbTrail;
+		void RefreshNestedObjectTree();
+	
+		TSharedPtr<IChooserTableViewModel> ViewModel;
+		
+		TSharedPtr<SChooserTableWidget> ChooserTableWidget;
+		
 		TRingBuffer<UChooserTable*> History;
 		int32 HistoryIndex = 0;
+		TSharedPtr<SBreadcrumbTrail<UChooserTable*>> BreadcrumbTrail;
+		
+		
+		TSharedRef<SWidget> MakeChoosersMenu(UObject* RootObject);
+		void MakeChoosersMenuRecursive(UObject* Outer, FMenuBuilder& MenuBuilder, const FString& Indent);
 		
-		void UpdateTableColumns();
-		TArray<TSharedPtr<FChooserTableRow>> TableRows;
-	
-		TSharedPtr<SPositiveActionButton> CreateColumnComboButton;
-		TSharedPtr<SWidget> CreateRowComboButton;
-
-		TSharedPtr<SHeaderRow> HeaderRow;
-		TSharedPtr<SListView<TSharedPtr<FChooserTableRow>>> TableView;
-
-		ESelectionType CurrentSelectionType = ESelectionType::Root;
-
-		TSharedRef<SWidget>	MakeCreateRowMenu();
-		TSharedRef<SWidget>	MakeCreateColumnMenu();
 	public:
 
-		TSharedPtr<SWidget>& GetCreateRowComboButton() { return CreateRowComboButton; };
 		TSharedPtr<SNestedChooserTree> NestedChooserTree;
 
 		/** The name given to all instances of this type of editor */
diff --git a/Engine/Plugins/Chooser/Source/ChooserEditor/Private/ChooserTableViewModel.cpp b/Engine/Plugins/Chooser/Source/ChooserEditor/Private/ChooserTableViewModel.cpp
new file mode 100644
--- /dev/null
+++ b/Engine/Plugins/Chooser/Source/ChooserEditor/Private/ChooserTableViewModel.cpp
@@ -0,0 +1,1396 @@
+// Copyright Epic Games, Inc. All Rights Reserved.
+
+#include "ChooserTableViewModel.h"
+
+#include "Chooser.h"
+#include "ChooserDetails.h"
+#include "ChooserEditorWidgets.h"
+#include "ChooserFindProperties.h"
+#include "ChooserTableEditorCommands.h"
+#include "ClassViewerFilter.h"
+#include "DetailCategoryBuilder.h"
+#include "Factories.h"
+#include "GraphEditorSettings.h"
+#include "IPropertyAccessEditor.h"
+#include "LandscapeRender.h"
+#include "ObjectChooserClassFilter.h"
+#include "ObjectChooserWidgetFactories.h"
+#include "ObjectChooser_Asset.h"
+#include "ObjectChooser_Class.h"
+#include "PersonaModule.h"
+#include "StructUtils/PropertyBag.h"
+#include "PropertyCustomizationHelpers.h"
+#include "PropertyEditorModule.h"
+#include "SAssetDropTarget.h"
+#include "SChooserColumnHandle.h"
+#include "SClassViewer.h"
+#include "ScopedTransaction.h"
+#include "SNestedChooserTree.h"
+#include "SourceCodeNavigation.h"
+#include "StructViewerModule.h"
+#include "DragAndDrop/DecoratedDragDropOp.h"
+#include "Framework/Docking/TabManager.h"
+#include "Framework/MultiBox/MultiBoxBuilder.h"
+#include "Misc/StringOutputDevice.h"
+#include "Misc/TransactionObjectEvent.h"
+#include "Modules/ModuleManager.h"
+#include "Styling/AppStyle.h"
+#include "SChooserTableRow.h"
+#include "Elements/Framework/TypedElementQueryBuilder.h"
+#include "Widgets/Docking/SDockTab.h"
+#include "Widgets/Input/SComboButton.h"
+#include "Widgets/Input/SHyperlink.h"
+#include "Widgets/Input/SNumericEntryBox.h"
+#include "Widgets/Layout/SScrollBox.h"
+#include "Widgets/Layout/SWidgetSwitcher.h"
+#include "Widgets/Views/SListView.h"
+#include "ToolMenus.h"
+#include "UnrealExporter.h"
+#include "Exporters/Exporter.h"
+#include "HAL/PlatformApplicationMisc.h"
+#include "Framework/Commands/GenericCommands.h"
+#include "SPositiveActionButton.h"
+#include "ChooserTableEditor.h"
+#include "SChooserTableWidget.h"
+
+#include UE_INLINE_GENERATED_CPP_BY_NAME(ChooserTableViewModel)
+
+#define LOCTEXT_NAMESPACE "ChooserEditor"
+
+namespace UE::ChooserEditor
+{
+
+
+FChooserTableViewModel::~FChooserTableViewModel()
+{
+	FTSTicker::RemoveTicker(TickDelegateHandle);
+	FCoreUObjectDelegates::OnObjectTransacted.RemoveAll(this);
+}
+
+void FChooserTableViewModel::RegisterMenus(TSharedPtr<FUICommandList> CommandList)
+{
+	UToolMenus* ToolMenus = UToolMenus::Get();
+	const FChooserTableEditorCommands& Commands = FChooserTableEditorCommands::Get();
+
+	// Table Context Menu
+	UToolMenu* ToolMenu;
+	if (ToolMenus->IsMenuRegistered(FChooserTableEditor::ContextMenuName))
+	{
+		ToolMenu = ToolMenus->ExtendMenu(FChooserTableEditor::ContextMenuName);
+	}
+	else
+	{
+		ToolMenu = UToolMenus::Get()->RegisterMenu(FChooserTableEditor::ContextMenuName, NAME_None, EMultiBoxType::Menu);
+	}
+
+	if (ToolMenu)
+	{
+		FToolMenuSection& Section = ToolMenu->AddSection("ChooserEditorContext", TAttribute<FText>());
+	
+		Section.AddEntry(FToolMenuEntry::InitMenuEntry(FGenericCommands::Get().Copy));
+		Section.AddEntry(FToolMenuEntry::InitMenuEntry(FGenericCommands::Get().Cut));
+		Section.AddEntry(FToolMenuEntry::InitMenuEntry(FGenericCommands::Get().Paste));
+		Section.AddEntry(FToolMenuEntry::InitMenuEntry(FGenericCommands::Get().Duplicate));
+		Section.AddEntry(FToolMenuEntry::InitMenuEntry(FGenericCommands::Get().Delete));
+		Section.AddEntry(FToolMenuEntry::InitMenuEntry(Commands.Disable));
+		Section.AddEntry(FToolMenuEntry::InitMenuEntry(Commands.MoveUp));
+		Section.AddEntry(FToolMenuEntry::InitMenuEntry(Commands.MoveDown));
+		Section.AddEntry(FToolMenuEntry::InitMenuEntry(Commands.MoveLeft));
+		Section.AddEntry(FToolMenuEntry::InitMenuEntry(Commands.MoveRight));
+		Section.AddEntry(FToolMenuEntry::InitMenuEntry(Commands.AutoPopulateSelection));
+	}
+}
+
+void FChooserTableViewModel::OnObjectsTransacted(UObject* Object, const FTransactionObjectEvent& Event)
+{
+	if (UChooserTable* ChooserTable = Cast<UChooserTable>(Object))
+	{
+		// if this is the chooser we're editing
+		if (GetChooser() == ChooserTable)
+		{
+			if (GetCurrentSelectionType() == FChooserTableViewModel::ESelectionType::Rows)
+			{
+				// refresh details if we have rows selected
+				RefreshRowSelectionDetails();
+			}
+		}
+	}
+	
+	if (UChooserRowDetails* RowDetails = Cast<UChooserRowDetails>(Object))
+	{
+		// if this is for the chooser we're editing
+		if (GetChooser() == RowDetails->Chooser)
+		{
+			if (RowDetails->Chooser->ResultsStructs.IsValidIndex(RowDetails->Row))
+			{
+				// copy all the values over
+				TValueOrError<FStructView, EPropertyBagResult> Result = RowDetails->Properties.GetValueStruct("Result", FInstancedStruct::StaticStruct());
+				if (Result.IsValid())
+				{
+					RowDetails->Chooser->ResultsStructs[RowDetails->Row] = Result.GetValue().Get<FInstancedStruct>();
+				}
+
+				int32 ColumnIndex = 0;
+				for (FInstancedStruct& ColumnData : RowDetails->Chooser->ColumnsStructs)
+				{
+					FChooserColumnBase& Column = ColumnData.GetMutable<FChooserColumnBase>();
+					Column.SetFromDetails(RowDetails->Properties, ColumnIndex, RowDetails->Row);
+					ColumnIndex++;
+				}
+			
+				TValueOrError<bool, EPropertyBagResult> DisabledResult = RowDetails->Properties.GetValueBool("Disabled");
+				if (DisabledResult.IsValid())
+				{
+					RowDetails->Chooser->DisabledRows[RowDetails->Row] = DisabledResult.GetValue();
+				}
+			}
+			else if (RowDetails->Row == ColumnWidget_SpecialIndex_Fallback)
+			{
+				TValueOrError<FStructView, EPropertyBagResult> Result = RowDetails->Properties.GetValueStruct("Result", FInstancedStruct::StaticStruct());
+				if (Result.IsValid())
+				{
+					RowDetails->Chooser->FallbackResult = Result.GetValue().Get<FInstancedStruct>();
+				}
+				
+				int32 ColumnIndex = 0;
+				for (FInstancedStruct& ColumnData : RowDetails->Chooser->ColumnsStructs)
+				{
+					FChooserColumnBase& Column = ColumnData.GetMutable<FChooserColumnBase>();
+					Column.SetFromDetails(RowDetails->Properties, ColumnIndex, RowDetails->Row);
+					ColumnIndex++;
+				}
+			}
+			QueueRefreshAll();
+		}
+	}
+}
+	
+bool FChooserTableViewModel::HandleTicker(float DeltaTime)
+{
+	if (bQueueRefreshAll)
+	{
+		RefreshAll();
+		bQueueRefreshAll = false;
+	}
+	return true;
+}
+
+FChooserTableViewModel::FChooserTableViewModel(UChooserTable* Chooser)
+{
+	check(Chooser);
+
+	TickDelegateHandle = FTSTicker::GetCoreTicker().AddTicker(FTickerDelegate::CreateRaw(this, &FChooserTableViewModel::HandleTicker), 1.f);
+	
+	RootChooser = Chooser->GetRootChooser();
+	CurrentChooser = Chooser;
+
+	FCoreUObjectDelegates::OnObjectTransacted.AddRaw(this, &FChooserTableViewModel::OnObjectsTransacted);
+}
+
+void FChooserTableViewModel::RefreshAll()
+{
+	// Cache Selection state
+	ESelectionType CachedSelectionType = CurrentSelectionType;
+	int32 SelectedColumnIndex = -1;
+	UChooserTable* SelectedChooser = nullptr;
+	TArray<int32> CachedSelectedRows;
+
+	if (CachedSelectionType == ESelectionType::Column)
+	{
+		SelectedColumnIndex = SelectedColumn->Column;
+		SelectedChooser = SelectedColumn->Chooser;
+	}
+	else if (CachedSelectionType == ESelectionType::Rows)
+	{
+		if (!SelectedRows.IsEmpty())
+		{
+			SelectedChooser = SelectedRows[0]->Chooser;
+		}
+		for(const TObjectPtr<UChooserRowDetails>& SelectedRow : SelectedRows)
+		{
+			CachedSelectedRows.Add(SelectedRow->Row);
+		}
+	}
+	
+	UpdateTableColumns();
+	UpdateTableRows();
+
+	// reapply cached selection state
+	if (CachedSelectionType == ESelectionType::Root)
+	{
+		SelectRootProperties();
+	}
+	else if (CachedSelectionType == ESelectionType::Column)
+	{
+		SelectColumn(SelectedChooser, SelectedColumnIndex);
+	}
+	else if (CachedSelectionType == ESelectionType::Rows)
+	{
+		SelectRows(CachedSelectedRows);
+	}
+	
+	// RefreshNestedObjectTree();
+}
+
+void FChooserTableViewModel::QueueRefreshAll()
+{
+	bQueueRefreshAll = true;
+}
+
+bool FChooserTableViewModel::MatchesContext(const FTransactionContext& InContext, const TArray<TPair<UObject*, FTransactionObjectEvent>>& TransactionObjectContexts) const
+{
+	TArray<UObject*> ContainedObjects;
+	GetObjectsWithOuter(GetRootChooser()->GetPackage(), ContainedObjects, true);
+	
+	for(const TPair<UObject*, FTransactionObjectEvent>&  Entry : TransactionObjectContexts)
+	{
+		if (ContainedObjects.Contains(Entry.Key))
+		{
+			return true;
+		}
+	}
+	return false;
+}
+
+void FChooserTableViewModel::PostUndo(bool bSuccess)
+{
+	SelectRootProperties();
+	RefreshAll();
+	//RefreshNestedObjectTree(); // todo: fix this
+}
+
+void FChooserTableViewModel::PostRedo(bool bSuccess)
+{
+	SelectRootProperties();
+	RefreshAll();
+	// RefreshNestedObjectTree(); todo fix this
+} 
+
+
+void FChooserTableViewModel::NotifyPreChange(FProperty* PropertyAboutToChange)
+{
+}
+
+void FChooserTableViewModel::NotifyPostChange(const FPropertyChangedEvent& PropertyChangedEvent, FProperty* PropertyThatChanged)
+{
+	// Called on details panel edits
+	
+	if (CurrentSelectionType==ESelectionType::Root)
+	{
+		// Editing the root in the details panel can change ContextData that means all wigets need to be refreshed
+		UpdateTableColumns();
+		UpdateTableRows();
+		SelectRootProperties();
+	}
+	if (CurrentSelectionType==ESelectionType::Column)
+	{
+		check(SelectedColumn);
+		int32 SelectedColumnIndex = SelectedColumn->Column;
+		UChooserTable* SelectedColumnChooser = SelectedColumn->Chooser;
+		// Editing column properties can change the column type, which requires refreshing everything
+		UpdateTableColumns();
+		UpdateTableRows();
+		SelectColumn(SelectedColumnChooser, SelectedColumnIndex);
+	}
+	// editing row data should not require any refreshing
+}
+
+void FChooserTableViewModel::RemoveDisabledData()
+{
+	UChooserTable* Chooser = GetChooser();
+	const FScopedTransaction Transaction(LOCTEXT("Remove Disabled Data", "Remove Disabled Data"));
+
+	Chooser->Modify(true);
+	Chooser->RemoveDisabledData();
+	RefreshAll();
+}
+
+int32 FChooserTableViewModel::MoveColumn(int32 SourceIndex, int32 TargetIndex)
+{
+	UChooserTable* Chooser = GetChooser();
+	
+	TargetIndex = FMath::Clamp(TargetIndex, 0, Chooser->ColumnsStructs.Num());
+
+	if (SourceIndex < 0 || SourceIndex == TargetIndex)
+	{
+		return TargetIndex;
+	}
+
+	const FScopedTransaction Transaction(LOCTEXT("Move Row", "Move Row"));
+
+	Chooser->Modify(true);
+
+	FInstancedStruct ColumnData = Chooser->ColumnsStructs[SourceIndex];
+	Chooser->ColumnsStructs.RemoveAt(SourceIndex);
+		
+	if (SourceIndex < TargetIndex)
+	{
+		TargetIndex--;
+	}
+	
+	if (TargetIndex == Chooser->ColumnsStructs.Num() && !Chooser->ColumnsStructs.IsEmpty())
+	{
+		if (Chooser->ColumnsStructs.Last().Get<FChooserColumnBase>().IsRandomizeColumn())
+		{
+			// never drop after a Randomize Column;
+			TargetIndex--;
+		}
+	}
+
+	Chooser->ColumnsStructs.Insert(ColumnData, TargetIndex);
+
+	RefreshAll();
+
+	return TargetIndex;
+}
+
+void FChooserTableViewModel::UpdateTableColumns()
+{
+	RefreshTableColumnsDelegate.ExecuteIfBound();
+}
+
+void FChooserTableViewModel::AddColumn(const UScriptStruct* ColumnType)
+{
+	FSlateApplication::Get().DismissAllMenus();
+	UChooserTable* Chooser = GetChooser();
+	const FScopedTransaction Transaction(LOCTEXT("Add Column Transaction", "Add Column"));
+	Chooser->Modify(true);
+
+	FInstancedStruct NewColumn;
+	NewColumn.InitializeAs(ColumnType);
+	NewColumn.GetMutable<FChooserColumnBase>().Initialize(Chooser);
+	const FChooserColumnBase& NewColumnRef = NewColumn.Get<FChooserColumnBase>();
+	int32 InsertIndex = 0;
+	if (NewColumnRef.IsRandomizeColumn())
+	{
+		// add randomization column at the end (do nothing if there already is one)
+		InsertIndex = Chooser->ColumnsStructs.Num();
+		if (InsertIndex == 0 || !Chooser->ColumnsStructs[InsertIndex - 1].Get<FChooserColumnBase>().IsRandomizeColumn())
+		{
+			Chooser->ColumnsStructs.Add(NewColumn);
+		}
+	}
+	else if (NewColumnRef.HasOutputs())
+	{
+		// add output columns at the end (but before any randomization column)
+		InsertIndex = Chooser->ColumnsStructs.Num();
+		if (InsertIndex > 0 && Chooser->ColumnsStructs[InsertIndex - 1].Get<FChooserColumnBase>().IsRandomizeColumn())
+		{
+			InsertIndex--;
+		}
+		Chooser->ColumnsStructs.Insert(NewColumn, InsertIndex);
+	}
+	else
+	{
+		// add other columns after the last non-output, non-randomization column
+		while(InsertIndex < Chooser->ColumnsStructs.Num())
+		{
+			const FChooserColumnBase& Column = Chooser->ColumnsStructs[InsertIndex].Get<FChooserColumnBase>();
+			if (Column.HasOutputs() || Column.IsRandomizeColumn())
+			{
+				break;
+			}
+			InsertIndex++;
+		}
+		Chooser->ColumnsStructs.Insert(NewColumn, InsertIndex);
+	}
+
+	UpdateTableColumns();
+	UpdateTableRows();
+
+	SelectColumn(Chooser, InsertIndex);
+}
+
+void FChooserTableViewModel::RefreshRowSelectionDetails()
+{
+	SelectedRows.SetNum(0, EAllowShrinking::No);
+	UChooserTable* Chooser = GetChooser();
+
+	FPropertyBagPropertyDesc ResultPropertyDesc ("Result", EPropertyBagPropertyType::Struct, FInstancedStruct::StaticStruct());
+	ResultPropertyDesc.MetaData.Add({"ExcludeBaseStruct",""});
+	ResultPropertyDesc.MetaData.Add({"BaseStruct","/Script/Chooser.ObjectChooserBase"});
+	
+	if (SelectedRowsCache.Num() < SelectedRowIndices.Num())
+	{
+		SelectedRowsCache.SetNum(SelectedRowIndices.Num());
+	}
+	
+	int32 SelectedRowsCacheIndex = 0;
+	for(int32 SelectedRowIndex : SelectedRowIndices)
+	{
+		if (Chooser->ResultsStructs.IsValidIndex(SelectedRowIndex) || SelectedRowIndex == ColumnWidget_SpecialIndex_Fallback)
+		{
+			if (SelectedRowsCache[SelectedRowsCacheIndex] == nullptr)
+			{
+				SelectedRowsCache[SelectedRowsCacheIndex] = TStrongObjectPtr(NewObject<UChooserRowDetails>(RootChooser, NAME_None, RF_Transactional | RF_Transient));
+			}
+			TObjectPtr<UChooserRowDetails> Selection = SelectedRowsCache[SelectedRowsCacheIndex].Get();
+			SelectedRowsCacheIndex ++;
+			
+			Selection->Chooser = Chooser;
+			Selection->Row = SelectedRowIndex;
+			Selection->Properties.Reset();
+		
+			if (Chooser->ResultType != EObjectChooserResultType::NoPrimaryResult)
+			{
+				FInstancedStruct& Result = SelectedRowIndex == ColumnWidget_SpecialIndex_Fallback? Chooser->FallbackResult : Chooser->ResultsStructs[SelectedRowIndex];;
+				Selection->Properties.AddProperties({ResultPropertyDesc});
+				Selection->Properties.SetValueStruct("Result", FConstStructView(FInstancedStruct::StaticStruct(), reinterpret_cast<uint8*>(&Result)));
+			}
+
+			int32 ColumnIndex = 0;
+			for (FInstancedStruct& ColumnData : Chooser->ColumnsStructs)
+			{
+				FChooserColumnBase& Column = ColumnData.GetMutable<FChooserColumnBase>();
+
+				if (Column.HasOutputs() || SelectedRowIndex != ColumnWidget_SpecialIndex_Fallback)
+				{
+					// add each column's details, but on the fallback row, only add details for Output columns
+					Column.AddToDetails(Selection->Properties, ColumnIndex, SelectedRowIndex);
+				}
+				ColumnIndex++;
+			}
+
+			if (Chooser->DisabledRows.IsValidIndex(SelectedRowIndex))
+			{
+				Selection->Properties.AddProperty("Disabled", EPropertyBagPropertyType::Bool);
+				Selection->Properties.SetValueBool("Disabled", Chooser->DisabledRows[SelectedRowIndex]);
+			}
+
+			SelectedRows.Add(Selection);
+		}
+	}
+	
+	TArray<UObject*> DetailsObjects;
+	for(auto& Item : SelectedRows)
+	{
+		DetailsObjects.Add(Item.Get());
+	}
+
+	ShowDetailsDelegate.ExecuteIfBound( DetailsObjects );
+}
+
+void FChooserTableViewModel::UpdateTableRows()
+{
+	RefreshTableRowsDelegate.ExecuteIfBound();
+}
+
+void FChooserTableViewModel::DeleteColumn(int32 Index)
+{
+	const FScopedTransaction Transaction(LOCTEXT("Delete Column Transaction", "Delete Column"));
+	UChooserTable* Chooser = GetChooser();
+
+	if (Index < Chooser->ColumnsStructs.Num())
+	{
+		Chooser->Modify(true);
+		Chooser->ColumnsStructs.RemoveAt(Index);
+		UpdateTableColumns();
+	}
+}
+	
+void FChooserTableViewModel::DeleteRows(TConstArrayView<int32> RowsToDelete)
+{
+	const FScopedTransaction Transaction(LOCTEXT("Delete Row Transaction", "Delete Row"));
+	DeleteRowsInternal(RowsToDelete, 0);
+}
+
+int32 FChooserTableViewModel::DeleteRowsInternal(TConstArrayView<int32> InRowsToDelete, int32 RowIndexToRemember)
+{
+	if (UChooserTable* Chooser = GetChooser())
+	{
+		Chooser->Modify();
+		
+		// sort indices in reverse
+		TArray<int32> RowsToDelete;
+		RowsToDelete = InRowsToDelete;
+		RowsToDelete.Sort([](int32 A, int32 B){ return A>B; });
+		for(int32 RowIndex : RowsToDelete)
+		{
+			if (RowIndex == ColumnWidget_SpecialIndex_Fallback)
+			{
+				Chooser->FallbackResult.Reset();
+			}
+			else
+			{
+				if (RowIndex <= RowIndexToRemember)
+				{
+					RowIndexToRemember--;
+				}
+				Chooser->ResultsStructs.RemoveAt(RowIndex);
+				Chooser->DisabledRows.RemoveAt(RowIndex);
+			}
+		}
+
+		for(FInstancedStruct& ColumnData : Chooser->ColumnsStructs)
+		{
+			FChooserColumnBase& Column = ColumnData.GetMutable<FChooserColumnBase>();
+			Column.DeleteRows(RowsToDelete);
+		}
+		UpdateTableRows();
+	}
+
+	return RowIndexToRemember;
+}
+	
+void FChooserTableViewModel::MoveRows(TConstArrayView<int32> RowIndices, int32 TargetIndex)
+{
+	const FScopedTransaction Transaction(LOCTEXT("Move Row(s)", "Move Row(s)"));
+	UChooserTable* RowCopy = CopyRowsInternal(RowIndices);
+	TargetIndex = DeleteRowsInternal(RowIndices, TargetIndex);
+	PasteInternal(RowCopy, TargetIndex);
+}
+
+void FChooserTableViewModel::AutoPopulateColumn(int32 ColumnIndex)
+{
+	if (UChooserTable* Chooser = GetChooser())
+	{
+		if (Chooser->ColumnsStructs.IsValidIndex(ColumnIndex))
+		{
+			const FScopedTransaction Transaction(LOCTEXT("Autopopulate Column", "Autopopulate Column"));
+			Chooser->Modify();
+			AutoPopulateColumn(Chooser->ColumnsStructs[ColumnIndex].GetMutable<FChooserColumnBase>());
+		}
+	}
+}
+	
+void FChooserTableViewModel::AutoPopulateColumn(FChooserColumnBase& Column)
+{
+	if (UChooserTable* Chooser = GetChooser())
+	{
+		const int32 RowCount = Chooser->ResultsStructs.Num();
+		if (Column.AutoPopulates())
+		{
+			for (int32 i = 0; i < RowCount; ++i)
+			{
+				UObject* ReferencedObject = Chooser->ResultsStructs[i].IsValid() ? Chooser->ResultsStructs[i].Get<FObjectChooserBase>().GetReferencedObject() : nullptr;
+				Column.AutoPopulate(i, ReferencedObject);
+			}
+
+			UObject* ReferencedObject = Chooser->FallbackResult.IsValid() ? Chooser->FallbackResult.Get<FObjectChooserBase>().GetReferencedObject() : nullptr;
+			Column.AutoPopulate(ColumnWidget_SpecialIndex_Fallback, ReferencedObject);
+		}
+	}
+}
+
+void FChooserTableViewModel::AutoPopulateRows(TConstArrayView<int32> RowIndices)
+{
+	if (UChooserTable* Chooser = GetChooser())
+	{
+		const FScopedTransaction Transaction(LOCTEXT("Auto Populate Selection", "Auto Populate Selection"));
+		Chooser->Modify();
+		
+		for(int32 RowIndex : RowIndices)
+		{
+			AutoPopulateRow(RowIndex);
+		}
+	}
+}
+
+void FChooserTableViewModel::AutoPopulateRow(int32 Index)
+{
+	if (UChooserTable* Chooser = GetChooser())
+	{
+		if (Chooser->ResultsStructs.IsValidIndex(Index) && Chooser->ResultsStructs[Index].IsValid())
+		{
+			if (UObject* ReferencedObject = Chooser->ResultsStructs[Index].Get<FObjectChooserBase>().GetReferencedObject()) 
+			{
+				for (FInstancedStruct& ColumnData : Chooser->ColumnsStructs)
+				{
+					FChooserColumnBase& Column = ColumnData.GetMutable<FChooserColumnBase>();
+					Column.AutoPopulate(Index, ReferencedObject);
+				}
+			}
+		}
+		else if (Index == ColumnWidget_SpecialIndex_Fallback && Chooser->FallbackResult.IsValid())
+		{
+			if (UObject* ReferencedObject = Chooser->FallbackResult.Get<FObjectChooserBase>().GetReferencedObject())
+			{
+				for (FInstancedStruct& ColumnData : Chooser->ColumnsStructs)
+				{
+					FChooserColumnBase& Column = ColumnData.GetMutable<FChooserColumnBase>();
+					Column.AutoPopulate(Index, ReferencedObject);
+				}
+			}
+		}
+	}
+}
+	
+bool FChooserTableViewModel::CanAutoPopulateColumn(int32 ColumnIndex) const
+{
+	if (const UChooserTable* Chooser = GetChooser())
+	{
+		if (CurrentSelectionType == ESelectionType::Column && SelectedColumn)
+		{
+			// when a column is selected, return true if that column supports auto populate
+			if (Chooser->ColumnsStructs.IsValidIndex(ColumnIndex))
+			{
+				return Chooser->ColumnsStructs[ColumnIndex].Get<FChooserColumnBase>().AutoPopulates();
+			}	
+		}
+	}
+	return false;
+}
+	
+bool FChooserTableViewModel::CanAutoPopulateRows() const
+{
+	if (const UChooserTable* Chooser = GetChooser())
+	{
+		// when rows are selected, return true if any column supports auto populate
+		for (const FInstancedStruct& ColumnData : Chooser->ColumnsStructs)
+		{
+			const FChooserColumnBase& Column = ColumnData.Get<FChooserColumnBase>();
+			if (Column.AutoPopulates())
+			{
+				return true;
+			}
+		}
+	}
+	
+	return false;
+}
+	
+void FChooserTableViewModel::AutoPopulateAll()
+{
+	if (UChooserTable* Chooser = GetChooser())
+	{
+		const FScopedTransaction Transaction(LOCTEXT("Auto Populate Chooser", "Auto Populate All"));
+		Chooser->Modify();
+		for (FInstancedStruct& ColumnData : Chooser->ColumnsStructs)
+		{
+			FChooserColumnBase& Column = ColumnData.GetMutable<FChooserColumnBase>();
+			AutoPopulateColumn(Column);
+		}
+	}
+}
+
+bool FChooserTableViewModel::HasSelection() const
+{
+	return HasRowsSelected() || HasColumnSelected();
+}
+
+bool FChooserTableViewModel::HasRowsSelected() const
+{
+	return CurrentSelectionType == ESelectionType::Rows && !SelectedRows.IsEmpty();
+}
+	
+bool FChooserTableViewModel::HasColumnSelected() const
+{
+	const UChooserTable* Chooser = GetChooser();
+	return CurrentSelectionType == ESelectionType::Column && SelectedColumn && Chooser->ColumnsStructs.IsValidIndex(SelectedColumn->Column);
+}
+	
+bool FChooserTableViewModel::IsColumnDisabled(int32 ColumnIndex) const
+{
+	if (const UChooserTable* Chooser = GetChooser())
+	{
+		if ( Chooser->ColumnsStructs.IsValidIndex(ColumnIndex))
+		{
+			const FChooserColumnBase& Column = Chooser->ColumnsStructs[ColumnIndex].Get<FChooserColumnBase>();
+			return Column.bDisabled;
+		}
+	}
+	return false;
+}
+
+bool FChooserTableViewModel::AreRowsDisabled(TConstArrayView<int32> RowIndices) const
+{
+	if (const UChooserTable* Chooser = GetChooser())
+	{
+		bool bSomethingEnabled = false;
+		for(int32 Row : RowIndices)
+		{
+			if (!Chooser->IsRowDisabled(Row))
+			{
+				bSomethingEnabled = true;
+				break;
+			}
+		}
+		return !bSomethingEnabled;
+	}
+	return false;
+}
+	
+void FChooserTableViewModel::ToggleDisableColumn(int32 ColumnIndex)
+{
+	if (UChooserTable* Chooser = GetChooser())
+	{
+		const FScopedTransaction Transaction(LOCTEXT("Toggle Disable Column", "Toggle Disable Column"));
+		Chooser->Modify();
+		
+		if (Chooser->ColumnsStructs.IsValidIndex(ColumnIndex))
+		{
+			FChooserColumnBase& Column = Chooser->ColumnsStructs[ColumnIndex].GetMutable<FChooserColumnBase>();
+			Column.bDisabled = !Column.bDisabled;
+		}
+	}
+}
+
+void FChooserTableViewModel::ToggleDisableRows(TConstArrayView<int32> RowIndices)
+{
+	if (UChooserTable* Chooser = GetChooser())
+	{
+		const FScopedTransaction Transaction(LOCTEXT("Toggle Disable Rows", "Toggle Disable Row(s)"));
+		Chooser->Modify();
+		
+		bool bDisabled = AreRowsDisabled(RowIndices);
+		for (int32 Row : RowIndices)
+		{
+			if (Chooser->DisabledRows.IsValidIndex(Row))
+			{
+				Chooser->DisabledRows[Row] = !bDisabled;
+			}
+		}
+		RefreshRowSelectionDetails();
+	}
+}
+	
+void FChooserTableViewModel::DuplicateColumn(int32 ColumnIndex)
+{
+	const FScopedTransaction Transaction(LOCTEXT("Duplicate Column", "Duplicate Column"));
+	UChooserTable* Chooser = GetChooser();
+	if (Chooser->ColumnsStructs.IsValidIndex(ColumnIndex))
+	{
+		Chooser->Modify();
+		FInstancedStruct Column = Chooser->ColumnsStructs[ColumnIndex];
+		Chooser->ColumnsStructs.Insert(Column,ColumnIndex);
+	}
+
+	UpdateTableColumns();
+}
+	
+void FChooserTableViewModel::DuplicateRows(TConstArrayView<int32> RowIndices)
+{
+	const FScopedTransaction Transaction(LOCTEXT("Duplicate Row(s)", "Duplicate Row(s)"));
+	UChooserTable* Chooser = GetChooser();
+	Chooser->Modify();
+	
+	int32 MaxSelectedRow = -1;
+	for(int32 RowIndex : RowIndices)
+	{
+		MaxSelectedRow = FMath::Max(RowIndex, MaxSelectedRow);
+	}
+	
+	if (MaxSelectedRow>=0)
+	{
+		UChooserTable* RowCopy = CopyRowsInternal(RowIndices);
+		PasteInternal(RowCopy, MaxSelectedRow + 1);
+	}
+}
+
+bool FChooserTableViewModel::HasFallbackSelected()
+{
+	for(UChooserRowDetails* SelectedRow : SelectedRows)
+	{
+		if (SelectedRow->Row == ColumnWidget_SpecialIndex_Fallback)
+		{
+			return true;
+		}
+	}
+
+	return false;
+}
+
+bool FChooserTableViewModel::CanMoveRowsUp(TConstArrayView<int32> RowIndices)
+{
+	UChooserTable* Chooser = GetChooser();
+
+	int32 MinSelectedRow = Chooser->ResultsStructs.Num();
+	for(int32 RowIndex : RowIndices)
+	{
+		if (RowIndex != ColumnWidget_SpecialIndex_Fallback)
+		{
+			MinSelectedRow = FMath::Min(RowIndex, MinSelectedRow);
+		}
+	}
+
+	return MinSelectedRow > 0;
+}
+
+int32 FChooserTableViewModel::GetNumRows() const
+{
+	if (const UChooserTable* Chooser = GetChooser())
+	{
+		return Chooser->ResultsStructs.Num();
+	}
+	return 0;
+}
+
+bool FChooserTableViewModel::CanMoveRowsDown(TConstArrayView<int32> RowIndices)
+{
+	UChooserTable* Chooser = GetChooser();
+	
+	int32 MaxSelectedRow = -1;
+	for(int32 RowIndex : RowIndices)
+	{
+		MaxSelectedRow = FMath::Max(RowIndex, MaxSelectedRow);
+	}
+
+	return MaxSelectedRow < Chooser->ResultsStructs.Num() - 1;
+}
+
+bool FChooserTableViewModel::IsRandomizeColumn(int32 ColumnIndex)
+{
+	UChooserTable* Chooser = GetChooser();
+    
+	if (Chooser->ColumnsStructs.IsValidIndex(ColumnIndex))
+	{
+		if (const FChooserColumnBase* Column = Chooser->ColumnsStructs[ColumnIndex].GetPtr<FChooserColumnBase>())
+		{
+			return (Column->IsRandomizeColumn());
+		}
+	}
+	return false;
+}
+
+bool FChooserTableViewModel::CanMoveColumnLeft(int32 ColumnIndex)
+{
+	UChooserTable* Chooser = GetChooser();
+
+	return !IsRandomizeColumn(ColumnIndex) && ColumnIndex > 0;
+}
+
+void FChooserTableViewModel::MoveColumnLeft(int32 ColumnIndex)
+{
+	if (CanMoveColumnLeft(ColumnIndex))
+	{
+		SelectColumn(GetChooser(), MoveColumn(ColumnIndex, ColumnIndex - 1));
+	}
+}
+
+bool FChooserTableViewModel::CanMoveColumnRight(int32 ColumnIndex)
+{
+	UChooserTable* Chooser = GetChooser();
+	
+	if (Chooser->ColumnsStructs.IsValidIndex(ColumnIndex))
+	{
+		if (IsRandomizeColumn(ColumnIndex))
+		{
+			return false;
+		}
+		int32 NumColumns = Chooser->ColumnsStructs.Num();
+		if (IsRandomizeColumn(NumColumns -1))
+		{
+			NumColumns--;
+		}
+		return (ColumnIndex < NumColumns-1);
+	}
+   	return false;
+}
+
+void FChooserTableViewModel::MoveColumnRight(int32 ColumnIndex)
+{
+	if (CanMoveColumnRight(ColumnIndex))
+	{
+		SelectColumn(GetChooser(), MoveColumn(ColumnIndex, ColumnIndex + 2));
+	}
+}
+
+UChooserTable* DuplicateNestedChooser(UChooserTable* Chooser, UChooserTable* NewOuter)
+{
+	UChooserTable* RootTable = NewOuter->GetRootChooser();
+	if (TObjectPtr<UObject>* FoundTable = RootTable->NestedObjects.FindByPredicate([Chooser](UObject* Object)
+		{
+			if (Cast<UChooserTable>(Object))
+			{
+				return Object->GetName() == Chooser->GetName();
+			}
+			return false;
+		}))
+	{
+		// we already duplicated this table
+		return Cast<UChooserTable>(*FoundTable);
+	}
+	
+	UChooserTable* NewTable = NewObject<UChooserTable>(NewOuter, Chooser->GetFName(), RF_Transactional);
+	NewTable->ColumnsStructs = Chooser->ColumnsStructs;
+	NewTable->ResultsStructs = Chooser->ResultsStructs;
+	NewTable->RootChooser = RootTable;
+	RootTable->AddNestedObject(NewTable);
+
+	for (FInstancedStruct& ResultData : NewTable->ResultsStructs)
+	{
+		if (FNestedChooser* NestedChooser = ResultData.GetMutablePtr<FNestedChooser>())
+		{
+			NestedChooser->Chooser = DuplicateNestedChooser(NestedChooser->Chooser, NewOuter);
+		}
+	}
+
+	return NewTable;
+}
+
+UChooserTable* FChooserTableViewModel::CopyRowsInternal(TConstArrayView<int32> RowIndices)
+{
+	UChooserTable* CopyData = NewObject<UChooserTable>(GetTransientPackage());
+
+	UChooserTable* Chooser = GetChooser();
+
+	// copy context data from root table
+	CopyData->OutputObjectType = RootChooser->OutputObjectType;
+	CopyData->ResultType = RootChooser->ResultType;
+	CopyData->ContextData = RootChooser->ContextData;
+
+	TArray<int32> RowIndicesCopy;
+	RowIndicesCopy = RowIndices;
+	RowIndicesCopy.Sort();
+	
+	
+	if (RowIndicesCopy.Num() > 0 && RowIndicesCopy[0] == ColumnWidget_SpecialIndex_Fallback)
+	{
+		RowIndicesCopy.RemoveAt(0);
+		CopyData->FallbackResult = Chooser->FallbackResult;
+		if (FNestedChooser* CopiedNestedChooser = CopyData->FallbackResult.GetMutablePtr<FNestedChooser>())
+		{
+			// if the fallback result was a nested chooser, duplicate it
+			CopiedNestedChooser->Chooser = DuplicateNestedChooser(CopiedNestedChooser->Chooser, CopyData);
+		}
+	}
+	
+	CopyData->ResultsStructs.SetNum(RowIndicesCopy.Num());
+	CopyData->DisabledRows.SetNum(RowIndicesCopy.Num());
+
+	// add the selected results and column data
+	
+	for (int32 RowIndex = 0; RowIndex < RowIndicesCopy.Num(); RowIndex++)
+	{
+		CopyData->ResultsStructs[RowIndex] = Chooser->ResultsStructs[RowIndicesCopy[RowIndex]];
+		if (FNestedChooser* CopiedNestedChooser = CopyData->ResultsStructs[RowIndex].GetMutablePtr<FNestedChooser>())
+		{
+			if (CopiedNestedChooser->Chooser)
+			{
+				// if the result for this row was a nested chooser (with a valid chooser assigned), duplicate it
+				CopiedNestedChooser->Chooser = DuplicateNestedChooser(CopiedNestedChooser->Chooser, CopyData);
+			}
+		}
+
+		CopyData->DisabledRows[RowIndex] = Chooser->DisabledRows[RowIndicesCopy[RowIndex]];
+	}
+
+	// copy all columns
+	CopyData->ColumnsStructs = Chooser->ColumnsStructs;
+	for (FInstancedStruct& ColumnData : CopyData->ColumnsStructs)
+	{
+		FChooserColumnBase& Column = ColumnData.GetMutable<FChooserColumnBase>();
+		Column.Initialize(CopyData);
+	}
+
+	// clear all column's cell data
+	for (FInstancedStruct& ColumnData : CopyData->ColumnsStructs)
+	{
+		FChooserColumnBase& Column = ColumnData.GetMutable<FChooserColumnBase>();
+		Column.SetNumRows(0);
+		Column.SetNumRows(RowIndicesCopy.Num());
+	}
+
+	for (int32 RowIndex = 0; RowIndex < RowIndicesCopy.Num(); RowIndex++)
+	{
+		for (int32 ColumnIndex = 0; ColumnIndex < CopyData->ColumnsStructs.Num(); ColumnIndex++)
+		{
+			FChooserColumnBase& SourceColumn = Chooser->ColumnsStructs[ColumnIndex].GetMutable<FChooserColumnBase>();
+			FChooserColumnBase& TargetColumn = CopyData->ColumnsStructs[ColumnIndex].GetMutable<FChooserColumnBase>();
+			TargetColumn.CopyRow(SourceColumn, RowIndicesCopy[RowIndex], RowIndex);
+		}
+	}
+	
+	return CopyData;
+}
+
+void FChooserTableViewModel::SetChooser(UChooserTable* Chooser)
+{
+	RootChooser = Chooser->GetRootChooser();
+	CurrentChooser = Chooser;
+	RefreshAll();
+}
+
+UChooserTable* FChooserTableViewModel::CopyColumnInternal(int32 ColumnIndex)
+{
+	UChooserTable* CopyData = NewObject<UChooserTable>(GetTransientPackage());
+
+	UChooserTable* Chooser = GetChooser();
+
+	// copy context data from root table
+	CopyData->OutputObjectType = RootChooser->OutputObjectType;
+	CopyData->ResultType = RootChooser->ResultType;
+	CopyData->ContextData = RootChooser->ContextData;
+
+	// add selected column including all the cell data
+	CopyData->ColumnsStructs.Add(Chooser->ColumnsStructs[ColumnIndex]);
+	
+	return CopyData;
+}
+	
+void FChooserTableViewModel::CopyColumn(int32 ColumnIndex)
+{
+	UChooserTable* CopyData = CopyColumnInternal(ColumnIndex);
+
+	// Clear the mark state for saving.
+	UnMarkAllObjects(EObjectMark(OBJECTMARK_TagExp | OBJECTMARK_TagImp));
+	
+	// Export the clipboard to text.
+	FStringOutputDevice Archive;
+	const FExportObjectInnerContext Context;
+	UExporter::ExportToOutputDevice(&Context,CopyData, nullptr, Archive, TEXT("copy"), 0, PPF_ExportsNotFullyQualified | PPF_Copy | PPF_Delimited, false, CopyData->GetOuter());
+	FPlatformApplicationMisc::ClipboardCopy(*Archive);
+}
+
+void FChooserTableViewModel::CopyRows(TConstArrayView<int32> RowsToCopy)
+{
+	UChooserTable* CopyData = CopyRowsInternal(RowsToCopy);
+
+	// Clear the mark state for saving.
+	UnMarkAllObjects(EObjectMark(OBJECTMARK_TagExp | OBJECTMARK_TagImp));
+	
+	// Export the clipboard to text.
+	FStringOutputDevice Archive;
+	const FExportObjectInnerContext Context;
+	UExporter::ExportToOutputDevice(&Context,CopyData, nullptr, Archive, TEXT("copy"), 0, PPF_ExportsNotFullyQualified | PPF_Copy | PPF_Delimited, false, CopyData->GetOuter());
+	FPlatformApplicationMisc::ClipboardCopy(*Archive);
+}
+	
+class FChooserClipboardFactory : public FCustomizableTextObjectFactory
+{
+public:
+	
+	FChooserClipboardFactory()
+		: FCustomizableTextObjectFactory(GWarn)
+		, ClipboardContent(nullptr) 
+	{
+	}
+
+	UChooserTable* ClipboardContent;
+	
+protected:
+
+	virtual bool CanCreateClass(UClass* InObjectClass, bool& bOmitSubObjs) const override
+	{
+		if (InObjectClass->IsChildOf(UChooserTable::StaticClass()))
+		{
+			return true;
+		}
+		return false;
+	}
+	
+	virtual void ProcessConstructedObject(UObject* CreatedObject) override
+	{
+		if (CreatedObject->IsA<UChooserTable>())
+		{
+			ClipboardContent = CastChecked<UChooserTable>(CreatedObject);
+		}
+	}
+};
+
+static FString GetColumnName(FChooserColumnBase& Column)
+{
+	if (FChooserParameterBase* InputValue = Column.GetInputValue())
+	{
+		return InputValue->GetDebugName();
+	}
+	return FString();
+}
+	
+void FChooserTableViewModel::PasteInternal(UChooserTable* PastedContent, int32 PasteRowIndex)
+{
+	UChooserTable* Chooser = GetChooser();
+	Chooser->Modify();
+	
+	if (PastedContent->ResultsStructs.IsEmpty() && !PastedContent->FallbackResult.IsValid())
+	{
+		// pasting a column
+		int32 InsertColumnIndex = Chooser->ColumnsStructs.Num();
+		if (CurrentSelectionType == ESelectionType::Column && SelectedColumn)
+		{
+			InsertColumnIndex = FMath::Min(InsertColumnIndex, SelectedColumn->Column + 1);
+		}
+		
+		if (Chooser->ColumnsStructs.Num() > 0 && Chooser->ColumnsStructs.Num() == InsertColumnIndex)
+		{
+			// if were inserting at the end, there is a randomize column, insert new columns before it
+			if (Chooser->ColumnsStructs.Last().Get<FChooserColumnBase>().IsRandomizeColumn())
+			{
+				InsertColumnIndex--;
+			}
+		}
+		Chooser->ColumnsStructs.Insert(PastedContent->ColumnsStructs, InsertColumnIndex);
+		SelectColumn(Chooser, InsertColumnIndex);
+	}
+	else
+	{
+		// pasting rows
+		int32 RowsToPaste = PastedContent->ResultsStructs.Num();
+
+		// figure out where to start inserting
+		int32 InsertIndex = Chooser->ResultsStructs.Num();
+
+		if (PasteRowIndex >=0)
+		{
+			InsertIndex = PasteRowIndex;
+		}
+		else
+		{
+			if (SelectedRows.Num() > 0)
+			{
+				InsertIndex = SelectedRows[0]->Row;
+				for(int32 SelectedRowIndex = 1; SelectedRowIndex < SelectedRows.Num(); SelectedRowIndex ++)
+				{
+					InsertIndex = FMath::Max(InsertIndex, SelectedRows[SelectedRowIndex]->Row);
+				}
+				if (InsertIndex == ColumnWidget_SpecialIndex_Fallback)
+				{
+					// if the only row selected was the fallback, reset insert index to the last row
+					InsertIndex = Chooser->ResultsStructs.Num();
+				}
+				else
+				{
+					InsertIndex++;
+				}
+			}
+		}
+
+		if (PastedContent->ResultsStructs.Num() > 0)
+		{
+			Chooser->ResultsStructs.Insert(PastedContent->ResultsStructs, InsertIndex);
+			Chooser->DisabledRows.Insert(PastedContent->DisabledRows, InsertIndex);
+
+			// Make sure each column has the same number of row datas as there are results
+			for(FInstancedStruct& ColumnData : Chooser->ColumnsStructs)
+			{
+				FChooserColumnBase& Column = ColumnData.GetMutable<FChooserColumnBase>();
+				Column.InsertRows(InsertIndex, RowsToPaste);
+			}
+		}
+		if (PastedContent->FallbackResult.IsValid())
+		{
+			// paste fallback result if copy data has one
+			Chooser->FallbackResult = PastedContent->FallbackResult;
+			if (FNestedChooser* NestedChooser = Chooser->FallbackResult.GetMutablePtr<FNestedChooser>())
+			{
+				// duplicate the nested chooser if the fallback result refers to a nested chooser
+				NestedChooser->Chooser = DuplicateNestedChooser(NestedChooser->Chooser, Chooser);
+			}
+		}
+		
+		if (!PastedContent->NestedObjects.IsEmpty())
+		{
+			// if there were nested choosers in the copy buffer we have to remap or paste them here
+			for (int32 ResultIndex = InsertIndex; ResultIndex < PastedContent->ResultsStructs.Num() + InsertIndex; ResultIndex++)
+			{
+				if (FNestedChooser* NestedChooser = Chooser->ResultsStructs[ResultIndex].GetMutablePtr<FNestedChooser>())
+				{
+					NestedChooser->Chooser = DuplicateNestedChooser(NestedChooser->Chooser, Chooser);
+				}
+			}
+		}
+
+		// try to also paste column data from columns in the paste buffer which match the columns in the current chooser
+		// -- matching by column type and input value name
+
+		// keep track of target columns that have already been matched, to avoid matching multiple source columns with the same target column
+		TArray<bool> MatchedTargetColumns;
+		MatchedTargetColumns.SetNum(Chooser->ColumnsStructs.Num());
+
+		// keep track of which source columns were matched, so we can add new columns for the unmatched ones after
+		TArray<bool> MatchedSourceColumns;
+		MatchedSourceColumns.SetNum(PastedContent->ColumnsStructs.Num());
+		
+		for(int32 SourceColumnIndex = 0; SourceColumnIndex < PastedContent->ColumnsStructs.Num(); SourceColumnIndex++)
+		{
+			FInstancedStruct& PastedColumnData = PastedContent->ColumnsStructs[SourceColumnIndex];
+			FChooserColumnBase& PastedColumn = PastedColumnData.GetMutable<FChooserColumnBase>();
+			FString PastedColumnName = GetColumnName(PastedColumn);
+			for(int32 TargetColumnIndex = 0; TargetColumnIndex < Chooser->ColumnsStructs.Num(); TargetColumnIndex++)
+			{
+				if (!MatchedTargetColumns[TargetColumnIndex])
+				{
+					FInstancedStruct& ColumnData = Chooser->ColumnsStructs[TargetColumnIndex];
+					if (ColumnData.GetScriptStruct() == PastedColumnData.GetScriptStruct())
+					{
+						FChooserColumnBase& Column = ColumnData.GetMutable<FChooserColumnBase>();
+						FString ColumnName = GetColumnName(Column);
+
+						if (ColumnName == PastedColumnName)
+						{
+							MatchedTargetColumns[TargetColumnIndex] = true;
+							MatchedSourceColumns[SourceColumnIndex] = true;
+
+							// found a match, copy the data over
+							for (int32 i = 0; i < RowsToPaste; i++)
+							{
+								Column.CopyRow(PastedColumn, i, InsertIndex + i);
+							}
+
+							if (PastedContent->FallbackResult.IsValid())
+							{
+								// if the fallback row was copied, paste the fallback data for columns
+								Column.CopyFallback(PastedColumn);
+							}
+							break;
+						}
+					}
+				}
+			}
+		}
+		
+		// add new columns for any source columns that were unmatched
+		
+		int32 InsertColumnIndex = Chooser->ColumnsStructs.Num();
+		if (Chooser->ColumnsStructs.Num() > 0)
+		{
+			// if there is a randomize column, insert new columns before it
+			if (Chooser->ColumnsStructs.Last().Get<FChooserColumnBase>().IsRandomizeColumn())
+			{
+				InsertColumnIndex--;
+			}
+		}
+
+		for(int32 SourceColumnIndex = 0; SourceColumnIndex < PastedContent->ColumnsStructs.Num(); SourceColumnIndex++)
+		{
+			if (!MatchedSourceColumns[SourceColumnIndex])
+			{
+				FInstancedStruct& PastedColumnData = PastedContent->ColumnsStructs[SourceColumnIndex];
+				FChooserColumnBase& PastedColumn = PastedColumnData.GetMutable<FChooserColumnBase>();
+				// if we couldn't find a match, paste a new column
+				Chooser->ColumnsStructs.Insert(PastedColumnData, InsertColumnIndex);
+				FChooserColumnBase& Column = Chooser->ColumnsStructs[InsertColumnIndex].GetMutable<FChooserColumnBase>();
+				InsertColumnIndex++;
+				Column.SetNumRows(0);
+				Column.SetNumRows(Chooser->ResultsStructs.Num());
+				for (int32 i = 0; i < RowsToPaste; i++)
+				{
+					Column.CopyRow(PastedColumn, i, InsertIndex + i);
+				}
+			}
+		}
+
+		ClearSelectedRows();
+		RefreshAll();
+
+		TArray<int32, TInlineAllocator<256>> SelectedIndices;
+		SelectedIndices.Reserve(RowsToPaste+1);
+		
+		// select the inserted rows
+		for (int32 i = 0; i < RowsToPaste; i++)
+		{
+			SelectedIndices.Add(InsertIndex + i);
+		}
+		if(PastedContent->FallbackResult.IsValid())
+		{
+			SelectedIndices.Add(ColumnWidget_SpecialIndex_Fallback);
+		}
+
+		SelectRows(SelectedIndices);
+	}
+}
+
+void FChooserTableViewModel::Paste()
+{
+	// Get the text from the clipboard.
+	FString ClipboardText;
+	FPlatformApplicationMisc::ClipboardPaste(ClipboardText);
+	FChooserClipboardFactory Factory;
+
+	UChooserTable* Chooser = GetChooser();
+
+	if (Factory.CanCreateObjectsFromText(ClipboardText))
+	{
+		Factory.ProcessBuffer((UObject*)GetTransientPackage(), RF_Transactional, ClipboardText);
+		if (UChooserTable* PastedContent = Factory.ClipboardContent)
+		{
+			FScopedTransaction Transaction(LOCTEXT("Paste Chooser Data", "Paste Chooser Data"));
+			PasteInternal(PastedContent);
+			RefreshAll();
+		}
+	}
+}
+
+	
+bool FChooserTableViewModel::CanPaste() const
+{
+	// Get the text from the clipboard.
+	FString ClipboardText;
+	FPlatformApplicationMisc::ClipboardPaste(ClipboardText);
+	FChooserClipboardFactory Factory;
+	return Factory.CanCreateObjectsFromText(ClipboardText); 
+}
+
+// Selection stuff to move to widget
+	
+void FChooserTableViewModel::SelectColumn(UChooserTable* ChooserEditor, int32 Index)
+{
+	ClearSelectedRows();
+	
+	UChooserTable* Chooser = GetChooser();
+   	if (Chooser->ColumnsStructs.IsValidIndex(Index))
+   	{
+   		if (SelectedColumn == nullptr)
+   		{
+   			SelectedColumn = TStrongObjectPtr(NewObject<UChooserColumnDetails>(RootChooser, NAME_None, RF_Transactional | RF_Transient));
+   		}
+   
+   		SelectedColumn->Chooser = Chooser;
+   		SelectedColumn->Column = Index;
+
+   		ShowDetailsDelegate.ExecuteIfBound({SelectedColumn.Get()});
+   		
+   		CurrentSelectionType = ESelectionType::Column;
+   	}
+}
+	
+void FChooserTableViewModel::ClearSelectedColumn()
+{
+	UChooserTable* Chooser = GetChooser();
+	if (CurrentSelectionType == ESelectionType::Column && SelectedColumn)
+	{
+		SelectedColumn->Column = INDEX_NONE;
+	}
+}
+	
+void FChooserTableViewModel::SelectRootProperties()
+{
+	ClearSelectedColumn();
+	ClearSelectedRows();
+	
+	ShowDetailsDelegate.ExecuteIfBound( { GetRootChooser() } );
+	
+	CurrentSelectionType = ESelectionType::Root;
+}
+
+void FChooserTableViewModel::SetSelectedRows(const TConstArrayView<int32>& Rows)
+{
+	ClearSelectedColumn();
+	SetCurrentSelectionType(FChooserTableViewModel::ESelectionType::Rows);
+	SelectedRowIndices = Rows;
+	RefreshRowSelectionDetails();
+}
+	
+void FChooserTableViewModel::SelectRows(const TConstArrayView<int32>& Rows)
+{
+	SelectRowsDelegate.ExecuteIfBound(Rows);
+}
+	
+void FChooserTableViewModel::ClearSelectedRows()
+{
+	SelectRows({});
+}
+
+bool FChooserTableViewModel::IsRowSelected(int32 RowIndex)
+{
+	for(auto& SelectedRow:SelectedRows)
+ 	{
+ 		if (SelectedRow->Row == RowIndex)
+ 		{
+ 			return true;
+ 		}
+ 	}
+	return false;
+}
+	
+bool FChooserTableViewModel::IsColumnSelected(int32 ColumnIndex)
+{
+	return (CurrentSelectionType == ESelectionType::Column && SelectedColumn && SelectedColumn->Column == ColumnIndex);
+}
+
+
+	
+}
+
+#undef LOCTEXT_NAMESPACE
+
diff --git a/Engine/Plugins/Chooser/Source/ChooserEditor/Private/ChooserTableViewModel.h b/Engine/Plugins/Chooser/Source/ChooserEditor/Private/ChooserTableViewModel.h
new file mode 100644
--- /dev/null
+++ b/Engine/Plugins/Chooser/Source/ChooserEditor/Private/ChooserTableViewModel.h
@@ -0,0 +1,201 @@
+// Copyright Epic Games, Inc. All Rights Reserved.
+
+#pragma once
+
+#include "CoreMinimal.h"
+#include "Chooser.h"
+#include "EditorUndoClient.h"
+#include "PropertyEditorDelegates.h"
+#include "Containers/RingBuffer.h"
+#include "Toolkits/AssetEditorToolkit.h"
+#include "Toolkits/IToolkitHost.h"
+#include "Widgets/Navigation/SBreadcrumbTrail.h"
+#include "Widgets/Views/SListView.h"
+#include "ChooserDetails.h"
+#include "IChooserTableViewModel.h"
+#include "ChooserTableViewModel.generated.h"
+
+class SComboButton;
+class SPositiveActionButton;
+class SEditableText;
+
+// Class used for chooser editor details customization
+UCLASS()
+class UChooserColumnDetails : public UObject
+{
+	GENERATED_BODY()
+public:
+	UPROPERTY(EditAnywhere, Instanced, Category="Hidden")
+	TObjectPtr<UChooserTable> Chooser;
+	int32 Column = -1;
+};
+
+
+namespace UE::ChooserEditor
+{
+	class SChooserTableWidget;
+	class SNestedChooserTree;
+	struct FChooserTableRow;
+
+	DECLARE_DELEGATE(FRefreshTable);
+	DECLARE_DELEGATE_OneParam(FSelectRows, TConstArrayView<int32>);
+
+	class FChooserTableViewModel : public IChooserTableViewModel, public FSelfRegisteringEditorUndoClient
+	{
+	public:
+		virtual ~FChooserTableViewModel() override;
+		FChooserTableViewModel(UChooserTable* Chooser);
+
+
+		// IChooserTableViewModel interface
+		virtual void SetShowDetailsDelegate(FShowDetails InShowDetailsDelegate) override
+		{
+			ShowDetailsDelegate = InShowDetailsDelegate;
+		}
+
+		virtual void RegisterMenus(TSharedPtr<FUICommandList> CommandList) override;
+
+		virtual const UChooserTable* GetRootChooser() const override { return RootChooser; }
+		virtual UChooserTable* GetRootChooser() override { return RootChooser; }
+		
+		virtual UChooserTable* GetChooser() override { return CurrentChooser; }
+		virtual const UChooserTable* GetChooser() const override { return CurrentChooser; }
+
+		virtual void RefreshAll() override;
+
+		virtual void SetChooser(UChooserTable* Chooser) override;
+
+		virtual void AutoPopulateAll() override;
+
+		virtual void SelectRootProperties() override;
+		
+		/** Begin FNotifyHook Interface */
+		virtual void NotifyPreChange( FProperty* PropertyAboutToChange ) override;
+		virtual void NotifyPostChange( const FPropertyChangedEvent& PropertyChangedEvent, FProperty* PropertyThatChanged) override;
+
+
+		void QueueRefreshAll();
+	
+		int32 GetNumRows() const;
+
+		bool AreRowsDisabled(TConstArrayView<int32> RowIndices) const;
+		bool IsColumnDisabled(int32 ColumnIndex) const;
+		void ToggleDisableRows(TConstArrayView<int32> RowIndices);
+		void ToggleDisableColumn(int32 ColumnIndex);
+		void DuplicateColumn(int32 ColumnIndex);
+		void DuplicateRows(TConstArrayView<int32> RowIndices);
+		bool HasFallbackSelected();
+		
+		bool IsRandomizeColumn(int32 ColumnIndex);
+
+		bool CanMoveRowsUp(TConstArrayView<int32> RowIndices);
+		bool CanMoveRowsDown(TConstArrayView<int32> RowIndices);
+		bool CanMoveColumnLeft(int32 ColumnIndex);
+		void MoveColumnLeft(int32 ColumnIndex);
+		bool CanMoveColumnRight(int32 ColumnIndex);
+		void MoveColumnRight(int32 ColumnIndex);
+		
+		void CopyColumn(int32 ColumnIndex);
+		void CopyRows(TConstArrayView<int> RowIndices);
+				
+		void RemoveDisabledData();
+
+		bool CanAutoPopulateColumn(int32 ColumnIndex) const;
+		bool CanAutoPopulateRows() const;
+		void AutoPopulateColumn(int32 ColumnIndex);
+		void AutoPopulateRows(TConstArrayView<int> RowIndices);
+		void AutoPopulateRow(int32 Index);
+		
+		void ClearSelectedColumn();
+		void DeleteColumn(int32 Index);
+		void AddColumn(const UScriptStruct* ColumnType);
+
+		
+		void DeleteRows(TConstArrayView<int32> RowsToDelete);
+		int32 DeleteRowsInternal(TConstArrayView<int32> RowsToDelete, int32 RowIndexToRemember = 0);
+		
+		void MoveRows(TConstArrayView<int32> RowIndices, int32 TargetIndex);
+		int32 MoveColumn(int32 SourceIndex, int32 TargetIndex);
+		void Paste();
+		bool CanPaste() const;
+
+		enum class ESelectionType
+		{
+			Root, Rows, Column
+		};
+		
+		void RefreshRowSelectionDetails();
+		ESelectionType GetCurrentSelectionType() const { return CurrentSelectionType; }
+		void SetCurrentSelectionType(ESelectionType SelectionType) { CurrentSelectionType = SelectionType; }
+		
+
+		void ClearSelectedRows();
+		void SelectRows(const TConstArrayView<int32>& Rows);		// called internally to change row selection after operations
+		void SelectColumn(UChooserTable* Chooser, int32 Index);
+		bool IsRowSelected(int32 RowIndex);
+		bool IsColumnSelected(int32 ColumnIndex);
+		bool HasSelection() const;
+		bool HasRowsSelected() const;
+		bool HasColumnSelected() const;
+
+		void SetSelectedRows(const TConstArrayView<int32>& Rows); // called by widget when list view selection changes
+	
+		int32 GetSelectedColumn()
+		{
+			if (SelectedColumn)
+			{
+				return SelectedColumn->Column;
+			}
+			return INDEX_NONE;
+		}
+		
+		TConstArrayView<int32> GetSelectedRows() const
+		{
+			return SelectedRowIndices;
+		}
+		
+		ESelectionType CurrentSelectionType = ESelectionType::Root;
+
+		TArray<int32> SelectedRowIndices;
+		TStrongObjectPtr<UChooserColumnDetails> SelectedColumn;
+		TArray<TObjectPtr<UChooserRowDetails>> SelectedRows;
+		// cache of objects to reuse for SelectedRows
+		TArray<TStrongObjectPtr<UChooserRowDetails>> SelectedRowsCache;
+		
+		FShowDetails ShowDetailsDelegate;
+		FRefreshTable RefreshTableColumnsDelegate;
+		FRefreshTable RefreshTableRowsDelegate;
+		FSelectRows SelectRowsDelegate;
+		
+		UChooserTable* CopyColumnInternal(int32 ColumnIndex);
+		UChooserTable* CopyRowsInternal(TConstArrayView<int32> RowIndices);
+		void PasteInternal(UChooserTable* PasteObject, int32 PasteRowIndex=-1);
+
+		void UpdateTableRows();
+		void UpdateTableColumns();
+
+		void OnObjectsTransacted(UObject* Object, const FTransactionObjectEvent& Event);
+
+		/** FEditorUndoClient Interface */
+		virtual bool MatchesContext( const FTransactionContext& InContext, const TArray<TPair<UObject*, FTransactionObjectEvent>>& TransactionObjectContexts ) const override;
+		virtual void PostUndo(bool bSuccess) override;
+		virtual void PostRedo(bool bSuccess) override;
+		
+	private:
+		void AutoPopulateColumn(FChooserColumnBase& Column);
+		
+		bool HandleTicker(float DeltaTime);
+
+		bool bQueueRefreshAll = false;
+		FTSTicker::FDelegateHandle TickDelegateHandle;
+		
+		// the root asset for the chooser editor
+		UChooserTable* RootChooser = nullptr;
+
+		// the current (possibly nested) chooser table being edited
+		UChooserTable* CurrentChooser = nullptr;
+
+
+	};
+	
+}
\ No newline at end of file
diff --git a/Engine/Plugins/Chooser/Source/ChooserEditor/Private/SChooserColumnHandle.cpp b/Engine/Plugins/Chooser/Source/ChooserEditor/Private/SChooserColumnHandle.cpp
--- a/Engine/Plugins/Chooser/Source/ChooserEditor/Private/SChooserColumnHandle.cpp
+++ b/Engine/Plugins/Chooser/Source/ChooserEditor/Private/SChooserColumnHandle.cpp
@@ -11,10 +11,10 @@
 namespace UE::ChooserEditor
 {
 
-	TSharedRef<FChooserColumnDragDropOp> FChooserColumnDragDropOp::New(FChooserTableEditor* InEditor, uint32 InColumnIndex)
+	TSharedRef<FChooserColumnDragDropOp> FChooserColumnDragDropOp::New(TSharedPtr<FChooserTableViewModel> InViewModel, uint32 InColumnIndex)
 	{
 		TSharedRef<FChooserColumnDragDropOp> Operation = MakeShareable(new FChooserColumnDragDropOp());
-		Operation->ChooserEditor = InEditor;
+		Operation->ChooserViewModel = InViewModel;
 		Operation->ColumnIndex = InColumnIndex;
 		Operation->DefaultHoverText = LOCTEXT("Chooser Column", "Chooser Column");
 		Operation->CurrentHoverText = Operation->DefaultHoverText;
@@ -26,7 +26,9 @@ namespace UE::ChooserEditor
 	
 	void SChooserColumnHandle::Construct(const FArguments& InArgs)
 	{
-		ChooserEditor = InArgs._ChooserEditor;
+		TableHasFocus = InArgs._TableHasFocus;
+		ChooserViewModel = InArgs._ViewModel;
+		check(ChooserViewModel);
 		ColumnIndex = InArgs._ColumnIndex;
 		bNoDropAfter = InArgs._NoDropAfter;
 
@@ -35,11 +37,11 @@ namespace UE::ChooserEditor
 			SNew(SOverlay)
 				+ SOverlay::Slot()
 				[
-					SNew(SColorBlock).Color_Lambda([this](){ return ChooserEditor->TableHasFocus() ? FStyleColors::Select.GetSpecifiedColor() : FStyleColors::SelectInactive.GetSpecifiedColor(); } )
+					SNew(SColorBlock).Color_Lambda([this](){ return TableHasFocus.Get() ? FStyleColors::Select.GetSpecifiedColor() : FStyleColors::SelectInactive.GetSpecifiedColor(); } )
 							.Visibility_Lambda(
 							[this]()
 							{
-								if (ChooserEditor->IsColumnSelected(ColumnIndex))
+								if (ChooserViewModel->IsColumnSelected(ColumnIndex))
 								{
 									return EVisibility::Visible;
 								}
@@ -69,13 +71,13 @@ namespace UE::ChooserEditor
 	
 	FReply SChooserColumnHandle::OnMouseButtonDown(const FGeometry& MyGeometry, const FPointerEvent& MouseEvent)
 	{
-		ChooserEditor->SelectColumn(ChooserEditor->GetChooser(), ColumnIndex);
+		ChooserViewModel->SelectColumn(ChooserViewModel->GetChooser(), ColumnIndex);
 		
 		// act as a move handle if the row is already selected, and if there are multiselect modifiers pressed
 		if (MouseEvent.GetEffectingButton() != EKeys::RightMouseButton && !MouseEvent.IsControlDown() && !MouseEvent.IsShiftDown()
-			&& ChooserEditor->IsColumnSelected(ColumnIndex))
+			&& ChooserViewModel->IsColumnSelected(ColumnIndex))
 		{
-			UChooserTable* Chooser = ChooserEditor->GetChooser();
+			UChooserTable* Chooser = ChooserViewModel->GetChooser();
 			if (!Chooser->ColumnsStructs[ColumnIndex].GetPtr<FRandomizeColumn>()) // don't allow dragging the Randomize Column
 			{
 				return FReply::Handled().DetectDrag(SharedThis(this), EKeys::LeftMouseButton);
@@ -88,9 +90,9 @@ namespace UE::ChooserEditor
 	FReply SChooserColumnHandle::OnDragDetected(const FGeometry& MyGeometry, const FPointerEvent& MouseEvent)
 	{
 		// clear row selection so that delete key can't cause the selected row to be deleted
-		ChooserEditor->ClearSelectedColumn();
+		ChooserViewModel->ClearSelectedColumn();
 		
-		TSharedRef<FChooserColumnDragDropOp> DragDropOp = FChooserColumnDragDropOp::New(ChooserEditor, ColumnIndex);
+		TSharedRef<FChooserColumnDragDropOp> DragDropOp = FChooserColumnDragDropOp::New(ChooserViewModel, ColumnIndex);
 		return FReply::Handled().BeginDragDrop(DragDropOp);
 	}
 
@@ -99,7 +101,7 @@ namespace UE::ChooserEditor
 	{
 		bDropSupported = false;
 
-		UChooserTable* Chooser = ChooserEditor->GetChooser();
+		UChooserTable* Chooser = ChooserViewModel->GetChooser();
 		
 		if (!Chooser->ColumnsStructs.IsValidIndex(ColumnIndex))
 		{
@@ -140,21 +142,21 @@ namespace UE::ChooserEditor
 		
 		if (TSharedPtr<FChooserColumnDragDropOp> Operation = DragDropEvent.GetOperationAs<FChooserColumnDragDropOp>())
 		{
-			if (UChooserTable* Chooser = Operation->ChooserEditor->GetChooser())
+			if (UChooserTable* Chooser = Operation->ChooserViewModel->GetChooser())
 			{
 				int NewColumnIndex = 0; 
 				if (bDropBefore)
 				{
-					NewColumnIndex = ChooserEditor->MoveColumn(Operation->ColumnIndex, ColumnIndex);
+					NewColumnIndex = ChooserViewModel->MoveColumn(Operation->ColumnIndex, ColumnIndex);
 				}
 				else
 				{
-					NewColumnIndex = ChooserEditor->MoveColumn(Operation->ColumnIndex, ColumnIndex+1);
+					NewColumnIndex = ChooserViewModel->MoveColumn(Operation->ColumnIndex, ColumnIndex+1);
 				}
 
 				if (NewColumnIndex >= 0)
 				{
-					ChooserEditor->SelectColumn(Chooser, NewColumnIndex);
+					ChooserViewModel->SelectColumn(Chooser, NewColumnIndex);
 				}
 			}
 		}
diff --git a/Engine/Plugins/Chooser/Source/ChooserEditor/Private/SChooserColumnHandle.h b/Engine/Plugins/Chooser/Source/ChooserEditor/Private/SChooserColumnHandle.h
--- a/Engine/Plugins/Chooser/Source/ChooserEditor/Private/SChooserColumnHandle.h
+++ b/Engine/Plugins/Chooser/Source/ChooserEditor/Private/SChooserColumnHandle.h
@@ -6,18 +6,18 @@
 namespace UE::ChooserEditor
 {
 
-class FChooserTableEditor;
+class FChooserTableViewModel;
 
 class FChooserColumnDragDropOp : public FDecoratedDragDropOp
 {
 public:
 	DRAG_DROP_OPERATOR_TYPE(FChooserColumnDragDropOp, FDecoratedDragDropOp)
 
-	FChooserTableEditor* ChooserEditor;
+	TSharedPtr<FChooserTableViewModel> ChooserViewModel;
 	uint32 ColumnIndex;
 
 	/** Constructs the drag drop operation */
-	static TSharedRef<FChooserColumnDragDropOp> New(FChooserTableEditor* InEditor, uint32 InColumnIndex);
+	static TSharedRef<FChooserColumnDragDropOp> New(TSharedPtr<FChooserTableViewModel> InViewModel, uint32 InColumnIndex);
 };
 
 class SChooserColumnHandle : public SCompoundWidget
@@ -26,9 +26,10 @@ public:
 	SLATE_BEGIN_ARGS(SChooserColumnHandle)
 	{}
 	SLATE_DEFAULT_SLOT(FArguments, Content)
-	SLATE_ARGUMENT(FChooserTableEditor*, ChooserEditor)
+	SLATE_ARGUMENT(TSharedPtr<FChooserTableViewModel>, ViewModel)
 	SLATE_ARGUMENT(uint32, ColumnIndex)
 	SLATE_ARGUMENT(bool, NoDropAfter)
+	SLATE_ATTRIBUTE(bool, TableHasFocus);
 	SLATE_END_ARGS()
 
 	void Construct(const FArguments& InArgs);
@@ -41,7 +42,8 @@ public:
 	FReply OnDrop(const FGeometry& MyGeometry, const FDragDropEvent& DragDropEvent);
 
 private:
-	FChooserTableEditor* ChooserEditor = nullptr;
+	TAttribute<bool> TableHasFocus;
+	TSharedPtr<FChooserTableViewModel> ChooserViewModel;
 	int32 ColumnIndex = 0;
 	int32 DragActiveCounter = 0;
 	bool bDragActive = false;
diff --git a/Engine/Plugins/Chooser/Source/ChooserEditor/Private/SChooserCreateColumnButton.cpp b/Engine/Plugins/Chooser/Source/ChooserEditor/Private/SChooserCreateColumnButton.cpp
new file mode 100644
--- /dev/null
+++ b/Engine/Plugins/Chooser/Source/ChooserEditor/Private/SChooserCreateColumnButton.cpp
@@ -0,0 +1,122 @@
+// Copyright Epic Games, Inc. All Rights Reserved.
+
+#include "SChooserCreateColumnButton.h"
+#include "ChooserTableEditor.h"
+#include "Framework/Application/SlateApplication.h"
+#include "SPositiveActionButton.h"
+#include "UObject/UObjectIterator.h"
+
+#define LOCTEXT_NAMESPACE "AddColumnButton"
+
+namespace UE::ChooserEditor
+{
+	struct FColumnTypeInfoStruct
+	{
+		int SortOrder = 100;
+		FString Category;
+		UScriptStruct* Type;
+		bool operator < (const FColumnTypeInfoStruct& Other) const
+		{
+			if (Category == Other.Category)
+			{
+				return Type->GetDisplayNameText().ToString() < Other.Type->GetDisplayNameText().ToString();
+			}
+			else if (SortOrder == Other.SortOrder)
+			{
+				return Category < Other.Category;
+			}
+			else
+			{
+				return SortOrder < Other.SortOrder;
+			}
+		}
+	};
+
+	TSharedRef<SWidget>	SChooserCreateColumnButton::MakeCreateColumnMenu()
+    {
+    	FMenuBuilder MenuBuilder(true, nullptr);
+    	static TArray<FColumnTypeInfoStruct> ColumnTypes;
+    
+    	if (ColumnTypes.IsEmpty())
+    	{
+    		UScriptStruct* BaseType = FChooserColumnBase::StaticStruct();
+    		for (TObjectIterator<UScriptStruct> StructIt; StructIt; ++StructIt)
+    		{
+    			if (*StructIt != BaseType && StructIt->IsChildOf(BaseType))
+    			{
+    				if (!StructIt->HasMetaData("Hidden"))
+    				{
+    					FColumnTypeInfoStruct Info;
+    					Info.Type = *StructIt;
+    					Info.Category = StructIt->HasMetaData("Category") ? StructIt->GetMetaData("Category") : "Other";
+    
+    					if (Info.Category == "Filter")
+    					{
+    						Info.SortOrder = 1;
+    					}
+    					else if (Info.Category == "Scoring")
+                       	{
+                       		Info.SortOrder = 2;
+                       	}
+    					else if (Info.Category == "Output")
+                       	{
+                       		Info.SortOrder = 3;
+                       	}
+    					else if (Info.Category == "Random")
+    					{
+    						Info.SortOrder = 4;
+    					}
+    
+    					ColumnTypes.Add(Info);
+    				}
+    			}
+    		}
+    		ColumnTypes.Sort();
+    	}
+    
+    	FString Section = "";
+    	for(FColumnTypeInfoStruct& Type : ColumnTypes)
+    	{
+    		if (Section != Type.Category)
+    		{
+    			if (Section != "")
+    			{
+    				MenuBuilder.EndSection();
+    			}
+    			Section = Type.Category;
+    			MenuBuilder.BeginSection(FName(Section), FText::FromString(Section));
+    		}
+    		
+    		MenuBuilder.AddMenuEntry(Type.Type->GetDisplayNameText(), Type.Type->GetToolTipText(), FSlateIcon(),
+    			FUIAction(
+    				FExecuteAction::CreateLambda([this, Type]()
+    				{
+    					ChooserViewModel->AddColumn(Type.Type);
+    				}))
+    			);
+    							
+    	}
+    	return MenuBuilder.MakeWidget();
+    }
+		
+	void SChooserCreateColumnButton::Construct(const FArguments& InArgs)
+	{
+		ChooserViewModel = InArgs._ViewModel;
+		check(ChooserViewModel);
+		
+		ChildSlot
+		[
+			SNew(SPositiveActionButton)
+            	.Text(LOCTEXT("Add Column", "Add Column"))
+            	.OnGetMenuContent(this, &SChooserCreateColumnButton::MakeCreateColumnMenu)
+		];
+
+	}
+
+	SChooserCreateColumnButton::~SChooserCreateColumnButton()
+	{
+	}
+
+}
+
+#undef LOCTEXT_NAMESPACE
diff --git a/Engine/Plugins/Chooser/Source/ChooserEditor/Private/SChooserCreateColumnButton.h b/Engine/Plugins/Chooser/Source/ChooserEditor/Private/SChooserCreateColumnButton.h
new file mode 100644
--- /dev/null
+++ b/Engine/Plugins/Chooser/Source/ChooserEditor/Private/SChooserCreateColumnButton.h
@@ -0,0 +1,31 @@
+// Copyright Epic Games, Inc. All Rights Reserved.
+#pragma once
+
+#include "Widgets/SCompoundWidget.h"
+#include "Widgets/Views/SListView.h"
+#include "Framework/Commands/UICommandList.h"
+#include "SChooserTableRow.h"
+#include "ChooserTableEditor.h"
+
+class SPositiveActionButton;
+
+namespace UE::ChooserEditor
+{
+	
+class SChooserCreateColumnButton : public SCompoundWidget
+{
+public:
+	SLATE_BEGIN_ARGS(SChooserCreateColumnButton)
+	{}
+	SLATE_ARGUMENT(TSharedPtr<FChooserTableViewModel>, ViewModel)
+	SLATE_END_ARGS()
+
+	void Construct(const FArguments& InArgs);
+	virtual ~SChooserCreateColumnButton() override;
+
+ private:
+	TSharedPtr<FChooserTableViewModel> ChooserViewModel;
+	TSharedRef<SWidget>	MakeCreateColumnMenu();
+};
+
+}
diff --git a/Engine/Plugins/Chooser/Source/ChooserEditor/Private/SChooserCreateRowButton.cpp b/Engine/Plugins/Chooser/Source/ChooserEditor/Private/SChooserCreateRowButton.cpp
new file mode 100644
--- /dev/null
+++ b/Engine/Plugins/Chooser/Source/ChooserEditor/Private/SChooserCreateRowButton.cpp
@@ -0,0 +1,219 @@
+// Copyright Epic Games, Inc. All Rights Reserved.
+
+#include "SChooserCreateRowButton.h"
+#include "ChooserTableEditor.h"
+#include "Framework/Application/SlateApplication.h"
+#include "SPositiveActionButton.h"
+#include "UObject/UObjectIterator.h"
+
+#include "ScopedTransaction.h"
+#include "ObjectChooser_Class.h"
+
+#define LOCTEXT_NAMESPACE "AddColumnButton"
+
+namespace UE::ChooserEditor
+{
+	struct FResultTypeInfoStruct
+	{
+		bool ObjectOnly = false;
+		bool ClassOnly = false;
+		FString Category;
+		UScriptStruct* Type = nullptr;
+		
+		bool operator < (const FResultTypeInfoStruct& Other) const
+		{
+			if (Category == Other.Category)
+			{
+				return Type->GetDisplayNameText().ToString() < Other.Type->GetDisplayNameText().ToString();
+			}
+			else
+			{
+				return Category < Other.Category;
+			}
+		}
+	};
+
+	DECLARE_DELEGATE_OneParam(FCreateStructDelegate, UScriptStruct*);
+
+	void MakeCreateResultMenu(FMenuBuilder& MenuBuilder, EObjectChooserResultType ChooserResultType, FCreateStructDelegate CreateStruct)
+	{
+		static TArray<FResultTypeInfoStruct> ResultTypes;
+
+		if (ResultTypes.IsEmpty())
+		{
+			UScriptStruct* BaseType = FObjectChooserBase::StaticStruct();
+			for (TObjectIterator<UScriptStruct> StructIt; StructIt; ++StructIt)
+			{
+				if (*StructIt != BaseType && StructIt->IsChildOf(BaseType))
+				{
+					if (!StructIt->HasMetaData("Hidden"))
+					{
+						FResultTypeInfoStruct Info;
+						Info.Type = *StructIt;
+						Info.Category = StructIt->HasMetaData("Category") ? StructIt->GetMetaData("Category") : "Other";
+
+						if (StructIt->HasMetaData("ResultType"))
+						{
+							FString ResultTypeString = StructIt->GetMetaData("ResultType");
+							Info.ClassOnly = ResultTypeString == "Class";
+							Info.ObjectOnly = ResultTypeString == "Object";
+						}
+						
+						ResultTypes.Add(Info);
+					}
+				}
+			}
+			ResultTypes.Sort();
+		}
+
+		FString Section = "";
+		for(FResultTypeInfoStruct& Type : ResultTypes)
+		{
+			if (Section != Type.Category)
+			{
+				if (Section != "")
+				{
+					MenuBuilder.EndSection();
+				}
+				Section = Type.Category;
+				MenuBuilder.BeginSection(FName(Section), FText::FromString(Section));
+			}
+			
+			MenuBuilder.AddMenuEntry(Type.Type->GetDisplayNameText(), Type.Type->GetToolTipText(), FSlateIcon(),
+				FUIAction(
+					FExecuteAction::CreateLambda([Type, CreateStruct]()
+					{
+						CreateStruct.Execute(Type.Type);
+					}),
+					FCanExecuteAction::CreateLambda([Type, ChooserResultType]()
+					{
+						if (Type.ClassOnly && ChooserResultType == EObjectChooserResultType::ObjectResult)
+						{
+							return false;
+						}
+						if (Type.ObjectOnly && ChooserResultType == EObjectChooserResultType::ClassResult)
+						{
+							return false;
+						}
+						return true;
+					})
+					)
+				);
+								
+		}
+	}
+
+	
+
+	TSharedRef<SWidget>	SChooserCreateRowButton::MakeCreateRowMenu()
+    {
+		FMenuBuilder MenuBuilder(true, nullptr);
+	 
+		UChooserTable* Chooser = ChooserViewModel->GetChooser();
+		if (!Chooser->FallbackResult.IsValid())
+		{
+			if (Chooser->ResultType == EObjectChooserResultType::NoPrimaryResult)
+			{
+				MenuBuilder.AddMenuEntry(
+					LOCTEXT("Add Fallback Output", "Add Fallback Output"),
+					LOCTEXT("Add Fallback Output Tooltip", "Add a Fallback row to the chooser, which will be used in the case where no other rows passed all filter columns"),
+					FSlateIcon(),
+					FUIAction(FExecuteAction::CreateLambda([this]()
+					{
+						UChooserTable* Chooser = ChooserViewModel->GetChooser();
+						const FScopedTransaction Transaction(LOCTEXT("Add Fallback Row Transaction", "Add Fallback Row"));
+						Chooser->Modify(true);
+						
+						// Just construct a dummy result to make sure all rows always have "valid results"
+						// You can't just leave a null result otherwise rows don't apply their output.
+						Chooser->FallbackResult.InitializeAs(FClassChooser::StaticStruct());
+						Chooser->FallbackResult.GetMutable<FClassChooser>().Class = UClass::StaticClass();
+						
+						ChooserViewModel->UpdateTableRows();
+					}))
+				);
+			}
+			else
+			{
+				MenuBuilder.AddSubMenu(
+					LOCTEXT("Add Fallback", "Add Fallback Result"),
+						LOCTEXT("Add Fallback Tooltip", "Add a Fallback row to the chooser, which will be used in the case where no other rows passed all filter columns"),
+						FNewMenuDelegate::CreateLambda([this](FMenuBuilder& MenuBuilder)
+						{
+							UChooserTable* Chooser = ChooserViewModel->GetChooser();
+							MakeCreateResultMenu(MenuBuilder, Chooser->GetContextOwner()->ResultType, FCreateStructDelegate::CreateLambda([this](UScriptStruct* Type)
+							{
+								UChooserTable* Chooser = ChooserViewModel->GetChooser();
+								const FScopedTransaction Transaction(LOCTEXT("Add Fallback Row Transaction", "Add Fallback Row"));
+								Chooser->Modify(true);
+								
+								Chooser->FallbackResult.InitializeAs(Type);
+								
+								ChooserViewModel->UpdateTableRows();
+							}));
+						})
+					);
+			}
+		}
+		
+		if (Chooser->ResultType == EObjectChooserResultType::NoPrimaryResult)
+		{
+			MenuBuilder.AddMenuEntry(
+				LOCTEXT("Add Output Row", "Add Output Row"),
+				LOCTEXT("Add Output Row Tooltip", "Add a regular row to the chooser"),
+				FSlateIcon(),
+				FUIAction(FExecuteAction::CreateLambda([this]()
+				{
+					UChooserTable* Chooser = ChooserViewModel->GetChooser();
+					const FScopedTransaction Transaction(LOCTEXT("Add Row Transaction", "Add Row"));
+					Chooser->Modify(true);
+	 
+					// Just construct a dummy result to make sure all rows always have "valid results"
+					// You can't just leave a null result otherwise rows don't apply their output.
+					FInstancedStruct& NewResult = Chooser->ResultsStructs.AddDefaulted_GetRef();
+					NewResult.InitializeAs(FClassChooser::StaticStruct());
+					NewResult.GetMutable<FClassChooser>().Class = UClass::StaticClass();
+					
+					ChooserViewModel->UpdateTableRows();
+				}))
+			);
+		}
+		else
+		{
+			MakeCreateResultMenu(MenuBuilder, Chooser->GetContextOwner()->ResultType, FCreateStructDelegate::CreateLambda([this](UScriptStruct* Type)
+			{
+				UChooserTable* Chooser = ChooserViewModel->GetChooser();
+				const FScopedTransaction Transaction(LOCTEXT("Add Row Transaction", "Add Row"));
+				Chooser->Modify(true);
+				
+				FInstancedStruct& NewResult = Chooser->ResultsStructs.AddDefaulted_GetRef();
+				NewResult.InitializeAs(Type);
+				
+				ChooserViewModel->UpdateTableRows();
+			}));
+		}
+	 
+		return MenuBuilder.MakeWidget();
+    }
+		
+	void SChooserCreateRowButton::Construct(const FArguments& InArgs)
+	{
+		ChooserViewModel = InArgs._ViewModel;
+		check(ChooserViewModel);
+		
+		ChildSlot
+		[
+			SNew(SPositiveActionButton)
+				.Text(LOCTEXT("Add Row", "Add Row"))
+            	.OnGetMenuContent(this, &SChooserCreateRowButton::MakeCreateRowMenu)
+		];
+
+	}
+
+	SChooserCreateRowButton::~SChooserCreateRowButton()
+	{
+	}
+
+}
+
+#undef LOCTEXT_NAMESPACE
diff --git a/Engine/Plugins/Chooser/Source/ChooserEditor/Private/SChooserCreateRowButton.h b/Engine/Plugins/Chooser/Source/ChooserEditor/Private/SChooserCreateRowButton.h
new file mode 100644
--- /dev/null
+++ b/Engine/Plugins/Chooser/Source/ChooserEditor/Private/SChooserCreateRowButton.h
@@ -0,0 +1,31 @@
+// Copyright Epic Games, Inc. All Rights Reserved.
+#pragma once
+
+#include "Widgets/SCompoundWidget.h"
+#include "Widgets/Views/SListView.h"
+#include "Framework/Commands/UICommandList.h"
+#include "SChooserTableRow.h"
+#include "ChooserTableEditor.h"
+
+class SPositiveActionButton;
+
+namespace UE::ChooserEditor
+{
+	
+class SChooserCreateRowButton : public SCompoundWidget
+{
+public:
+	SLATE_BEGIN_ARGS(SChooserCreateRowButton)
+	{}
+	SLATE_ARGUMENT(TSharedPtr<FChooserTableViewModel>, ViewModel)
+	SLATE_END_ARGS()
+
+	void Construct(const FArguments& InArgs);
+	virtual ~SChooserCreateRowButton() override;
+
+ private:
+	TSharedPtr<FChooserTableViewModel> ChooserViewModel;
+	TSharedRef<SWidget>	MakeCreateRowMenu();
+};
+
+}
diff --git a/Engine/Plugins/Chooser/Source/ChooserEditor/Private/SChooserRowHandle.cpp b/Engine/Plugins/Chooser/Source/ChooserEditor/Private/SChooserRowHandle.cpp
--- a/Engine/Plugins/Chooser/Source/ChooserEditor/Private/SChooserRowHandle.cpp
+++ b/Engine/Plugins/Chooser/Source/ChooserEditor/Private/SChooserRowHandle.cpp
@@ -9,17 +9,17 @@
 namespace UE::ChooserEditor
 {
 
-	TSharedRef<FChooserRowDragDropOp> FChooserRowDragDropOp::New(FChooserTableEditor* InEditor, uint32 InRowIndex)
+	TSharedRef<FChooserRowDragDropOp> FChooserRowDragDropOp::New(TSharedPtr<FChooserTableViewModel> InViewModel, uint32 InRowIndex)
 	{
 		TSharedRef<FChooserRowDragDropOp> Operation = MakeShareable(new FChooserRowDragDropOp());
 		Operation->DefaultHoverText = LOCTEXT("Chooser Row", "Chooser Row");
 		Operation->CurrentHoverText = Operation->DefaultHoverText;
 
 		Operation->TransactionIndex = GEditor->BeginTransaction(LOCTEXT("Drag Chooser Table Rows", "Drag Chooser Table Rows"));
-		
-		Operation->RowData = InEditor->CopySelectionInternal();
-		InEditor->DeleteSelectedRowsInternal();
-		Operation->Editor = InEditor;
+
+		Operation->RowData = InViewModel->CopyRowsInternal(InViewModel->GetSelectedRows());
+		InViewModel->DeleteRowsInternal(InViewModel->GetSelectedRows());
+		Operation->ChooserViewModel = InViewModel;
 		
 		Operation->Construct();
 
@@ -35,14 +35,14 @@ namespace UE::ChooserEditor
 			GEditor->CancelTransaction(TransactionIndex);
 
 			// PostUndo doesn't get called for cancelled transactions, so refresh the editor manually
-			Editor->RefreshAll();
+			ChooserViewModel->RefreshAll();
 		}
 		FDecoratedDragDropOp::OnDrop(bDropWasHandled, MouseEvent);
 	}
 
 	void SChooserRowHandle::Construct(const FArguments& InArgs, bool bShowImage)
 	{
-		ChooserEditor = InArgs._ChooserEditor;
+		ChooserViewModel = InArgs._ViewModel;
 		RowIndex = InArgs._RowIndex;
 
 		if (bShowImage)
@@ -65,7 +65,7 @@ namespace UE::ChooserEditor
 						SNew(SImage)
 						.Visibility_Lambda([this]()
 						{
-							return ChooserEditor->GetChooser()->GetDebugTestValuesValid() && RowIndex == ChooserEditor->GetChooser()->GetDebugSelectedRow() ? EVisibility::HitTestInvisible : EVisibility::Hidden;
+							return ChooserViewModel->GetChooser()->GetDebugTestValuesValid() && RowIndex == ChooserViewModel->GetChooser()->GetDebugSelectedRow() ? EVisibility::HitTestInvisible : EVisibility::Hidden;
 						})
 						.Image(FAppStyle::Get().GetBrush("Icons.ArrowRight"))
 					]
@@ -74,25 +74,11 @@ namespace UE::ChooserEditor
 		}
 	}
 	
-	FReply SChooserRowHandle::OnMouseButtonDoubleClick(const FGeometry& InMyGeometry, const FPointerEvent& InMouseEvent)
-	{
-		UChooserTable* Chooser = ChooserEditor->GetChooser();
-
-		if (const FNestedChooser* NestedChooserResult = Chooser->ResultsStructs[RowIndex].GetPtr<FNestedChooser>())
-		{
-			if (NestedChooserResult->Chooser)
-			{
-				ChooserEditor->PushChooserTableToEdit(NestedChooserResult->Chooser);
-			}
-		}
-		return FReply::Handled();
-	}
-
 	FReply SChooserRowHandle::OnMouseButtonDown(const FGeometry& MyGeometry, const FPointerEvent& MouseEvent)
 	{
 		// act as a move handle if the row is already selected, and if there are multiselect modifiers pressed
 		if (MouseEvent.GetEffectingButton() != EKeys::RightMouseButton && !MouseEvent.IsControlDown() && !MouseEvent.IsShiftDown()
-			&& ChooserEditor->IsRowSelected(RowIndex))
+			&& ChooserViewModel->IsRowSelected(RowIndex))
 		{
 			return FReply::Handled().DetectDrag(SharedThis(this), EKeys::LeftMouseButton);
 		}
@@ -104,7 +90,7 @@ namespace UE::ChooserEditor
 
 	FReply SChooserRowHandle::OnDragDetected(const FGeometry& MyGeometry, const FPointerEvent& MouseEvent)
 	{
-		TSharedRef<FChooserRowDragDropOp> DragDropOp = FChooserRowDragDropOp::New(ChooserEditor, RowIndex);
+		TSharedRef<FChooserRowDragDropOp> DragDropOp = FChooserRowDragDropOp::New(ChooserViewModel, RowIndex);
 		return FReply::Handled().BeginDragDrop(DragDropOp);
 	}
 
diff --git a/Engine/Plugins/Chooser/Source/ChooserEditor/Private/SChooserRowHandle.h b/Engine/Plugins/Chooser/Source/ChooserEditor/Private/SChooserRowHandle.h
--- a/Engine/Plugins/Chooser/Source/ChooserEditor/Private/SChooserRowHandle.h
+++ b/Engine/Plugins/Chooser/Source/ChooserEditor/Private/SChooserRowHandle.h
@@ -8,7 +8,7 @@ class UChooserTable;
 namespace UE::ChooserEditor
 {
 
-class FChooserTableEditor;
+class FChooserTableViewModel;
 
 class FChooserRowDragDropOp : public FDecoratedDragDropOp
 {
@@ -17,10 +17,10 @@ public:
 
 	UChooserTable* RowData = nullptr;
 	int32 TransactionIndex = -1;
-	FChooserTableEditor* Editor = nullptr;
+	TSharedPtr<FChooserTableViewModel> ChooserViewModel;
 
 	/** Constructs the drag drop operation */
-	static TSharedRef<FChooserRowDragDropOp> New(FChooserTableEditor* InEditor, uint32 InRowIndex);
+	static TSharedRef<FChooserRowDragDropOp> New(TSharedPtr<FChooserTableViewModel> InViewModel, uint32 InRowIndex);
 
 	virtual void OnDrop(bool bDropWasHandled, const FPointerEvent& MouseEvent) override;
 };
@@ -31,18 +31,17 @@ public:
 	SLATE_BEGIN_ARGS(SChooserRowHandle)
 	{}
 	SLATE_DEFAULT_SLOT(FArguments, Content)
-	SLATE_ARGUMENT(FChooserTableEditor*, ChooserEditor)
+	SLATE_ARGUMENT(TSharedPtr<FChooserTableViewModel>, ViewModel)
 	SLATE_ARGUMENT(uint32, RowIndex)
 	SLATE_END_ARGS()
 
 	void Construct(const FArguments& InArgs, bool bShowImage);
 	
-	virtual FReply OnMouseButtonDoubleClick(const FGeometry& InMyGeometry, const FPointerEvent& InMouseEvent) override;
 	FReply OnMouseButtonDown(const FGeometry& MyGeometry, const FPointerEvent& MouseEvent) override;
 	FReply OnDragDetected(const FGeometry& MyGeometry, const FPointerEvent& MouseEvent) override;
 
 private:
-	FChooserTableEditor* ChooserEditor = nullptr;
+	TSharedPtr<FChooserTableViewModel> ChooserViewModel;
 	uint32 RowIndex = 0;
 };
 
diff --git a/Engine/Plugins/Chooser/Source/ChooserEditor/Private/SChooserTableRow.cpp b/Engine/Plugins/Chooser/Source/ChooserEditor/Private/SChooserTableRow.cpp
--- a/Engine/Plugins/Chooser/Source/ChooserEditor/Private/SChooserTableRow.cpp
+++ b/Engine/Plugins/Chooser/Source/ChooserEditor/Private/SChooserTableRow.cpp
@@ -16,6 +16,8 @@
 #include "Widgets/Views/STableRow.h"
 #include "Widgets/Layout/SSeparator.h"
 
+#include "SChooserCreateRowButton.h"
+
 #define LOCTEXT_NAMESPACE "ChooserTableRow"
 
 namespace UE::ChooserEditor
@@ -26,9 +28,11 @@ namespace UE::ChooserEditor
 	
 	void SChooserTableRow::Construct(const FArguments& Args, const TSharedRef<STableViewBase>& OwnerTableView)
 	{
+		TableHasFocus = Args._TableHasFocus;
 		RowIndex = Args._Entry;
 		Chooser = Args._Chooser;
-		Editor = Args._Editor;
+		ChooserViewModel = Args._ViewModel;
+		check(ChooserViewModel);
 
 		SMultiColumnTableRow<TSharedPtr<FChooserTableRow>>::Construct(
 			FSuperRowType::FArguments(),
@@ -94,7 +98,7 @@ namespace UE::ChooserEditor
 				return SNew(SOverlay)
 						+ SOverlay::Slot()
 						[
-							SNew(SChooserRowHandle, bShowHandleImage).ChooserEditor(Editor).RowIndex(RowIndex->RowIndex)
+							SNew(SChooserRowHandle, bShowHandleImage).ViewModel(ChooserViewModel).RowIndex(RowIndex->RowIndex)
 						]
 						+ SOverlay::Slot()
 						[
@@ -115,8 +119,8 @@ namespace UE::ChooserEditor
 			{
 				FChooserWidgetValueChanged ChooserWidgetValueChanged = FChooserWidgetValueChanged::CreateLambda([this]()
 					{
-						check(Editor);
-						Editor->AutoPopulateRow(RowIndex->RowIndex);
+						check(ChooserViewModel);
+						ChooserViewModel->AutoPopulateRow(RowIndex->RowIndex);
 					});
 
 				UChooserTable* ContextOwner = Chooser->GetRootChooser();
@@ -192,11 +196,11 @@ namespace UE::ChooserEditor
 								SNew(SOverlay)
 								+ SOverlay::Slot()
 								[
-									SNew(SColorBlock).Color_Lambda([this, ColumnIndex](){ return Editor->TableHasFocus() ? FStyleColors::Select.GetSpecifiedColor() : FStyleColors::SelectInactive.GetSpecifiedColor(); } )
+									SNew(SColorBlock).Color_Lambda([this, ColumnIndex](){ return TableHasFocus.Get() ? FStyleColors::Select.GetSpecifiedColor() : FStyleColors::SelectInactive.GetSpecifiedColor(); } )
 										.Visibility_Lambda(
 										[this, ColumnIndex]()
 										{
-											if (Editor->IsColumnSelected(ColumnIndex))
+											if (ChooserViewModel->IsColumnSelected(ColumnIndex))
 											{
 												return EVisibility::Visible;
 											}
@@ -302,8 +306,8 @@ namespace UE::ChooserEditor
 			{
 				FChooserWidgetValueChanged ChooserWidgetValueChanged = FChooserWidgetValueChanged::CreateLambda([this]()
 					{
-						check(Editor);
-						Editor->AutoPopulateRow(RowIndex->RowIndex);
+						check(ChooserViewModel);
+						ChooserViewModel->AutoPopulateRow(RowIndex->RowIndex);
 					});
 
 				UChooserTable* ContextOwner = Chooser->GetRootChooser();
@@ -362,11 +366,11 @@ namespace UE::ChooserEditor
 									SNew(SOverlay)
 									+SOverlay::Slot()
 									[
-										SNew(SColorBlock).Color_Lambda([this, ColumnIndex](){ return Editor->TableHasFocus() ? FStyleColors::Select.GetSpecifiedColor() : FStyleColors::SelectInactive.GetSpecifiedColor(); } )
+										SNew(SColorBlock).Color_Lambda([this, ColumnIndex](){ return TableHasFocus.Get() ? FStyleColors::Select.GetSpecifiedColor() : FStyleColors::SelectInactive.GetSpecifiedColor(); } )
 											.Visibility_Lambda(
 											[this, ColumnIndex]()
 											{
-												if (Editor->IsColumnSelected(ColumnIndex))
+												if (ChooserViewModel->IsColumnSelected(ColumnIndex))
 												{
 													return EVisibility::Visible;
 												}
@@ -405,7 +409,7 @@ namespace UE::ChooserEditor
 				return SNew(SHorizontalBox)
 								+SHorizontalBox::Slot().AutoWidth()
 								[
-									Editor->GetCreateRowComboButton().ToSharedRef()
+									SNew(SChooserCreateRowButton).ViewModel(ChooserViewModel)
 								];
 			}
 		}
@@ -499,7 +503,7 @@ namespace UE::ChooserEditor
 			{
 				NewRowIndex = RowIndex->RowIndex+1;
 			}
-			Editor->PasteInternal(Operation->RowData, NewRowIndex);
+			ChooserViewModel->PasteInternal(Operation->RowData, NewRowIndex);
 			GEditor->EndTransaction();
 		}
 		else if (TSharedPtr<FAssetDragDropOp> ContentDragDropOp = DragDropEvent.GetOperationAs<FAssetDragDropOp>())
@@ -558,17 +562,17 @@ namespace UE::ChooserEditor
 						Column.InsertRows(InsertRowIndex, NewResults.Num());
 					}
 
-					Editor->ClearSelectedRows();
-					Editor->RefreshAll();
+					ChooserViewModel->ClearSelectedRows();
+					ChooserViewModel->RefreshAll();
 					
 					TArray<int32, TInlineAllocator<256>> RowsToSelect;
 					RowsToSelect.Reserve(NewResults.Num());
 					for(int Index = InsertRowIndex; Index < InsertRowIndex + NewResults.Num(); Index++)
 					{
-						Editor->AutoPopulateRow(Index);
+						ChooserViewModel->AutoPopulateRow(Index);
 						RowsToSelect.Add(Index);
 					}
-					Editor->SelectRows(RowsToSelect);
+					ChooserViewModel->SelectRows(RowsToSelect);
 				}
 			}
 		}
diff --git a/Engine/Plugins/Chooser/Source/ChooserEditor/Private/SChooserTableRow.h b/Engine/Plugins/Chooser/Source/ChooserEditor/Private/SChooserTableRow.h
--- a/Engine/Plugins/Chooser/Source/ChooserEditor/Private/SChooserTableRow.h
+++ b/Engine/Plugins/Chooser/Source/ChooserEditor/Private/SChooserTableRow.h
@@ -9,7 +9,7 @@ class UChooserTable;
 namespace UE::ChooserEditor
 {
 
-class FChooserTableEditor;
+class FChooserTableViewModel;
 
 struct FChooserTableRow
 {
@@ -24,7 +24,8 @@ SLATE_BEGIN_ARGS(SChooserTableRow) {}
 /** The list item for this row */
 	SLATE_ARGUMENT(TSharedPtr<FChooserTableRow>, Entry)
 	SLATE_ARGUMENT(UChooserTable*, Chooser)
-	SLATE_ARGUMENT(FChooserTableEditor*, Editor)
+	SLATE_ARGUMENT(TSharedPtr<FChooserTableViewModel>, ViewModel)
+	SLATE_ATTRIBUTE(bool, TableHasFocus);
 	SLATE_END_ARGS()
 
 	enum { SpecialIndex_AddRow = -1, SpecialIndex_Fallback = -2 };
@@ -40,9 +41,10 @@ SLATE_BEGIN_ARGS(SChooserTableRow) {}
 	virtual void Tick(const FGeometry& AllottedGeometry, const double InCurrentTime, const float InDeltaTime) override;
 
 private:
+	TAttribute<bool> TableHasFocus;
 	TSharedPtr<FChooserTableRow> RowIndex;
 	UChooserTable* Chooser;
-	FChooserTableEditor* Editor;
+	TSharedPtr<FChooserTableViewModel> ChooserViewModel;
 	TSharedPtr<SBorder> CacheBorder;
 	int DragActiveCounter = 0;
 	bool bDragActive = false;
diff --git a/Engine/Plugins/Chooser/Source/ChooserEditor/Private/SChooserTableWidget.cpp b/Engine/Plugins/Chooser/Source/ChooserEditor/Private/SChooserTableWidget.cpp
new file mode 100644
--- /dev/null
+++ b/Engine/Plugins/Chooser/Source/ChooserEditor/Private/SChooserTableWidget.cpp
@@ -0,0 +1,607 @@
+// Copyright Epic Games, Inc. All Rights Reserved.
+
+#include "SChooserTableWidget.h"
+#include "ChooserTableEditor.h"
+#include "Chooser.h"
+#include "Framework/Application/SlateApplication.h"
+#include "ObjectTools.h"
+#include "OutputObjectColumn.h"
+#include "Widgets/Input/STextEntryPopup.h"
+#include "ObjectChooserWidgetFactories.h"
+#include "SChooserCreateColumnButton.h"
+#include "SChooserColumnHandle.h"
+#include "SChooserTableRow.h"
+#include "ToolMenus.h"
+#include "ChooserTableEditorCommands.h"
+#include "Framework/Commands/GenericCommands.h"
+
+
+#define LOCTEXT_NAMESPACE "ChooserTableWidget"
+
+namespace UE::ChooserEditor
+{
+	
+	bool SChooserTableWidget::TableHasFocus() const
+	{
+		if (TableView)
+		{
+			return TableView->HasKeyboardFocus();
+		}
+		return false;
+	}
+	
+	void SChooserTableWidget::ClearSelection()
+	{
+		if (TableView)
+		{
+			TableView->ClearSelection();
+		}
+	}
+	
+	void SChooserTableWidget::GetSelectedItems(TArray<TSharedPtr<FChooserTableRow>>& SelectedItems)
+	{
+		TableView->GetSelectedItems(SelectedItems);
+	}
+
+	int SChooserTableWidget::GetSelectedColumn() const
+	{
+		if (ChooserViewModel)
+		{
+			return ChooserViewModel->GetSelectedColumn();
+		}
+		return INDEX_NONE;
+	}
+	
+	void SChooserTableWidget::UpdateTableColumns()
+	{
+		if (ChooserViewModel)
+		{
+			UChooserTable* Chooser = ChooserViewModel->GetChooser();
+
+			HeaderRow->ClearColumns();
+
+			HeaderRow->AddColumn(SHeaderRow::Column("Handles")
+							.DefaultLabel(FText())
+							.ManualWidth(30));
+
+			if (Chooser->ResultType != EObjectChooserResultType::NoPrimaryResult)
+			{
+				HeaderRow->AddColumn(SHeaderRow::Column("Result")
+								.ManualWidth_Lambda([Chooser]() { return Chooser->EditorResultsColumnWidth; } )
+								.OnWidthChanged_Lambda([Chooser](float NewWidth) { Chooser->EditorResultsColumnWidth = NewWidth; })
+								.HeaderContent()
+								[
+									SNew(SVerticalBox)
+									+ SVerticalBox::Slot()
+									.VAlign(VAlign_Top)
+									[
+										SNew(STextBlock)
+											.Text(LOCTEXT("Result", "Result"))
+											.ToolTipText(LOCTEXT("ResultTooltip", "The Result is the asset which will be returned if a row is selected (or other Chooser to evaluate to get the asset to return"))
+									]
+								]);
+			}
+
+			FName ColumnId("ChooserColumn", 1);
+			int NumColumns = Chooser->ColumnsStructs.Num();	
+			for(int ColumnIndex = 0; ColumnIndex < NumColumns; ColumnIndex++)
+			{
+				FChooserColumnBase& Column = Chooser->ColumnsStructs[ColumnIndex].GetMutable<FChooserColumnBase>();
+
+				TSharedPtr<SWidget> HeaderWidget = FObjectChooserWidgetFactories::CreateColumnWidget(&Column, Chooser->ColumnsStructs[ColumnIndex].GetScriptStruct(), Chooser->GetRootChooser(), -1);
+				if (!HeaderWidget.IsValid())
+				{
+					HeaderWidget = SNullWidget::NullWidget;
+				}
+				
+				HeaderRow->AddColumn(SHeaderRow::FColumn::FArguments()
+					.ColumnId(ColumnId)
+					.ManualWidth(Column.EditorColumnWidth)
+					.ManualWidth_Lambda([&Column]() { return Column.EditorColumnWidth; } )
+					.OnWidthChanged_Lambda([&Column](float NewWidth) { Column.EditorColumnWidth = NewWidth; })
+					.HeaderComboVisibility(EHeaderComboVisibility::Ghosted)
+					.HeaderContent()
+					[
+						SNew(SChooserColumnHandle)
+							.ViewModel(ChooserViewModel)
+							.TableHasFocus_Raw(this, &SChooserTableWidget::TableHasFocus)
+							.ColumnIndex(ColumnIndex)
+							.NoDropAfter(Chooser->ColumnsStructs[ColumnIndex].Get<FChooserColumnBase>().IsRandomizeColumn())
+						[
+							HeaderWidget.ToSharedRef()
+						]
+					
+					]);
+			
+				ColumnId.SetNumber(ColumnId.GetNumber() + 1);
+			}
+
+			HeaderRow->AddColumn( SHeaderRow::FColumn::FArguments()
+				.ColumnId("Add")
+				.FillWidth(1.0)
+				.HeaderContent( )
+				[
+					SNew(SVerticalBox)
+					 + SVerticalBox::Slot().AutoHeight()
+					 [
+						SNew(SHorizontalBox)
+						+ SHorizontalBox::Slot().MaxWidth(150)
+						[
+							CreateColumnComboButton.ToSharedRef()
+						]
+					]
+				]
+				);
+		}
+	}
+
+	void SChooserTableWidget::UpdateTableRows()
+	{
+		UChooserTable* Chooser = ChooserViewModel->GetChooser();
+		int32 NewNum = Chooser->ResultsStructs.Num();
+		Chooser->DisabledRows.SetNum(NewNum);
+    
+		// Sync the TableRows array which drives the ui table to match the number of results.
+		TableRows.SetNum(0, EAllowShrinking::No);
+		for(int i =0; i < NewNum; i++)
+		{
+			TableRows.Add(MakeShared<FChooserTableRow>(i));
+		}
+    
+		// Add one at the end, for the Fallback result
+		if (Chooser->FallbackResult.IsValid())
+		{
+			TableRows.Add(MakeShared<FChooserTableRow>(SChooserTableRow::SpecialIndex_Fallback));
+		}
+    	
+		// Add one at the end, for the "Add Row" control
+		TableRows.Add(MakeShared<FChooserTableRow>(SChooserTableRow::SpecialIndex_AddRow));
+    
+		// Make sure each column has the same number of row datas as there are results
+		for(FInstancedStruct& ColumnData : Chooser->ColumnsStructs)
+		{
+			FChooserColumnBase& Column = ColumnData.GetMutable<FChooserColumnBase>();
+			Column.SetNumRows(NewNum);
+		}
+    
+		if (TableView.IsValid())
+		{
+			TableView->RebuildList();
+		}
+	}
+			
+	TSharedPtr<SWidget> SChooserTableWidget::GenerateRowContextMenu()
+	{
+		UToolMenus* ToolMenus = UToolMenus::Get();
+		FToolMenuContext ToolMenuContext;
+		ToolMenuContext.AppendCommandList(CommandList);
+		return ToolMenus->GenerateWidget(FChooserTableEditor::ContextMenuName, ToolMenuContext);
+	}
+	
+	TSharedRef<ITableRow> SChooserTableWidget::GenerateTableRow(TSharedPtr<FChooserTableRow> InItem, const TSharedRef<STableViewBase>& OwnerTable)
+	{
+		UChooserTable* Chooser = ChooserViewModel->GetChooser();
+
+		return SNew(SChooserTableRow, OwnerTable)
+			.Entry(InItem).Chooser(Chooser).ViewModel(ChooserViewModel).TableHasFocus(this, &SChooserTableWidget::TableHasFocus);
+	}
+
+	void SChooserTableWidget::SelectRow(int32 RowIndex, bool bClear)
+    {
+    	if (TSharedPtr<FChooserTableRow>* Row = TableRows.FindByPredicate([RowIndex](const TSharedPtr<FChooserTableRow>& InRow)
+    		{
+    			return InRow->RowIndex == RowIndex;
+    		}))
+    	{
+    		if (!TableView->IsItemSelected(*Row))
+    		{
+    			if (bClear)
+    			{
+    				TableView->ClearSelection();
+    			}
+    			TableView->SetItemSelection(*Row, true, ESelectInfo::OnMouseClick);
+    		}
+    	}
+    }
+    	
+    void SChooserTableWidget::SelectRows(const TConstArrayView<int32> Rows)
+    {
+    	
+    	TableView->ClearSelection();
+    	TArray<TSharedPtr<FChooserTableRow>> RowsToSelect;
+    
+    	for (TSharedPtr<FChooserTableRow>& Row : TableRows)
+    	{
+    		if (Rows.Contains(Row->RowIndex))
+    		{
+    			RowsToSelect.Add(Row);
+    		}
+    	}
+    	
+    	TableView->SetItemSelection(RowsToSelect, true);
+    }
+	
+	
+	void SChooserTableWidget::Construct(const FArguments& InArgs)
+	{
+		ChooserViewModel = StaticCastSharedPtr<FChooserTableViewModel>(InArgs._ViewModel);
+		check(ChooserViewModel);
+
+		BindCommands(InArgs._Commands);
+
+		ChooserViewModel->RefreshTableRowsDelegate = FRefreshTable::CreateLambda([WeakPtr = AsWeak()]()
+		{
+			if (TSharedPtr<SWidget> Ptr = WeakPtr.Pin())
+			{
+				StaticCastSharedPtr<SChooserTableWidget>(Ptr)->UpdateTableRows();
+			}
+		});
+
+		ChooserViewModel->RefreshTableColumnsDelegate = FRefreshTable::CreateLambda([WeakPtr = AsWeak()]()
+		{
+			if (TSharedPtr<SWidget> Ptr = WeakPtr.Pin())
+			{
+				StaticCastSharedPtr<SChooserTableWidget>(Ptr)->UpdateTableColumns();
+			}
+		});
+
+		ChooserViewModel->SelectRowsDelegate = FSelectRows::CreateLambda([WeakPtr = AsWeak()](TConstArrayView<int32> Rows)
+		{
+			if (TSharedPtr<SWidget> Ptr = WeakPtr.Pin())
+			{
+				StaticCastSharedPtr<SChooserTableWidget>(Ptr)->SelectRows(Rows);
+			}
+		});
+
+		CreateColumnComboButton = SNew(SChooserCreateColumnButton).ViewModel(ChooserViewModel);
+		
+		HeaderRow = SNew(SHeaderRow);
+		UpdateTableColumns();
+		
+		TableView = SNew(SListView<TSharedPtr<FChooserTableRow>>)
+					.OnKeyDownHandler_Lambda([Commands = InArgs._Commands](const FGeometry&, const FKeyEvent& Event)
+					{
+						if (Commands)
+						{
+							if (Commands->ProcessCommandBindings(Event))
+							{
+								return FReply::Handled();
+							}
+						}
+						return FReply::Unhandled();
+					})
+					.ListItemsSource(&TableRows)
+					.OnContextMenuOpening(this, &SChooserTableWidget::GenerateRowContextMenu)
+					.OnGenerateRow(this, &SChooserTableWidget::GenerateTableRow)
+					.HeaderRow(HeaderRow)
+					.OnSelectionChanged_Lambda([this](TSharedPtr<FChooserTableRow>,  ESelectInfo::Type SelectInfo)
+    				{
+						TArray<TSharedPtr<FChooserTableRow>> SelectedItems;
+						GetSelectedItems(SelectedItems);
+						TArray<int32> SelectedRows;
+						for (TSharedPtr<FChooserTableRow> RowItem : SelectedItems)
+						{
+							SelectedRows.Add(RowItem->RowIndex);
+						}
+						ChooserViewModel->SetSelectedRows(SelectedRows);
+    				});
+
+		ChildSlot
+		[
+			SNew(SScrollBox).Orientation(Orient_Horizontal)
+			+ SScrollBox::Slot()
+			[
+				TableView.ToSharedRef()
+			]
+		];
+	}
+
+	SChooserTableWidget::~SChooserTableWidget()
+	{
+	}
+
+
+	bool SChooserTableWidget::HasSelection() const
+	{
+		return ChooserViewModel->HasSelection();
+	}
+
+	bool SChooserTableWidget::HasRowsSelected() const
+	{
+		return ChooserViewModel->HasRowsSelected();
+	}
+
+	bool SChooserTableWidget::HasColumnSelected() const
+	{
+		return ChooserViewModel->HasColumnSelected();
+	}
+
+	void SChooserTableWidget::SelectRootProperties()
+	{
+		ChooserViewModel->SelectRootProperties();
+	}
+
+	bool SChooserTableWidget::CanAutoPopulateSelection() const 
+	{
+		if (HasColumnSelected())
+		{
+			return ChooserViewModel->CanAutoPopulateColumn(GetSelectedColumn());
+		}
+		else if (HasRowsSelected())
+		{
+			return ChooserViewModel->CanAutoPopulateRows();	
+		}
+		return false;
+	}
+
+	void SChooserTableWidget::AutoPopulateSelection()
+	{
+		if (HasColumnSelected())
+		{
+			ChooserViewModel->AutoPopulateColumn(GetSelectedColumn());
+		}
+		else if (HasRowsSelected())
+		{
+			ChooserViewModel->AutoPopulateRows(ChooserViewModel->GetSelectedRows());
+		}
+	}
+
+	void SChooserTableWidget::AutoPopulateAll()
+	{
+		ChooserViewModel->AutoPopulateAll();
+	}
+	
+	void SChooserTableWidget::RemoveDisabledData()
+	{
+		ChooserViewModel->RemoveDisabledData();
+	}
+
+
+	void SChooserTableWidget::DeleteSelection()
+	{
+		if (HasColumnSelected())
+		{
+			ChooserViewModel->DeleteColumn(GetSelectedColumn());
+			SelectRootProperties();
+		}
+		else if (HasRowsSelected())
+		{
+			ChooserViewModel->DeleteRows(ChooserViewModel->GetSelectedRows());
+		}
+	}
+
+	void SChooserTableWidget::DuplicateSelection()
+	{
+		if (HasRowsSelected())
+		{
+			ChooserViewModel->DuplicateRows(ChooserViewModel->GetSelectedRows());
+		}
+		else if(HasColumnSelected())
+		{
+			ChooserViewModel->DuplicateColumn(GetSelectedColumn());
+		}
+	}
+	
+	bool SChooserTableWidget::IsSelectionDisabled() const
+	{
+		if (HasRowsSelected())
+		{
+			return ChooserViewModel->AreRowsDisabled(ChooserViewModel->GetSelectedRows());
+		}
+		else if(HasColumnSelected())
+		{
+			return ChooserViewModel->IsColumnDisabled(GetSelectedColumn());
+		}
+		return false;
+	}
+
+	void SChooserTableWidget::ToggleDisableSelection()
+	{
+		if (HasRowsSelected())
+		{
+			ChooserViewModel->ToggleDisableRows(ChooserViewModel->GetSelectedRows());
+		}
+		else if(HasColumnSelected())
+		{
+			ChooserViewModel->ToggleDisableColumn(GetSelectedColumn());
+		}
+	}
+	
+	void SChooserTableWidget::CopySelection()
+	{
+		if (HasRowsSelected())
+		{
+			ChooserViewModel->CopyRows(ChooserViewModel->GetSelectedRows());
+		}
+		else if(HasColumnSelected())
+		{
+			ChooserViewModel->CopyColumn(GetSelectedColumn());
+		}
+	}
+
+	void SChooserTableWidget::CutSelection()
+	{
+		if (HasRowsSelected())
+		{
+			TConstArrayView<int32> SelectedRows = ChooserViewModel->GetSelectedRows();
+			ChooserViewModel->CopyRows(SelectedRows);
+			ChooserViewModel->DeleteRows(SelectedRows);
+		}
+		else if(HasColumnSelected())
+		{
+			int ColumnIndex = GetSelectedColumn();
+			ChooserViewModel->CopyColumn(ColumnIndex);
+			ChooserViewModel->DeleteColumn(ColumnIndex);
+		}
+	}
+
+	bool SChooserTableWidget::CanPaste() const
+	{
+		return ChooserViewModel->CanPaste();
+	}
+
+	void SChooserTableWidget::Paste()
+	{
+		ChooserViewModel->Paste();
+	}
+
+	
+	bool SChooserTableWidget::CanMoveRowsUp()
+	{
+		return ChooserViewModel->CanMoveRowsUp(ChooserViewModel->GetSelectedRows());
+	}
+
+	void SChooserTableWidget::MoveRowsUp()
+	{
+		if (HasRowsSelected())
+		{
+			TConstArrayView<int32> SelectedRows = ChooserViewModel->GetSelectedRows();
+
+			int MinSelectedRow = ChooserViewModel->GetNumRows();
+			for(int SelectedRow : SelectedRows)
+			{
+				if (SelectedRow != ColumnWidget_SpecialIndex_Fallback)
+				{
+					MinSelectedRow = FMath::Min(SelectedRow, MinSelectedRow);
+				}
+			}
+			ChooserViewModel->MoveRows(SelectedRows, MinSelectedRow - 1);
+		}
+	}
+
+	bool SChooserTableWidget::CanMoveRowsDown()
+	{
+		return ChooserViewModel->CanMoveRowsDown(ChooserViewModel->GetSelectedRows());
+	}
+
+	void SChooserTableWidget::MoveRowsDown()
+	{
+		if (HasRowsSelected())
+		{
+			TConstArrayView<int32> SelectedRows =	ChooserViewModel->GetSelectedRows();
+
+			int MaxSelectedRow = -1;
+			for(int SelectedRow : SelectedRows)
+			{
+				MaxSelectedRow = FMath::Max(SelectedRow, MaxSelectedRow);
+			}
+			ChooserViewModel->MoveRows(SelectedRows, MaxSelectedRow + 2);
+		}
+	}
+
+	bool SChooserTableWidget::CanMoveColumnLeft()
+	{
+		return ChooserViewModel->CanMoveColumnLeft(GetSelectedColumn());
+	}
+
+	void SChooserTableWidget::MoveColumnLeft()
+	{
+		ChooserViewModel->MoveColumnLeft(GetSelectedColumn());
+	}
+
+	bool SChooserTableWidget::CanMoveColumnRight()
+	{
+		return ChooserViewModel->CanMoveColumnRight(GetSelectedColumn());
+	}
+
+	void SChooserTableWidget::MoveColumnRight()
+	{
+		ChooserViewModel->MoveColumnRight(GetSelectedColumn());
+	}
+
+	void SChooserTableWidget::BindCommands(TSharedPtr<FUICommandList> InCommandList)
+	{
+		CommandList = InCommandList;
+
+		const FChooserTableEditorCommands& Commands = FChooserTableEditorCommands::Get();
+
+		CommandList->MapAction(
+			Commands.EditChooserSettings,
+			FExecuteAction::CreateSP(SharedThis(this), &SChooserTableWidget::SelectRootProperties));
+		
+		CommandList->MapAction(
+			Commands.AutoPopulateAll,
+			FExecuteAction::CreateSP(SharedThis(this), &SChooserTableWidget::AutoPopulateAll));
+		
+		CommandList->MapAction(
+			Commands.RemoveDisabledData,
+			FExecuteAction::CreateSP(SharedThis(this), &SChooserTableWidget::RemoveDisabledData));
+		
+		CommandList->MapAction(
+			FGenericCommands::Get().Delete,
+			FExecuteAction::CreateSP(SharedThis(this), &SChooserTableWidget::DeleteSelection),
+			FCanExecuteAction::CreateSP(SharedThis(this), &SChooserTableWidget::HasSelection)
+			);
+
+		CommandList->MapAction(
+			FGenericCommands::Get().Duplicate,
+			FExecuteAction::CreateSP(SharedThis(this), &SChooserTableWidget::DuplicateSelection),
+			FCanExecuteAction::CreateSP(SharedThis(this), &SChooserTableWidget::HasSelection)
+			);
+
+		CommandList->MapAction(
+			Commands.AutoPopulateSelection,
+			FExecuteAction::CreateSP(SharedThis(this), &SChooserTableWidget::AutoPopulateSelection),
+			FCanExecuteAction::CreateSP(SharedThis(this), &SChooserTableWidget::CanAutoPopulateSelection),
+			FGetActionCheckState(),FIsActionButtonVisible::CreateSP(SharedThis(this), &SChooserTableWidget::HasSelection)
+			);
+		
+		CommandList->MapAction(
+			Commands.Disable,
+			FExecuteAction::CreateSP(SharedThis(this), &SChooserTableWidget::ToggleDisableSelection),
+			FCanExecuteAction::CreateSP(SharedThis(this), &SChooserTableWidget::HasSelection),
+			FIsActionChecked::CreateSP(SharedThis(this), &SChooserTableWidget::IsSelectionDisabled),
+			FIsActionButtonVisible::CreateSP(SharedThis(this), &SChooserTableWidget::HasSelection)
+			);
+
+		CommandList->MapAction(
+			FGenericCommands::Get().Copy,
+			FExecuteAction::CreateSP(SharedThis(this), &SChooserTableWidget::CopySelection),
+			FCanExecuteAction::CreateSP(SharedThis(this), &SChooserTableWidget::HasSelection));
+		
+		CommandList->MapAction(
+			FGenericCommands::Get().Cut,
+			FExecuteAction::CreateSP(SharedThis(this), &SChooserTableWidget::CutSelection),
+			FCanExecuteAction::CreateSP(SharedThis(this), &SChooserTableWidget::HasSelection));
+		
+		CommandList->MapAction(
+			FGenericCommands::Get().Paste,
+			FExecuteAction::CreateSP(SharedThis(this), &SChooserTableWidget::Paste),
+			FCanExecuteAction::CreateSP(SharedThis(this), &SChooserTableWidget::CanPaste));
+		
+		CommandList->MapAction(
+			Commands.MoveUp,
+			FExecuteAction::CreateSP(SharedThis(this), &SChooserTableWidget::MoveRowsUp),
+			FCanExecuteAction::CreateSP(SharedThis(this), &SChooserTableWidget::CanMoveRowsUp),
+			FIsActionChecked(),
+			FIsActionButtonVisible::CreateSP(SharedThis(this), &SChooserTableWidget::HasRowsSelected)
+			);
+		
+		CommandList->MapAction(
+				Commands.MoveDown,
+				FExecuteAction::CreateSP(SharedThis(this), &SChooserTableWidget::MoveRowsDown),
+				FCanExecuteAction::CreateSP(SharedThis(this), &SChooserTableWidget::CanMoveRowsDown),
+				FIsActionChecked(),
+				FIsActionButtonVisible::CreateSP(SharedThis(this), &SChooserTableWidget::HasRowsSelected)
+				);
+		
+		CommandList->MapAction(
+				Commands.MoveLeft,
+				FExecuteAction::CreateSP(SharedThis(this), &SChooserTableWidget::MoveColumnLeft),
+				FCanExecuteAction::CreateSP(SharedThis(this), &SChooserTableWidget::CanMoveColumnLeft),
+				FIsActionChecked(),
+				FIsActionButtonVisible::CreateSP(SharedThis(this), &SChooserTableWidget::HasColumnSelected)
+				);
+
+		CommandList->MapAction(
+				Commands.MoveRight,
+				FExecuteAction::CreateSP(SharedThis(this), &SChooserTableWidget::MoveColumnRight),
+				FCanExecuteAction::CreateSP(SharedThis(this), &SChooserTableWidget::CanMoveColumnRight),
+				FIsActionChecked(),
+				FIsActionButtonVisible::CreateSP(SharedThis(this), &SChooserTableWidget::HasColumnSelected)
+				);
+	}
+
+}
+
+#undef LOCTEXT_NAMESPACE
diff --git a/Engine/Plugins/Chooser/Source/ChooserEditor/Private/SChooserTableWidget.h b/Engine/Plugins/Chooser/Source/ChooserEditor/Private/SChooserTableWidget.h
new file mode 100644
--- /dev/null
+++ b/Engine/Plugins/Chooser/Source/ChooserEditor/Private/SChooserTableWidget.h
@@ -0,0 +1,88 @@
+// Copyright Epic Games, Inc. All Rights Reserved.
+#pragma once
+
+#include "Widgets/SCompoundWidget.h"
+#include "Widgets/Views/SListView.h"
+#include "Framework/Commands/UICommandList.h"
+#include "ChooserTableViewModel.h"
+
+class SPositiveActionButton;
+
+namespace UE::ChooserEditor
+{
+	
+class SChooserTableWidget : public SCompoundWidget
+{
+public:
+	SLATE_BEGIN_ARGS(SChooserTableWidget)
+	{}
+	SLATE_ARGUMENT(TSharedPtr<FUICommandList>, Commands)
+	SLATE_ARGUMENT(TSharedPtr<IChooserTableViewModel>, ViewModel)
+	SLATE_END_ARGS()
+
+	void Construct(const FArguments& InArgs);
+	virtual ~SChooserTableWidget() override;
+	
+	void SelectRows(const TConstArrayView<int32> Rows);
+	void SelectRow(int32 RowIndex, bool bClear = true);
+
+	void UpdateTableColumns();
+	void UpdateTableRows();
+
+	void ClearSelection();
+	void GetSelectedItems(TArray<TSharedPtr<FChooserTableRow>> &SelectedItems);
+
+	int GetSelectedColumn() const;
+
+	bool TableHasFocus() const;
+	
+	// Command handler functions
+	
+	bool HasSelection() const;
+	bool HasRowsSelected() const;
+	bool HasColumnSelected() const;
+	
+	void SelectRootProperties();
+	
+	bool CanAutoPopulateSelection() const;
+	void AutoPopulateSelection();
+	void AutoPopulateAll();
+	
+	void RemoveDisabledData();
+	void DeleteSelection();
+	void DuplicateSelection();
+	
+	bool IsSelectionDisabled() const;
+	void ToggleDisableSelection();
+	
+	void CopySelection();
+	void CutSelection();
+	bool CanPaste() const;
+	void Paste();
+	
+	bool CanMoveRowsUp();
+	void MoveRowsUp();
+	bool CanMoveRowsDown();
+	void MoveRowsDown();
+	bool CanMoveColumnLeft();
+	void MoveColumnLeft();
+	bool CanMoveColumnRight();
+	void MoveColumnRight();
+	
+ private:
+	void BindCommands(TSharedPtr<FUICommandList> CommandList);
+
+	TSharedPtr<SListView<TSharedPtr<FChooserTableRow>>> TableView;
+	TSharedPtr<SHeaderRow> HeaderRow;
+	
+	TSharedPtr<SWidget> CreateColumnComboButton;
+	TSharedPtr<SWidget> GenerateRowContextMenu();
+	TSharedRef<ITableRow> GenerateTableRow(TSharedPtr<FChooserTableRow> InItem, const TSharedRef<STableViewBase>& OwnerTable);
+
+	TSharedPtr<FUICommandList> CommandList;
+
+	TSharedPtr<FChooserTableViewModel> ChooserViewModel;
+	TArray<TSharedPtr<FChooserTableRow>> TableRows;
+};
+
+}
diff --git a/Engine/Plugins/Chooser/Source/ChooserEditor/Public/ChooserEditorStyle.h b/Engine/Plugins/Chooser/Source/ChooserEditor/Public/ChooserEditorStyle.h
new file mode 100644
--- /dev/null
+++ b/Engine/Plugins/Chooser/Source/ChooserEditor/Public/ChooserEditorStyle.h
@@ -0,0 +1,24 @@
+// Copyright Epic Games, Inc. All Rights Reserved.
+#pragma once
+
+#include "CoreMinimal.h"
+#include "Styling/SlateStyle.h"
+
+namespace UE::ChooserEditor
+{
+	
+class FChooserEditorStyle : FSlateStyleSet
+{
+public:
+    FChooserEditorStyle();
+	static void Initialize();
+	static void Shutdown();
+
+	CHOOSEREDITOR_API static const ISlateStyle& Get();
+
+private:
+
+	static TSharedPtr< FChooserEditorStyle > StyleInstance;
+};
+
+}
diff --git a/Engine/Plugins/Chooser/Source/ChooserEditor/Public/IChooserEditorModule.h b/Engine/Plugins/Chooser/Source/ChooserEditor/Public/IChooserEditorModule.h
new file mode 100644
--- /dev/null
+++ b/Engine/Plugins/Chooser/Source/ChooserEditor/Public/IChooserEditorModule.h
@@ -0,0 +1,22 @@
+// Copyright Epic Games, Inc. All Rights Reserved.
+
+#pragma once
+
+#include "CoreMinimal.h"
+#include "Modules/ModuleManager.h"
+#include "IChooserTableViewModel.h"
+
+class UChooserTable;
+class SWidget;
+
+namespace UE::ChooserEditor
+{
+
+class IChooserEditorModule : public IModuleInterface
+{
+public:
+	virtual TSharedPtr<IChooserTableViewModel> CreateChooserTableViewModel(UChooserTable* ChooserTable) = 0;
+	virtual TSharedPtr<SWidget> CreateChooserTableView(TSharedPtr<IChooserTableViewModel> ViewModel, TSharedPtr<FUICommandList> CommandList) = 0;
+};
+
+}
diff --git a/Engine/Plugins/Chooser/Source/ChooserEditor/Public/IChooserTableViewModel.h b/Engine/Plugins/Chooser/Source/ChooserEditor/Public/IChooserTableViewModel.h
new file mode 100644
--- /dev/null
+++ b/Engine/Plugins/Chooser/Source/ChooserEditor/Public/IChooserTableViewModel.h
@@ -0,0 +1,36 @@
+// Copyright Epic Games, Inc. All Rights Reserved.
+
+#pragma once
+
+#include "CoreMinimal.h"
+#include "Modules/ModuleManager.h"
+#include "Misc/NotifyHook.h"
+
+class FUICommandList;
+class UChooserTable;
+
+namespace UE::ChooserEditor
+{
+
+DECLARE_DELEGATE_OneParam(FShowDetails, const TArray<UObject*>&);
+
+class IChooserTableViewModel : public TSharedFromThis<IChooserTableViewModel>, public FNotifyHook
+{
+public:
+	virtual ~IChooserTableViewModel() {}
+
+	virtual void SetShowDetailsDelegate(FShowDetails InShowDetailsDelegate) = 0;
+	virtual void RegisterMenus(TSharedPtr<FUICommandList> CommandList) = 0;
+
+	virtual const UChooserTable* GetRootChooser() const = 0;
+	virtual UChooserTable* GetRootChooser() = 0;
+	virtual UChooserTable* GetChooser() = 0;
+	virtual const UChooserTable* GetChooser() const = 0;
+	virtual void SetChooser(UChooserTable* Chooser) = 0;
+
+	virtual void RefreshAll() = 0;
+	virtual void AutoPopulateAll() = 0;
+	virtual void SelectRootProperties() = 0;
+};
+
+}
