Fix for allowing binding to root context struct in chooser columns
- it is now allowed for struct output columns, if the context entry is a struct, and for other types of struct columns (eg Gameplay Tags) if the context entry is the same type as the column requires

diff --git a/Engine/Plugins/Chooser/Source/ChooserEditor/Private/SPropertyAccessChainWidget.cpp b/Engine/Plugins/Chooser/Source/ChooserEditor/Private/SPropertyAccessChainWidget.cpp
--- a/Engine/Plugins/Chooser/Source/ChooserEditor/Private/SPropertyAccessChainWidget.cpp
+++ b/Engine/Plugins/Chooser/Source/ChooserEditor/Private/SPropertyAccessChainWidget.cpp
@@ -153,6 +153,8 @@ void CreatePropertyAccessMenus(const FString& TypeFilter, FOnAddBinding OnAddBin
 	{
 		const FInstancedStruct &ContextStruct = ContextData[ContextIndex];
 		
+		bool bAllowBindRoot = false; 
+		
 		UStruct* Struct = nullptr;
 		if (const FContextObjectTypeClass* ClassType = ContextStruct.GetPtr<FContextObjectTypeClass>())
 		{
@@ -161,13 +163,13 @@ void CreatePropertyAccessMenus(const FString& TypeFilter, FOnAddBinding OnAddBin
 		else if (const FContextObjectTypeStruct* StructType = ContextStruct.GetPtr<FContextObjectTypeStruct>())
 		{
 			Struct = StructType->Struct;
+			bAllowBindRoot = TypeFilter == "struct" || (Struct->GetPrefixCPP() + Struct->GetName()) == TypeFilter;
 		}
 		else
 		{
 			continue;
 		}
 
-		bool bAllowBindRoot = TypeFilter == "struct";
 
 		// Show struct properties.
 		MenuBuilder.AddSubMenu(
