Fix Nested Choosers list not updating after Undo/Redo

diff --git a/Engine/Plugins/Chooser/Source/ChooserEditor/Private/ChooserTableViewModel.cpp b/Engine/Plugins/Chooser/Source/ChooserEditor/Private/ChooserTableViewModel.cpp
--- a/Engine/Plugins/Chooser/Source/ChooserEditor/Private/ChooserTableViewModel.cpp
+++ b/Engine/Plugins/Chooser/Source/ChooserEditor/Private/ChooserTableViewModel.cpp
@@ -256,14 +256,14 @@ void FChooserTableViewModel::PostUndo(bool bSuccess)
 {
 	SelectRootProperties();
 	RefreshAll();
-	//RefreshNestedObjectTree(); // todo: fix this
+	RootChooser->NestedObjectsChanged.Broadcast();
 }
 
 void FChooserTableViewModel::PostRedo(bool bSuccess)
 {
 	SelectRootProperties();
 	RefreshAll();
-	// RefreshNestedObjectTree(); todo fix this
+	RootChooser->NestedObjectsChanged.Broadcast();
 } 
 
 
