Fix chooser editor scroll wheel behavior

#Synced-CL 48625956

diff --git a/Engine/Plugins/Chooser/Source/ChooserEditor/Private/ChooserTableEditor.cpp b/Engine/Plugins/Chooser/Source/ChooserEditor/Private/ChooserTableEditor.cpp
--- a/Engine/Plugins/Chooser/Source/ChooserEditor/Private/ChooserTableEditor.cpp
+++ b/Engine/Plugins/Chooser/Source/ChooserEditor/Private/ChooserTableEditor.cpp
@@ -597,11 +597,7 @@ TSharedRef<SDockTab> FChooserTableEditor::SpawnTableTab( const FSpawnTabArgs& Ar
 			]
 			+ SVerticalBox::Slot().FillHeight(1)
 			[
-				SNew(SScrollBox).Orientation(Orient_Horizontal)
-				+ SScrollBox::Slot()
-				[
 					ChooserTableView
-				]
 			]
 		];
 }
diff --git a/Engine/Plugins/Chooser/Source/ChooserEditor/Private/SChooserTableWidget.cpp b/Engine/Plugins/Chooser/Source/ChooserEditor/Private/SChooserTableWidget.cpp
--- a/Engine/Plugins/Chooser/Source/ChooserEditor/Private/SChooserTableWidget.cpp
+++ b/Engine/Plugins/Chooser/Source/ChooserEditor/Private/SChooserTableWidget.cpp
@@ -257,6 +257,8 @@ namespace UE::ChooserEditor
 		
 		HeaderRow = SNew(SHeaderRow);
 		UpdateTableColumns();
+
+		TSharedRef VerticalScrollbar = SNew(SScrollBar);
 		
 		TableView = SNew(SListView<TSharedPtr<FChooserTableRow>>)
 					.OnKeyDownHandler_Lambda([Commands = InArgs._Commands](const FGeometry&, const FKeyEvent& Event)
@@ -274,6 +276,8 @@ namespace UE::ChooserEditor
 					.OnContextMenuOpening(this, &SChooserTableWidget::GenerateRowContextMenu)
 					.OnGenerateRow(this, &SChooserTableWidget::GenerateTableRow)
 					.HeaderRow(HeaderRow)
+					.ConsumeMouseWheel(EConsumeMouseWheel::Always)
+					.ExternalScrollbar(VerticalScrollbar)
 					.OnSelectionChanged_Lambda([this](TSharedPtr<FChooserTableRow>,  ESelectInfo::Type SelectInfo)
     				{
 						TArray<TSharedPtr<FChooserTableRow>> SelectedItems;
@@ -288,10 +292,18 @@ namespace UE::ChooserEditor
 
 		ChildSlot
 		[
-			SNew(SScrollBox).Orientation(Orient_Horizontal)
-			+ SScrollBox::Slot()
+			SNew(SHorizontalBox)
+			+SHorizontalBox::Slot()
+			[
+				SNew(SScrollBox).Orientation(Orient_Horizontal)//.ConsumeMouseWheel(EConsumeMouseWheel::Never)
+					+ SScrollBox::Slot()
+				[
+						TableView.ToSharedRef()
+				]
+			]
+			+ SHorizontalBox::Slot().AutoWidth()
 			[
-				TableView.ToSharedRef()
+				VerticalScrollbar
 			]
 		];
 	}
