Fixes for chooser versioning

#Synced-CL 48768753

diff --git a/Engine/Plugins/Chooser/Source/Chooser/Private/Chooser.cpp b/Engine/Plugins/Chooser/Source/Chooser/Private/Chooser.cpp
--- a/Engine/Plugins/Chooser/Source/Chooser/Private/Chooser.cpp
+++ b/Engine/Plugins/Chooser/Source/Chooser/Private/Chooser.cpp
@@ -65,22 +65,26 @@ void UChooserTable::PostLoad()
 			}
 		}
 
-		if(RootChooser == nullptr && NestedObjects.IsEmpty())
+		if (RootChooser == nullptr && NestedChoosers.IsEmpty() && NestedObjects.IsEmpty())
 		{
-			// data upgrade for root tables: add elements to NestedObjects list
+			// data upgrade for root tables: add elements to NestedTables list
 			TArray<UObject*> ChildObjects;
 			GetObjectsWithOuter(GetPackage(), ChildObjects, true);
-			for(UObject* ChildObject : ChildObjects)
+			for (UObject* ChildObject : ChildObjects)
 			{
-				if(UChooserTable* Chooser = Cast<UChooserTable>(ChildObject))
+				if (UChooserTable* Chooser = Cast<UChooserTable>(ChildObject))
 				{
-					NestedObjects.Add(Chooser);
+					if (Chooser != this)
+					{
+						NestedChoosers.Add(Chooser);
+					}
 				}
 			}
 		}
-		Version = 2;
+		Version = 1;
 	}
-	if (Version < 2)
+	
+	if (Version < 2 && NestedObjects.IsEmpty())
 	{
 		for(UChooserTable* NestedChooser : NestedChoosers)
 		{
@@ -88,6 +92,12 @@ void UChooserTable::PostLoad()
 		}
 		Version = 2;
 	}
+
+	if (Version < 3)
+	{
+		NestedChoosers.Reset();
+		Version = 3;
+	}
 #endif
 
 	Compile();
diff --git a/Engine/Plugins/Chooser/Source/Chooser/Public/Chooser.h b/Engine/Plugins/Chooser/Source/Chooser/Public/Chooser.h
--- a/Engine/Plugins/Chooser/Source/Chooser/Public/Chooser.h
+++ b/Engine/Plugins/Chooser/Source/Chooser/Public/Chooser.h
@@ -180,7 +180,9 @@ public:
 	float EditorResultsColumnWidth = 300;
 
 	UPROPERTY()
-	uint32 Version = 2;
+	uint32 Version = 0;
+
+	static constexpr int CurrentVersion = 3;
 #endif
 
 	UFUNCTION()
diff --git a/Engine/Plugins/Chooser/Source/ChooserEditor/Private/ChooserFactory.cpp b/Engine/Plugins/Chooser/Source/ChooserEditor/Private/ChooserFactory.cpp
--- a/Engine/Plugins/Chooser/Source/ChooserEditor/Private/ChooserFactory.cpp
+++ b/Engine/Plugins/Chooser/Source/ChooserEditor/Private/ChooserFactory.cpp
@@ -45,5 +45,8 @@ UObject* UChooserTableFactory::FactoryCreateNew(UClass* Class, UObject* InParent
 	{
 		NewChooser = NewObject<UChooserTable>(InParent, Class, Name, Flags);
 	}
+	
+	NewChooser->Version = UChooserTable::CurrentVersion;
+	
 	return NewChooser;
 }
\ No newline at end of file
