Fix details panel chooser property bindings in UAF choosers

#Synced-CL 48881602

diff --git a/Engine/Plugins/Chooser/Source/ChooserEditor/Private/SPropertyAccessChainWidget.cpp b/Engine/Plugins/Chooser/Source/ChooserEditor/Private/SPropertyAccessChainWidget.cpp
--- a/Engine/Plugins/Chooser/Source/ChooserEditor/Private/SPropertyAccessChainWidget.cpp
+++ b/Engine/Plugins/Chooser/Source/ChooserEditor/Private/SPropertyAccessChainWidget.cpp
@@ -206,32 +206,51 @@ void CreatePropertyAccessMenus(const FString& TypeFilter, FOnAddBinding OnAddBin
 	
 void SPropertyAccessChainWidget::SetPropertyBinding(IHasContextClass* HasContext, const TArray<FBindingChainElement>& InBindingChain, FChooserPropertyBinding& OutPropertyBinding)
 {
+	if (!HasContext)
+	{
+		return;
+	}
+	
+	TConstArrayView<FInstancedStruct> ContextData = HasContext->GetContextData();
 	Chooser::CopyPropertyChain(InBindingChain, OutPropertyBinding);
+
+	// remap selected context index from the UI index (which skips over non object/struct context entries), to an actual index in to the context data array
+	uint32 UIIndex = 0;
+	for (int32 ContextIndex=0; ContextIndex<ContextData.Num(); ContextIndex++)
+	{
+		const FInstancedStruct& ContextStruct = ContextData[ContextIndex];
+		
+		if (ContextStruct.GetPtr<FContextObjectTypeClass>() || ContextStruct.GetPtr<FContextObjectTypeStruct>())
+		{
+			if (UIIndex == OutPropertyBinding.ContextIndex)
+			{
+				OutPropertyBinding.ContextIndex = ContextIndex;
+				break;
+			}
+			UIIndex++;
+		}
+	}	
 	
 	FField* Property = InBindingChain.Last().Field.ToField();
 		
 	OutPropertyBinding.DisplayName = "";
 	if (InBindingChain.Num() == 1)
 	{
-		// direct binding to a context struct/class, set display name to the struct/clase name
-		if (HasContext)
+		// direct binding to a context struct/class, set display name to the struct/class name
+		if (ContextData.IsValidIndex(OutPropertyBinding.ContextIndex))
 		{
-			TConstArrayView<FInstancedStruct> ContextData = HasContext->GetContextData();
-			if (ContextData.IsValidIndex(InBindingChain[0].ArrayIndex))
+			if (const FContextObjectTypeStruct* StructContext = ContextData[OutPropertyBinding.ContextIndex].GetPtr<FContextObjectTypeStruct>())
 			{
-				if (const FContextObjectTypeStruct* StructContext = ContextData[InBindingChain[0].ArrayIndex].GetPtr<FContextObjectTypeStruct>())
+				if (StructContext->Struct)
 				{
-					if (StructContext->Struct)
-					{
-						OutPropertyBinding.DisplayName = StructContext->Struct->GetDisplayNameText().ToString();
-					}
+					OutPropertyBinding.DisplayName = StructContext->Struct->GetDisplayNameText().ToString();
 				}
-				else if (const FContextObjectTypeClass* ClassContext = ContextData[InBindingChain[0].ArrayIndex].GetPtr<FContextObjectTypeClass>())
+			}
+			else if (const FContextObjectTypeClass* ClassContext = ContextData[OutPropertyBinding.ContextIndex].GetPtr<FContextObjectTypeClass>())
+			{
+				if (ClassContext->Class)
 				{
-					if (ClassContext->Class)
-					{
-						OutPropertyBinding.DisplayName = ClassContext->Class->GetDisplayNameText().ToString();
-					}
+					OutPropertyBinding.DisplayName = ClassContext->Class->GetDisplayNameText().ToString();
 				}
 			}
 		}
