UE: Chooser now populates cooked data at save time, even when not cooking, so that tools and packaging code can know which references are non-editor-only

#Synced-CL 48838269

diff --git a/Engine/Plugins/Chooser/Source/Chooser/Private/Chooser.cpp b/Engine/Plugins/Chooser/Source/Chooser/Private/Chooser.cpp
--- a/Engine/Plugins/Chooser/Source/Chooser/Private/Chooser.cpp
+++ b/Engine/Plugins/Chooser/Source/Chooser/Private/Chooser.cpp
@@ -207,9 +207,9 @@ void UChooserTable::Compile(bool bForce)
 void UChooserTable::Serialize(FArchive& Ar)
 {
 #if WITH_EDITORONLY_DATA
-	if (Ar.IsCooking())
+	if (Ar.IsSaving())
 	{
-		CookData();
+		PopulateCookedData(Ar.IsCooking());
 	}
 #endif
 	
@@ -258,6 +258,8 @@ void UChooserTable::Serialize(FArchive& Ar)
 			Results_DEPRECATED.SetNum(0);
 			Columns_DEPRECATED.SetNum(0);
 		}
+
+		bIsCookedPackage = Ar.IsLoadingFromCookedPackage();
 	}
 #endif
 }
@@ -294,12 +296,28 @@ void UChooserTable::RemoveDisabledData()
 	}
 }
 
-void UChooserTable::CookData()
+void UChooserTable::PopulateCookedData(bool bCooking)
 {
-	RemoveDisabledData();
+	if (bCooking)
+	{
+		RemoveDisabledData();
 
-	// copy stripped results struct into CookedResults array
-	CookedResults = ResultsStructs;
+		// copy stripped results struct into CookedResults array
+		CookedResults = ResultsStructs;
+	}
+	else
+	{
+		// Build CookedResults array from stripped results struct.
+		// When we aren't cooking, this is used to make sure the actually cooked rows show up as non-editor-only in tools and packaging code.
+		CookedResults.Empty();
+		for (int32 ResultIndex = 0; ResultIndex < ResultsStructs.Num(); ++ResultIndex)
+		{
+			if (!IsRowDisabled(ResultIndex))
+			{
+				CookedResults.Add(ResultsStructs[ResultIndex]);
+			}
+		}
+	}
 }
 
 
@@ -919,7 +937,7 @@ namespace {
 	{
 		for (TObjectIterator<UChooserTable> It; It; ++It)
 		{
-			It->CookData();
+			It->PopulateCookedData(true);
 		}
 	}
 }
diff --git a/Engine/Plugins/Chooser/Source/Chooser/Public/Chooser.h b/Engine/Plugins/Chooser/Source/Chooser/Public/Chooser.h
--- a/Engine/Plugins/Chooser/Source/Chooser/Public/Chooser.h
+++ b/Engine/Plugins/Chooser/Source/Chooser/Public/Chooser.h
@@ -32,14 +32,14 @@ public:
 
 #if WITH_EDITORONLY_DATA
 	CHOOSER_API void RemoveDisabledData();
-	CHOOSER_API void CookData();
+	CHOOSER_API void PopulateCookedData(bool bCooking);
 	CHOOSER_API void DeleteNestedChooser(UChooserTable* NestedChooser);
 #endif
 
 	bool IsRowDisabled(int32 RowIndex) const
 	{
 #if WITH_EDITORONLY_DATA
-		return CookedResults.IsEmpty() && DisabledRows.IsValidIndex(RowIndex) && DisabledRows[RowIndex];
+		return !IsCookedData() && DisabledRows.IsValidIndex(RowIndex) && DisabledRows[RowIndex];
 #else
 		return false;
 #endif
@@ -183,6 +183,11 @@ public:
 	uint32 Version = 0;
 
 	static constexpr int CurrentVersion = 3;
+
+private:
+	bool bIsCookedPackage = false;
+
+public:
 #endif
 
 	UFUNCTION()
@@ -195,7 +200,14 @@ public:
 	UPROPERTY(EditAnywhere, NoClear, DisplayName = "Columns", Meta = (ExcludeBaseStruct, BaseStruct = "/Script/Chooser.ChooserColumnBase"), Category = "Hidden")
 	TArray<FInstancedStruct> ColumnsStructs;
 	
-	bool IsCookedData() const { return !CookedResults.IsEmpty(); }
+	bool IsCookedData() const
+	{
+#if WITH_EDITORONLY_DATA
+		return bIsCookedPackage;
+#else
+		return true;
+#endif
+	}
 	static CHOOSER_API FObjectChooserBase::EIteratorStatus EvaluateChooser(FChooserEvaluationContext& Context, const UChooserTable* Chooser, FObjectChooserBase::FObjectChooserIteratorCallback Callback);
 	static CHOOSER_API FObjectChooserBase::EIteratorStatus EvaluateChooser(FChooserEvaluationContext& Context, const UChooserTable* Chooser, FObjectChooserBase::FObjectChooserSoftObjectIteratorCallback Callback);
 	static CHOOSER_API FObjectChooserBase::EIteratorStatus IterateChooser(FChooserEvaluationContext& Context, const UChooserTable* Chooser, FObjectChooserBase::FObjectChooserIteratorCallback Callback);
