Make sure to call SetValues in FOutputObjectColumn::SetOutputs even if the value is null to reset context values when iterating over multiple rows

#Synced-CL 49097988

diff --git a/Engine/Plugins/Chooser/Source/Chooser/Private/OutputObjectColumn.cpp b/Engine/Plugins/Chooser/Source/Chooser/Private/OutputObjectColumn.cpp
--- a/Engine/Plugins/Chooser/Source/Chooser/Private/OutputObjectColumn.cpp
+++ b/Engine/Plugins/Chooser/Source/Chooser/Private/OutputObjectColumn.cpp
@@ -38,11 +38,9 @@ void FOutputObjectColumn::SetOutputs(FChooserEvaluationContext& Context, int Row
 	{
 		if (RowIndex == ChooserColumn_SpecialIndex_Fallback || RowValues.IsValidIndex(RowIndex))
 		{
-			if (const FObjectChooserBase* Value = GetValueForIndex(RowIndex).Value.GetPtr<FObjectChooserBase>())
-			{
-				UObject* Result = Value->ChooseObject(Context);
-				InputValue.Get<FChooserParameterObjectBase>().SetValue(Context, Result);
-			}
+			const FObjectChooserBase* Value = GetValueForIndex(RowIndex).Value.GetPtr<FObjectChooserBase>();
+			UObject* Result = Value ? Value->ChooseObject(Context) : nullptr;
+			InputValue.Get<FChooserParameterObjectBase>().SetValue(Context, Result);
 		}
 		else
 		{
