[Corgi] add a way to provide random stream to ChooserTable to make the search more deterministic. The last missing part for determinism is GetRepeatProbabilityMultiplier

#Synced-CL 49371773

diff --git a/Engine/Plugins/Chooser/Source/Chooser/Private/RandomizeColumn.cpp b/Engine/Plugins/Chooser/Source/Chooser/Private/RandomizeColumn.cpp
--- a/Engine/Plugins/Chooser/Source/Chooser/Private/RandomizeColumn.cpp
+++ b/Engine/Plugins/Chooser/Source/Chooser/Private/RandomizeColumn.cpp
@@ -108,7 +108,7 @@ void FRandomizeColumn::Filter(FChooserEvaluationContext& Context, const FChooser
 			}
 
 			// pick a random float from 0-total weight
-			float RandomNumber = FMath::FRandRange(0.0f, TotalWeight);
+			const float RandomNumber = Context.RandomStream ? Context.RandomStream->RandRange(0.f, TotalWeight) : FMath::FRandRange(0.0f, TotalWeight);
 			float Weight = 0;
 
 			// add up the weights again, and select the index where our sum clears the random float
@@ -160,7 +160,7 @@ void FRandomizeColumn::Filter(FChooserEvaluationContext& Context, const FChooser
 			}
 
 			// pick a random float from 0-total weight
-			float RandomNumber = FMath::FRandRange(0.0f, TotalWeight);
+			const float RandomNumber = Context.RandomStream ? Context.RandomStream->RandRange(0.f, TotalWeight) : FMath::FRandRange(0.0f, TotalWeight);
 			float Weight = 0;
 
 			// add up the weights again, and select the index where our sum clears the random float
diff --git a/Engine/Plugins/Chooser/Source/Chooser/Public/IObjectChooser.h b/Engine/Plugins/Chooser/Source/Chooser/Public/IObjectChooser.h
--- a/Engine/Plugins/Chooser/Source/Chooser/Public/IObjectChooser.h
+++ b/Engine/Plugins/Chooser/Source/Chooser/Public/IObjectChooser.h
@@ -100,6 +100,8 @@ struct FChooserEvaluationContext
 
 	// storage for Object Params, call AddObjectParam to allocate one FChooserEvaluationInputObject in this array and then add a StructView of it to the Params array
 	TArray<FChooserEvaluationInputObject, TFixedAllocator<4>> ObjectParams;
+
+	FRandomStream* RandomStream = nullptr;
 };
 
 USTRUCT()
