fix Chooser Table editor appearance when the table width is smaller than the window (background color and header will now fill the empty space)

#Synced-CL 49562535

diff --git a/Engine/Plugins/Chooser/Source/ChooserEditor/Private/SChooserTableWidget.cpp b/Engine/Plugins/Chooser/Source/ChooserEditor/Private/SChooserTableWidget.cpp
--- a/Engine/Plugins/Chooser/Source/ChooserEditor/Private/SChooserTableWidget.cpp
+++ b/Engine/Plugins/Chooser/Source/ChooserEditor/Private/SChooserTableWidget.cpp
@@ -295,10 +295,10 @@ namespace UE::ChooserEditor
 			SNew(SHorizontalBox)
 			+SHorizontalBox::Slot()
 			[
-				SNew(SScrollBox).Orientation(Orient_Horizontal)//.ConsumeMouseWheel(EConsumeMouseWheel::Never)
-					+ SScrollBox::Slot()
+				SNew(SScrollBox).Orientation(Orient_Horizontal)
+					+ SScrollBox::Slot().FillContentSize(1.0f)
 				[
-						TableView.ToSharedRef()
+					TableView.ToSharedRef()
 				]
 			]
 			+ SHorizontalBox::Slot().AutoWidth()
