Fix check with incorrect logic in EvaluateChooserNode

#Synced-CL 49562535

diff --git a/Engine/Plugins/Chooser/Source/ChooserUncooked/Private/EvaluateChooserNode.cpp b/Engine/Plugins/Chooser/Source/ChooserUncooked/Private/EvaluateChooserNode.cpp
--- a/Engine/Plugins/Chooser/Source/ChooserUncooked/Private/EvaluateChooserNode.cpp
+++ b/Engine/Plugins/Chooser/Source/ChooserUncooked/Private/EvaluateChooserNode.cpp
@@ -604,7 +604,7 @@ void UK2Node_EvaluateChooser2::ExpandNode(class FKismetCompilerContext& Compiler
 	UEdGraphPin* ExecInput = GetExecPin();
 	UEdGraphPin* ExecOutput = GetPassThroughPin(ExecInput);
 	UEdGraphPin* ResultPin = FindPin(TEXT("Result"));
-	check(ResultPin || (Chooser && Chooser->ResultType == EObjectChooserResultType::NoPrimaryResult));
+	ensure(Chooser == nullptr || ResultPin || Chooser->ResultType == EObjectChooserResultType::NoPrimaryResult);
 
 	UEdGraphPin* PreviousNodeExecOutput = nullptr;
 	
