Fix filtering in Soft Asset references in choosers

#Synced-CL 49899510

diff --git a/Engine/Plugins/Chooser/Source/Chooser/Internal/ObjectChooser_Asset.h b/Engine/Plugins/Chooser/Source/Chooser/Internal/ObjectChooser_Asset.h
--- a/Engine/Plugins/Chooser/Source/Chooser/Internal/ObjectChooser_Asset.h
+++ b/Engine/Plugins/Chooser/Source/Chooser/Internal/ObjectChooser_Asset.h
@@ -39,7 +39,7 @@ struct FSoftAssetChooser : public FObjectChooserBase
 	virtual UObject* GetReferencedObject() const override { return Asset.LoadSynchronous(); }
 #endif
 	
-	UPROPERTY(EditAnywhere, Category = "Parameters")
+	UPROPERTY(EditAnywhere, Category = "Parameters", meta = (GetAssetFilter="ResultAssetFilter"))
 	TSoftObjectPtr<UObject> Asset;
 };
 
