Fix disabled rows not updating correctly when assets are drag/dropped into a chooser table

#Synced-CL 50178448

diff --git a/Engine/Plugins/Chooser/Source/ChooserEditor/Private/SChooserTableRow.cpp b/Engine/Plugins/Chooser/Source/ChooserEditor/Private/SChooserTableRow.cpp
--- a/Engine/Plugins/Chooser/Source/ChooserEditor/Private/SChooserTableRow.cpp
+++ b/Engine/Plugins/Chooser/Source/ChooserEditor/Private/SChooserTableRow.cpp
@@ -571,6 +571,7 @@ namespace UE::ChooserEditor
 					}
 
 					Chooser->ResultsStructs.Insert(NewResults, InsertRowIndex);
+					Chooser->DisabledRows.InsertZeroed(InsertRowIndex, NewResults.Num());
 					
 					// Make sure each column has the same number of row datas as there are results
 					for(FInstancedStruct& ColumnData : Chooser->ColumnsStructs)
