Fix UAF Workspace tree not responding to new nested choosers added via copy/paste operations etc

#Synced-CL 50281394

diff --git a/Engine/Plugins/Chooser/Source/ChooserEditor/Private/ChooserTableViewModel.cpp b/Engine/Plugins/Chooser/Source/ChooserEditor/Private/ChooserTableViewModel.cpp
--- a/Engine/Plugins/Chooser/Source/ChooserEditor/Private/ChooserTableViewModel.cpp
+++ b/Engine/Plugins/Chooser/Source/ChooserEditor/Private/ChooserTableViewModel.cpp
@@ -52,6 +52,9 @@
 #include "SPositiveActionButton.h"
 #include "ChooserTableEditor.h"
 #include "SChooserTableWidget.h"
+#include "AssetRegistry/IAssetRegistry.h"
+#include "UObject/AssetRegistryTagsContext.h"
+
 
 #include UE_INLINE_GENERATED_CPP_BY_NAME(ChooserTableViewModel)
 
@@ -229,8 +232,8 @@ void FChooserTableViewModel::RefreshAll()
 	{
 		SelectRows(CachedSelectedRows);
 	}
-	
-	// RefreshNestedObjectTree();
+
+	RefreshAssetRegistry();
 }
 
 void FChooserTableViewModel::QueueRefreshAll()
@@ -1377,7 +1380,16 @@ void FChooserTableViewModel::SelectRows(const TConstArrayView<int32>& Rows)
 {
 	SelectRowsDelegate.ExecuteIfBound(Rows);
 }
-	
+
+void FChooserTableViewModel::RefreshAssetRegistry()
+{
+	if (IAssetRegistry* AssetRegistry = IAssetRegistry::Get())
+	{
+		// this causes the UAF workspace outliner to refresh
+		AssetRegistry->AssetUpdateTags(RootChooser, EAssetRegistryTagsCaller::Fast);
+	}	
+}
+
 void FChooserTableViewModel::ClearSelectedRows()
 {
 	SelectRows({});
diff --git a/Engine/Plugins/Chooser/Source/ChooserEditor/Private/ChooserTableViewModel.h b/Engine/Plugins/Chooser/Source/ChooserEditor/Private/ChooserTableViewModel.h
--- a/Engine/Plugins/Chooser/Source/ChooserEditor/Private/ChooserTableViewModel.h
+++ b/Engine/Plugins/Chooser/Source/ChooserEditor/Private/ChooserTableViewModel.h
@@ -139,7 +139,8 @@ namespace UE::ChooserEditor
 		void RefreshRowSelectionDetails();
 		ESelectionType GetCurrentSelectionType() const { return CurrentSelectionType; }
 		void SetCurrentSelectionType(ESelectionType SelectionType) { CurrentSelectionType = SelectionType; }
-		
+
+		void RefreshAssetRegistry(); // update nested chooser info in asset registry
 
 		void ClearSelectedRows();
 		void SelectRows(const TConstArrayView<int32>& Rows);		// called internally to change row selection after operations
