Fix hash bug in compiled chooser context bindings (now including parameter index in hash)
Add more null checks in chooser context validation code, and disable log spam from UAF context entries

#Synced-CL 50686245

diff --git a/Engine/Plugins/Chooser/Source/Chooser/Private/ChooserPropertyAccess.cpp b/Engine/Plugins/Chooser/Source/Chooser/Private/ChooserPropertyAccess.cpp
--- a/Engine/Plugins/Chooser/Source/Chooser/Private/ChooserPropertyAccess.cpp
+++ b/Engine/Plugins/Chooser/Source/Chooser/Private/ChooserPropertyAccess.cpp
@@ -54,24 +54,25 @@ FName FChooserPropertyBinding::GetUniqueId() const
 struct FCompiledBindingCacheId
 {
 	FName BindingPath;
-	const UStruct* Type;
+	const UStruct* Type = nullptr;
+	int32 ContextIndex = 0;
 	
 	bool operator == (const FCompiledBindingCacheId& Other) const
 	{
-		 return Type == Other.Type
-		  && BindingPath == Other.BindingPath;
+		 return ContextIndex == Other.ContextIndex
+				&&Type == Other.Type
+				&& BindingPath == Other.BindingPath;
 	}
 };
 
 FORCEINLINE uint32 GetTypeHash(const FCompiledBindingCacheId& Id)
 {
-	return HashCombineFast(GetTypeHash(Id.Type), GetTypeHash(Id.BindingPath));
+	return HashCombineFast(GetTypeHash(Id.ContextIndex),HashCombineFast(GetTypeHash(Id.Type), GetTypeHash(Id.BindingPath)));
 }
 
 TMap<FCompiledBindingCacheId, TWeakPtr<UE::Chooser::FCompiledBinding>> CompiledBindingCache;
 FCriticalSection CompiledBindingCacheLock;
 
-
 void FChooserPropertyBinding::Compile(IHasContextClass* Owner, bool bForce)
 {
 	TRACE_CPUPROFILER_EVENT_SCOPE(CompilePropertyChain);
@@ -117,6 +118,7 @@ void FChooserPropertyBinding::Compile(IHasContextClass* Owner, bool bForce)
 	FCompiledBindingCacheId Id;
 	Id.Type = StructType;
 	Id.BindingPath = GetUniqueId();
+	Id.ContextIndex = ContextIndex;
 
 	{
 		FScopeLock Lock(&CompiledBindingCacheLock);
@@ -403,8 +405,17 @@ namespace UE::Chooser
 						}
 						else
 						{
-							UE_LOG(LogChooser, Error, TEXT("Chooser Table: %s ContextData entry %d expects an object of type %s, but was passed a struct of type %s."),
-								ToCStr(Chooser->GetName()), i, ToCStr(ExpectedClassType->Class->GetName()), ToCStr(Context.Params[i].GetScriptStruct()->GetName()));
+							if (Context.Params[i].GetScriptStruct())
+							{
+								UE_LOG(LogChooser, Error, TEXT("Chooser Table: %s ContextData entry %d expects an object of type %s, but was passed a struct of type %s."),
+									ToCStr(Chooser->GetName()), i, ToCStr(ExpectedClassType->Class->GetName()), ToCStr(Context.Params[i].GetScriptStruct()->GetName()));
+							}
+							else
+							{
+								
+								UE_LOG(LogChooser, Error, TEXT("Chooser Table: %s ContextData entry %d expects an object of type %s, but was passed a null struct param"),
+									ToCStr(Chooser->GetName()), i, ToCStr(ExpectedClassType->Class->GetName()));
+							}
 						}
 					}
 					else
@@ -422,15 +433,31 @@ namespace UE::Chooser
 					{
 						if (FChooserEvaluationInputObject* InputObjectParam = Context.Params[i].GetPtr<FChooserEvaluationInputObject>())
 						{
-							UE_LOG(LogChooser, Error, TEXT("Chooser Table: %s ContextData entry %d expects an struct of type %s, but was passed a object of type %s."),
-								ToCStr(Chooser->GetName()), i, ToCStr(ExpectedStructType->Struct->GetName()), ToCStr(InputObjectParam->Object->GetClass()->GetName()));
+							if (InputObjectParam->Object)
+							{
+								UE_LOG(LogChooser, Error, TEXT("Chooser Table: %s ContextData entry %d expects an struct of type %s, but was passed a object of type %s."),
+									ToCStr(Chooser->GetName()), i, ToCStr(ExpectedStructType->Struct->GetName()), ToCStr(InputObjectParam->Object->GetClass()->GetName()));
+							}
+							else
+							{
+								UE_LOG(LogChooser, Error, TEXT("Chooser Table: %s ContextData entry %d expects an struct of type %s, but was passed a null object."),
+									ToCStr(Chooser->GetName()), i, ToCStr(ExpectedStructType->Struct->GetName()));
+							}
 						}
 						else
 						{
 							if (Context.Params[i].GetScriptStruct() != ExpectedStructType->Struct)
 							{
-								UE_LOG(LogChooser, Error, TEXT("Chooser Table: %s ContextData entry %d expects an struct of type %s, but was passed a struct of type %s."),
-									ToCStr(Chooser->GetName()), i, ToCStr(ExpectedStructType->Struct->GetName()), ToCStr(Context.Params[i].GetScriptStruct()->GetName()));
+								if (Context.Params[i].GetScriptStruct())
+								{
+									UE_LOG(LogChooser, Error, TEXT("Chooser Table: %s ContextData entry %d expects an struct of type %s, but was passed a struct of type %s."),
+										ToCStr(Chooser->GetName()), i, ToCStr(ExpectedStructType->Struct->GetName()), ToCStr(Context.Params[i].GetScriptStruct()->GetName()));
+								}
+								else
+								{
+									UE_LOG(LogChooser, Error, TEXT("Chooser Table: %s ContextData entry %d expects an struct of type %s, but was passed a null struct parameter."),
+										ToCStr(Chooser->GetName()), i, ToCStr(ExpectedStructType->Struct->GetName()));
+								}
 							}
 						}
 					}
@@ -443,7 +470,8 @@ namespace UE::Chooser
 			}
 			else
 			{
-				UE_LOG(LogChooser, Error, TEXT("Chooser Table: %s  ContextData entry %i is of unknown type or none."), ToCStr(Chooser->GetName()), i);
+				// UAF has it's own child struct, so since this code can't be aware of that type, for now disabling this error
+				//UE_LOG(LogChooser, Error, TEXT("Chooser Table: %s  ContextData entry %i is of unknown type or none."), ToCStr(Chooser->GetName()), i);
 			}
 		}
 	}
