fix bug where Choosers are not marked dirty after details edits

#Synced-CL 51028968

#Synced-CL 51028968

diff --git a/Engine/Plugins/Chooser/Source/ChooserEditor/Private/ChooserTableViewModel.cpp b/Engine/Plugins/Chooser/Source/ChooserEditor/Private/ChooserTableViewModel.cpp
--- a/Engine/Plugins/Chooser/Source/ChooserEditor/Private/ChooserTableViewModel.cpp
+++ b/Engine/Plugins/Chooser/Source/ChooserEditor/Private/ChooserTableViewModel.cpp
@@ -164,6 +164,7 @@ void FChooserTableViewModel::OnObjectsTransacted(UObject* Object, const FTransac
 					ColumnIndex++;
 				}
 			}
+			RowDetails->Chooser->MarkPackageDirty();
 			QueueRefreshAll();
 		}
 	}
