Chooser - removing unnecessary Modify()/PostEditChange() from SChooserTableWidget::UpdateTableRows

#Synced-CL 51062901

diff --git a/Engine/Plugins/Chooser/Source/ChooserEditor/Private/SChooserTableWidget.cpp b/Engine/Plugins/Chooser/Source/ChooserEditor/Private/SChooserTableWidget.cpp
--- a/Engine/Plugins/Chooser/Source/ChooserEditor/Private/SChooserTableWidget.cpp
+++ b/Engine/Plugins/Chooser/Source/ChooserEditor/Private/SChooserTableWidget.cpp
@@ -138,7 +138,7 @@ namespace UE::ChooserEditor
 	void SChooserTableWidget::UpdateTableRows()
 	{
 		UChooserTable* Chooser = ChooserViewModel->GetChooser();
-		Chooser->Modify();
+
 		int32 NewNum = Chooser->ResultsStructs.Num();
 		Chooser->DisabledRows.SetNum(NewNum);
     
@@ -164,7 +164,6 @@ namespace UE::ChooserEditor
 			FChooserColumnBase& Column = ColumnData.GetMutable<FChooserColumnBase>();
 			Column.SetNumRows(NewNum);
 		}
-		Chooser->PostEditChange();
     
 		if (TableView.IsValid())
 		{
