fix invalid index access for fallback row

#Synced-CL 51094799

diff --git a/Engine/Plugins/Chooser/Source/Chooser/Private/OutputGameplayTagColumn.cpp b/Engine/Plugins/Chooser/Source/Chooser/Private/OutputGameplayTagColumn.cpp
--- a/Engine/Plugins/Chooser/Source/Chooser/Private/OutputGameplayTagColumn.cpp
+++ b/Engine/Plugins/Chooser/Source/Chooser/Private/OutputGameplayTagColumn.cpp
@@ -50,7 +50,7 @@ void FOutputGameplayTagColumn::SetOutputs(FChooserEvaluationContext& Context, in
 		FPropertyBagPropertyDesc PropertyDesc(PropertyName, EPropertyBagPropertyType::Struct, FGameplayTagContainer::StaticStruct());
 		PropertyDesc.MetaData.Add(FPropertyBagPropertyDescMetaData("DisplayName", DisplayName.ToString()));
 		PropertyBag.AddProperties({PropertyDesc});
-		PropertyBag.SetValueStruct(PropertyName, RowValues[RowIndex]);
+		PropertyBag.SetValueStruct(PropertyName, GetValueForIndex(RowIndex));
 	}
 
 	void FOutputGameplayTagColumn::SetFromDetails(FInstancedPropertyBag& PropertyBag, int32 ColumnIndex, int32 RowIndex)
@@ -60,7 +60,7 @@ void FOutputGameplayTagColumn::SetOutputs(FChooserEvaluationContext& Context, in
 		TValueOrError<FStructView, EPropertyBagResult> Result = PropertyBag.GetValueStruct(PropertyName, FGameplayTagContainer::StaticStruct());
 		if (FStructView* StructView = Result.TryGetValue())
 		{
-			RowValues[RowIndex] = StructView->Get<FGameplayTagContainer>();
+			GetValueForIndex(RowIndex) = StructView->Get<FGameplayTagContainer>();
 		}
 	}
 #endif
diff --git a/Engine/Plugins/Chooser/Source/Chooser/Private/OutputGameplayTagQueryColumn.cpp b/Engine/Plugins/Chooser/Source/Chooser/Private/OutputGameplayTagQueryColumn.cpp
--- a/Engine/Plugins/Chooser/Source/Chooser/Private/OutputGameplayTagQueryColumn.cpp
+++ b/Engine/Plugins/Chooser/Source/Chooser/Private/OutputGameplayTagQueryColumn.cpp
@@ -55,17 +55,17 @@ void FOutputGameplayTagQueryColumn::SetOutputs(FChooserEvaluationContext& Contex
 		FPropertyBagPropertyDesc PropertyDesc(PropertyName, EPropertyBagPropertyType::Struct, FGameplayTagQuery::StaticStruct());
 		PropertyDesc.MetaData.Add(FPropertyBagPropertyDescMetaData("DisplayName", DisplayName.ToString()));
 		PropertyBag.AddProperties({PropertyDesc});
-		PropertyBag.SetValueStruct(PropertyName, RowValues[RowIndex]);
+		PropertyBag.SetValueStruct(PropertyName, GetValueForIndex(RowIndex));
 	}
 
 	void FOutputGameplayTagQueryColumn::SetFromDetails(FInstancedPropertyBag& PropertyBag, int32 ColumnIndex, int32 RowIndex)
 	{
 		FName PropertyName("RowData", ColumnIndex);
-		
+
 		TValueOrError<FStructView, EPropertyBagResult> Result = PropertyBag.GetValueStruct(PropertyName, FGameplayTagQuery::StaticStruct());
 		if (FStructView* StructView = Result.TryGetValue())
 		{
-			RowValues[RowIndex] = StructView->Get<FGameplayTagQuery>();
+			GetValueForIndex(RowIndex) = StructView->Get<FGameplayTagQuery>();
 		}
-	}	
+	}
 #endif
