Support soft object references properly in the Object Output column

#Synced-CL 51055096
#Synced-CL 51055096

diff --git a/Engine/Plugins/Chooser/Source/Chooser/Internal/ObjectColumn.h b/Engine/Plugins/Chooser/Source/Chooser/Internal/ObjectColumn.h
--- a/Engine/Plugins/Chooser/Source/Chooser/Internal/ObjectColumn.h
+++ b/Engine/Plugins/Chooser/Source/Chooser/Internal/ObjectColumn.h
@@ -24,6 +24,7 @@ struct FObjectContextProperty : public FChooserParameterObjectBase
 	UE_API virtual bool GetValue(FChooserEvaluationContext& Context, FSoftObjectPath& OutResult) const override;
 	UE_API virtual bool GetValue(FChooserEvaluationContext& Context, UObject*& OutResult) const override;
 	UE_API virtual bool SetValue(FChooserEvaluationContext& Context, UObject* InValue) const override;
+	UE_API virtual bool SetValue(FChooserEvaluationContext& Context, const FSoftObjectPath& InValue) const override;
 
 	CHOOSER_PARAMETER_BOILERPLATE();
 
diff --git a/Engine/Plugins/Chooser/Source/Chooser/Private/ObjectColumn.cpp b/Engine/Plugins/Chooser/Source/Chooser/Private/ObjectColumn.cpp
--- a/Engine/Plugins/Chooser/Source/Chooser/Private/ObjectColumn.cpp
+++ b/Engine/Plugins/Chooser/Source/Chooser/Private/ObjectColumn.cpp
@@ -85,6 +85,29 @@ bool FObjectContextProperty::SetValue(FChooserEvaluationContext& Context, UObjec
 	return false;
 }
 
+bool FObjectContextProperty::SetValue(FChooserEvaluationContext& Context, const FSoftObjectPath& InValue) const
+{
+	UE::Chooser::FResolvedPropertyChainResult Result;
+	if (UE::Chooser::ResolvePropertyChain(Context, Binding, Result))
+	{
+		if (Result.Function == nullptr)
+		{
+			if (Result.PropertyType == UE::Chooser::EChooserPropertyAccessType::SoftObjectRef)
+			{
+				FSoftObjectPtr& SoftObjectPtr = *reinterpret_cast<FSoftObjectPtr*>(Result.Container + Result.PropertyOffset);
+				SoftObjectPtr = InValue;
+				return true;
+			}
+		
+			UObject** ObjectPtr = reinterpret_cast<UObject**>(Result.Container + Result.PropertyOffset);
+			*ObjectPtr = InValue.ResolveObject();
+			return true;
+		}
+	}
+
+	return false;
+}
+
 FObjectColumn::FObjectColumn()
 {
 #if WITH_EDITOR
diff --git a/Engine/Plugins/Chooser/Source/Chooser/Private/OutputObjectColumn.cpp b/Engine/Plugins/Chooser/Source/Chooser/Private/OutputObjectColumn.cpp
--- a/Engine/Plugins/Chooser/Source/Chooser/Private/OutputObjectColumn.cpp
+++ b/Engine/Plugins/Chooser/Source/Chooser/Private/OutputObjectColumn.cpp
@@ -38,9 +38,16 @@ void FOutputObjectColumn::SetOutputs(FChooserEvaluationContext& Context, int Row
 	{
 		if (RowIndex == ChooserColumn_SpecialIndex_Fallback || RowValues.IsValidIndex(RowIndex))
 		{
+			TSoftObjectPtr<UObject> Result;
 			const FObjectChooserBase* Value = GetValueForIndex(RowIndex).Value.GetPtr<FObjectChooserBase>();
-			UObject* Result = Value ? Value->ChooseObject(Context) : nullptr;
-			InputValue.Get<FChooserParameterObjectBase>().SetValue(Context, Result);
+			if (Value)
+			{
+				Value->ChooseObject(Context, Result);
+			}
+			if (const FChooserParameterObjectBase* OutputParameter = InputValue.GetPtr<FChooserParameterObjectBase>())
+			{
+				OutputParameter->SetValue(Context, Result.ToSoftObjectPath());
+			}
 		}
 		else
 		{
diff --git a/Engine/Plugins/Chooser/Source/Chooser/Public/IChooserParameterObject.h b/Engine/Plugins/Chooser/Source/Chooser/Public/IChooserParameterObject.h
--- a/Engine/Plugins/Chooser/Source/Chooser/Public/IChooserParameterObject.h
+++ b/Engine/Plugins/Chooser/Source/Chooser/Public/IChooserParameterObject.h
@@ -15,6 +15,7 @@ struct FChooserParameterObjectBase : public FChooserParameterBase
 
 	virtual bool GetValue(FChooserEvaluationContext& Context, FSoftObjectPath& OutResult) const { return false; }
 	virtual bool GetValue(FChooserEvaluationContext& Context, UObject*& OutResult) const { return false; }
+	virtual bool SetValue(FChooserEvaluationContext& Context, const FSoftObjectPath& OutResult) const { return false; }
 	virtual bool SetValue(FChooserEvaluationContext& Context, UObject* OutResult) const { return false; }
 
 #if WITH_EDITOR
