Added a new output name column for choosers
#Synced-CL 51211439

diff --git a/Engine/Plugins/Chooser/Source/Chooser/Internal/OutputNameColumn.h b/Engine/Plugins/Chooser/Source/Chooser/Internal/OutputNameColumn.h
new file mode 100644
--- /dev/null
+++ b/Engine/Plugins/Chooser/Source/Chooser/Internal/OutputNameColumn.h
@@ -0,0 +1,93 @@
+// Copyright Epic Games, Inc. All Rights Reserved.
+
+#pragma once
+
+#include "ChooserPropertyAccess.h"
+#include "CoreMinimal.h"
+#include "ChooserParameterName.h"
+#include "IChooserColumn.h"
+#include "StructUtils/InstancedStruct.h"
+#include "Serialization/MemoryReader.h"
+#include "OutputNameColumn.generated.h"
+
+#define UE_API CHOOSER_API
+
+struct FBindingChainElement;
+
+USTRUCT(DisplayName = "Output Name", Meta = (Category = "Output", Tooltip = "A column which writes a FName value."))
+struct FChooserOutputNameColumn : public FChooserColumnBase
+{
+	GENERATED_BODY()
+
+	UE_API FChooserOutputNameColumn();
+	virtual bool HasFilters() const override { return false; }
+	virtual bool HasOutputs() const override { return true; }
+	UE_API virtual void SetOutputs(FChooserEvaluationContext& Context, int RowIndex) const override;
+
+	UPROPERTY(EditAnywhere, NoClear, Meta = (ExcludeBaseStruct, BaseStruct = "/Script/Chooser.ChooserParameterNameBase", ToolTip="The Name reference property this column will write to"), Category = "Hidden")
+	FInstancedStruct InputValue;
+
+	UPROPERTY()
+	TArray<FName> RowValues;
+
+#if WITH_EDITOR
+	mutable FName TestValue = NAME_None;
+	UE_API virtual void AddToDetails(FInstancedPropertyBag& PropertyBag, int32 ColumnIndex, int32 RowIndex) override;
+	UE_API virtual void SetFromDetails(FInstancedPropertyBag& PropertyBag, int32 ColumnIndex, int32 RowIndex) override;
+	virtual void CopyFallback(FChooserColumnBase& SourceColumn) override { FallbackValue = static_cast<FChooserOutputNameColumn&>(SourceColumn).FallbackValue; }
+#endif
+
+	FName& GetValueForIndex(int32 Index)
+	{
+		if (Index == ChooserColumn_SpecialIndex_Fallback)
+		{
+			return FallbackValue;
+		}
+		else if (!RowValues.IsValidIndex(Index))
+		{
+			return FallbackValue;
+		}
+
+		return RowValues[Index];
+	}
+
+	FName GetValueForIndex(int32 Index) const
+   	{
+		if (Index == ChooserColumn_SpecialIndex_Fallback)
+		{
+			return FallbackValue;
+		}
+		else if (!RowValues.IsValidIndex(Index))
+		{
+			return FallbackValue;
+		}
+
+		return RowValues[Index];
+   	}
+
+	// FallbackValue will be used as the output value if the all rows in the chooser fail, and the FallbackResult from the chooser is used.
+	UPROPERTY()
+	FName FallbackValue = NAME_None;
+
+	CHOOSER_COLUMN_BOILERPLATE(FChooserParameterNameBase);
+
+#if WITH_EDITOR
+	virtual void PostLoad() override
+	{
+		Super::PostLoad();
+
+		if (InputValue.IsValid())
+		{
+			InputValue.GetMutable<FChooserParameterBase>().PostLoad();
+		}
+	}
+#endif
+
+#if WITH_EDITORONLY_DATA
+	UPROPERTY(EditAnywhere, Category=Data, meta=(ToolTip="DefaultRowValue will be assigned to cells when new rows are created"));
+	FName DefaultRowValue = NAME_None;
+#endif
+
+};
+
+#undef UE_API
diff --git a/Engine/Plugins/Chooser/Source/Chooser/Private/OutputNameColumn.cpp b/Engine/Plugins/Chooser/Source/Chooser/Private/OutputNameColumn.cpp
new file mode 100644
--- /dev/null
+++ b/Engine/Plugins/Chooser/Source/Chooser/Private/OutputNameColumn.cpp
@@ -0,0 +1,71 @@
+// Copyright Epic Games, Inc. All Rights Reserved.
+#include "OutputNameColumn.h"
+#include "ChooserIndexArray.h"
+#include "ChooserPropertyAccess.h"
+#include "ChooserTrace.h"
+#include "NameColumn.h"
+
+#if WITH_EDITOR
+#include "IPropertyAccessEditor.h"
+#include "StructUtils/PropertyBag.h"
+#endif
+
+#include UE_INLINE_GENERATED_CPP_BY_NAME(OutputNameColumn)
+
+FChooserOutputNameColumn::FChooserOutputNameColumn()
+{
+#if WITH_EDITOR
+	InputValue.InitializeAs(FNameContextProperty::StaticStruct());
+#endif
+}
+
+void FChooserOutputNameColumn::SetOutputs(FChooserEvaluationContext& Context, int RowIndex) const
+{
+	if (InputValue.IsValid())
+	{
+		if (RowIndex == ChooserColumn_SpecialIndex_Fallback || RowValues.IsValidIndex(RowIndex))
+		{
+			FName Value = GetValueForIndex(RowIndex);
+			InputValue.Get<FNameContextProperty>().SetValue(Context, Value);
+						
+			#if WITH_EDITOR
+			if (Context.DebuggingInfo.bCurrentDebugTarget)
+			{
+				TestValue = Value;
+			}
+			#endif
+		}
+		else
+		{
+#if CHOOSER_DEBUGGING_ENABLED
+			UE_ASSET_LOG(LogChooser, Error, Context.DebuggingInfo.CurrentChooser, TEXT("Invalid index %d passed to FChooserOutputNameColumn::SetOutputs"), RowIndex);
+#else
+			UE_LOG(LogChooser, Error, TEXT("Invalid index %d passed to FChooserOutputNameColumn::SetOutputs"), RowIndex);
+#endif
+		}
+	}
+}
+
+#if WITH_EDITOR
+	void FChooserOutputNameColumn::AddToDetails(FInstancedPropertyBag& PropertyBag, int32 ColumnIndex, int32 RowIndex)
+	{
+		FText DisplayName;
+		InputValue.Get<FChooserParameterNameBase>().GetDisplayName(DisplayName);
+		FName PropertyName("RowData",ColumnIndex);
+		FPropertyBagPropertyDesc PropertyDesc(PropertyName, EPropertyBagPropertyType::Name);
+		PropertyDesc.MetaData.Add(FPropertyBagPropertyDescMetaData("DisplayName", DisplayName.ToString()));
+		PropertyBag.AddProperties({PropertyDesc});
+		PropertyBag.SetValueName(PropertyName, GetValueForIndex(RowIndex));
+	}
+
+	void FChooserOutputNameColumn::SetFromDetails(FInstancedPropertyBag& PropertyBag, int32 ColumnIndex, int32 RowIndex)
+	{
+		FName PropertyName("RowData", ColumnIndex);
+
+		TValueOrError<FName, EPropertyBagResult> Result = PropertyBag.GetValueName(PropertyName);
+		if (FName* Name = Result.TryGetValue())
+		{
+			GetValueForIndex(RowIndex) = *Name;
+		}
+	}
+#endif
diff --git a/Engine/Plugins/Chooser/Source/ChooserEditor/Private/ChooserEditorModule.cpp b/Engine/Plugins/Chooser/Source/ChooserEditor/Private/ChooserEditorModule.cpp
--- a/Engine/Plugins/Chooser/Source/ChooserEditor/Private/ChooserEditorModule.cpp
+++ b/Engine/Plugins/Chooser/Source/ChooserEditor/Private/ChooserEditorModule.cpp
@@ -24,6 +24,7 @@
 #include "ObjectClassColumnEditor.h"
 #include "OutputEnumColumnEditor.h"
 #include "OutputFloatColumnEditor.h"
+#include "OutputNameColumnEditor.h"
 #include "OutputObjectColumnEditor.h"
 #include "OutputStructColumnEditor.h"
 #include "PropertyAccessChainCustomization.h"
@@ -143,6 +144,7 @@ void FChooserEditorModule::StartupModule()
 	RegisterMultiEnumWidgets();
 	RegisterOutputEnumWidgets();
 	RegisterNameWidgets();
+	RegisterOutputNameWidgets();
 	RegisterObjectWidgets();
 	RegisterObjectClassWidgets();
 	RegisterOutputObjectWidgets();
diff --git a/Engine/Plugins/Chooser/Source/ChooserEditor/Private/OutputNameColumnEditor.cpp b/Engine/Plugins/Chooser/Source/ChooserEditor/Private/OutputNameColumnEditor.cpp
new file mode 100644
--- /dev/null
+++ b/Engine/Plugins/Chooser/Source/ChooserEditor/Private/OutputNameColumnEditor.cpp
@@ -0,0 +1,82 @@
+// Copyright Epic Games, Inc. All Rights Reserved.
+
+#include "OutputNameColumnEditor.h"
+
+#include <ChooserColumnHeader.h>
+
+#include "SPropertyAccessChainWidget.h"
+#include "GraphEditorSettings.h"
+#include "ObjectChooserWidgetFactories.h"
+#include "OutputNameColumn.h"
+#include "PropertyCustomizationHelpers.h"
+#include "TransactionCommon.h"
+#include "Widgets/Images/SImage.h"
+#include "Widgets/Input/SEditableTextBox.h"
+#include "Widgets/Input/SButton.h"
+#include "Widgets/Layout/SBox.h"
+#include "ScopedTransaction.h"
+
+#define LOCTEXT_NAMESPACE "OutputNameColumnEditor"
+
+namespace UE::ChooserEditor
+{
+	TSharedRef<SWidget> CreateOutputNameColumnWidget(UChooserTable* Chooser, FChooserColumnBase* Column, int Row)
+	{
+		FChooserOutputNameColumn* OutputNameColumn = static_cast<FChooserOutputNameColumn*>(Column);
+
+		if (Row == ColumnWidget_SpecialIndex_Header)
+		{
+			// create column header widget
+			const FSlateBrush* ColumnIcon = FCoreStyle::Get().GetBrush("Icons.ArrowRight");
+			const FText ColumnTooltip = LOCTEXT("Output Name Tooltip", "Output Name: writes the value from cell in the result row to the bound variable");
+			const FText ColumnName = LOCTEXT("Output Name","Output Name");
+					
+			TSharedPtr<SWidget> DebugWidget = nullptr;
+			if (Chooser->GetEnableDebugTesting())
+			{
+				DebugWidget = SNew(SEditableTextBox)
+					.IsEnabled(false)
+					.Text_Lambda([OutputNameColumn]() { return FText::FromName(OutputNameColumn->TestValue); });
+			}
+			
+			return MakeColumnHeaderWidget(Chooser, Column, ColumnName, ColumnTooltip, ColumnIcon, DebugWidget);
+		}
+		if (Row == ColumnWidget_SpecialIndex_Fallback)
+		{
+			return SNew(SEditableTextBox)
+					.Text_Lambda([OutputNameColumn]()
+					{
+						return FText::FromName(OutputNameColumn->FallbackValue);
+					})
+					.OnTextCommitted_Lambda([Chooser, OutputNameColumn](const FText& NewValue, ETextCommit::Type CommitType)
+					{
+						const FScopedTransaction Transaction(LOCTEXT("Edit Name Value", "Edit Name Value"));
+						Chooser->Modify(true);
+						OutputNameColumn->FallbackValue = FName(NewValue.ToString());
+					});	
+		}
+
+		return SNew(SEditableTextBox)
+				.Text_Lambda([OutputNameColumn, Row]()
+				{
+					return OutputNameColumn->RowValues.IsValidIndex(Row) ? FText::FromName(OutputNameColumn->RowValues[Row]) : FText();
+				})
+				.OnTextCommitted_Lambda([Chooser, OutputNameColumn, Row](const FText& NewValue, ETextCommit::Type CommitType)
+				{
+					const FScopedTransaction Transaction(LOCTEXT("Edit Name", "Edit Name"));
+					Chooser->Modify(true);
+					if (OutputNameColumn->RowValues.IsValidIndex(Row))
+					{
+						OutputNameColumn->RowValues[Row] = FName(NewValue.ToString());
+					}
+				});
+	}
+
+	void RegisterOutputNameWidgets()
+	{
+		FObjectChooserWidgetFactories::RegisterColumnWidgetCreator(FChooserOutputNameColumn::StaticStruct(), CreateOutputNameColumnWidget);
+	}
+
+} // namespace UE::ChooserEditor
+
+#undef LOCTEXT_NAMESPACE
diff --git a/Engine/Plugins/Chooser/Source/ChooserEditor/Private/OutputNameColumnEditor.h b/Engine/Plugins/Chooser/Source/ChooserEditor/Private/OutputNameColumnEditor.h
new file mode 100644
--- /dev/null
+++ b/Engine/Plugins/Chooser/Source/ChooserEditor/Private/OutputNameColumnEditor.h
@@ -0,0 +1,10 @@
+// Copyright Epic Games, Inc. All Rights Reserved.
+
+#pragma once
+
+#include "CoreMinimal.h"
+
+namespace UE::ChooserEditor
+{
+	void RegisterOutputNameWidgets();
+}
