fix crash in FloatRangeColumnEditor

#Synced-CL 51640274

diff --git a/Engine/Plugins/Chooser/Source/ChooserEditor/Private/BoolColumnEditor.cpp b/Engine/Plugins/Chooser/Source/ChooserEditor/Private/BoolColumnEditor.cpp
--- a/Engine/Plugins/Chooser/Source/ChooserEditor/Private/BoolColumnEditor.cpp
+++ b/Engine/Plugins/Chooser/Source/ChooserEditor/Private/BoolColumnEditor.cpp
@@ -139,7 +139,7 @@ TSharedRef<SWidget> CreateOutputBoolColumnWidget(UChooserTable* Chooser, FChoose
 		SNew (SCheckBox)
 		.OnCheckStateChanged_Lambda([Chooser, BoolColumn,Row](ECheckBoxState State)
 		{
-			if (Row < BoolColumn->RowValues.Num())
+			if (BoolColumn->RowValues.IsValidIndex(Row))
 			{
 				const FScopedTransaction Transaction(LOCTEXT("Change Bool Value", "Change Bool Value"));
 				Chooser->Modify(true);
@@ -148,7 +148,7 @@ TSharedRef<SWidget> CreateOutputBoolColumnWidget(UChooserTable* Chooser, FChoose
 		})
 		.IsChecked_Lambda([BoolColumn, Row]()
 		{
-			const bool value = (Row < BoolColumn->RowValues.Num()) ? BoolColumn->RowValues[Row] : false;
+			const bool value = (BoolColumn->RowValues.IsValidIndex(Row)) ? BoolColumn->RowValues[Row] : false;
 			return value ? ECheckBoxState::Checked : ECheckBoxState::Unchecked;
 		})
 	]
diff --git a/Engine/Plugins/Chooser/Source/ChooserEditor/Private/EnumColumnEditor.cpp b/Engine/Plugins/Chooser/Source/ChooserEditor/Private/EnumColumnEditor.cpp
--- a/Engine/Plugins/Chooser/Source/ChooserEditor/Private/EnumColumnEditor.cpp
+++ b/Engine/Plugins/Chooser/Source/ChooserEditor/Private/EnumColumnEditor.cpp
@@ -123,16 +123,19 @@ TSharedRef<SWidget> CreateEnumColumnWidget(UChooserTable* Chooser, FChooserColum
     				.Visibility(Row < 0 ? EVisibility::Hidden : EVisibility::Visible)
 					.Text_Lambda([EnumColumn, Row]()
 					{
-						switch (EnumColumn->RowValues[Row].Comparison)
+						if (EnumColumn->RowValues.IsValidIndex(Row))
 						{
-						case EEnumColumnCellValueComparison::MatchEqual:
-							return LOCTEXT("CompEqual", "=");
+							switch (EnumColumn->RowValues[Row].Comparison)
+							{
+							case EEnumColumnCellValueComparison::MatchEqual:
+								return LOCTEXT("CompEqual", "=");
 
-						case EEnumColumnCellValueComparison::MatchNotEqual:
-							return LOCTEXT("CompNotEqual", "Not");
+							case EEnumColumnCellValueComparison::MatchNotEqual:
+								return LOCTEXT("CompNotEqual", "Not");
 
-						case EEnumColumnCellValueComparison::MatchAny:
-							return LOCTEXT("CompAny", "Any");
+							case EEnumColumnCellValueComparison::MatchAny:
+								return LOCTEXT("CompAny", "Any");
+							}
 						}
 						return FText::GetEmpty();
 					})
diff --git a/Engine/Plugins/Chooser/Source/ChooserEditor/Private/FloatDistanceColumnEditor.cpp b/Engine/Plugins/Chooser/Source/ChooserEditor/Private/FloatDistanceColumnEditor.cpp
--- a/Engine/Plugins/Chooser/Source/ChooserEditor/Private/FloatDistanceColumnEditor.cpp
+++ b/Engine/Plugins/Chooser/Source/ChooserEditor/Private/FloatDistanceColumnEditor.cpp
@@ -33,7 +33,7 @@ TSharedRef<SWidget> CreateFloatDistanceColumnWidget(UChooserTable* Chooser, FCho
 		// create column header widget
 		
 		const FSlateBrush* ColumnIcon = FCoreStyle::Get().GetBrush("Icons.SortUp");
-		const FText ColumnTooltip = LOCTEXT("Float difference tooltip", "Float Difference Column: rows recieve a Score based on how different the input float is from the row value");
+		const FText ColumnTooltip = LOCTEXT("Float difference tooltip", "Float Difference Column: rows receive a Score based on how different the input float is from the row value");
 		const FText ColumnName = LOCTEXT("Float Difference","Float Difference");
 		
 		TSharedPtr<SWidget> DebugWidget = nullptr;
@@ -52,11 +52,11 @@ TSharedRef<SWidget> CreateFloatDistanceColumnWidget(UChooserTable* Chooser, FCho
 	return SNew(SNumericEntryBox<float>)
 	.Value_Lambda([FloatDistanceColumn, Row]()
 	{
-		return (Row < FloatDistanceColumn->RowValues.Num()) ? FloatDistanceColumn->RowValues[Row].Value : 0;
+		return (FloatDistanceColumn->RowValues.IsValidIndex(Row)) ? FloatDistanceColumn->RowValues[Row].Value : 0;
 	})
 	.OnValueCommitted_Lambda([Chooser, FloatDistanceColumn, Row](float NewValue, ETextCommit::Type CommitType)
 	{
-		if (Row < FloatDistanceColumn->RowValues.Num())
+		if (FloatDistanceColumn->RowValues.IsValidIndex(Row))
 		{
 			const FScopedTransaction Transaction(LOCTEXT("Edit Float Distance Value", "Edit Float Distance Value"));
 			Chooser->Modify(true);
diff --git a/Engine/Plugins/Chooser/Source/ChooserEditor/Private/FloatRangeColumnEditor.cpp b/Engine/Plugins/Chooser/Source/ChooserEditor/Private/FloatRangeColumnEditor.cpp
--- a/Engine/Plugins/Chooser/Source/ChooserEditor/Private/FloatRangeColumnEditor.cpp
+++ b/Engine/Plugins/Chooser/Source/ChooserEditor/Private/FloatRangeColumnEditor.cpp
@@ -78,19 +78,19 @@ TSharedRef<SWidget> CreateFloatRangeColumnWidget(UChooserTable* Chooser, FChoose
 			SNew(SNumericEntryBox<float>)
 			.Visibility_Lambda([FloatRangeColumn, Row]
 			{
-				return FloatRangeColumn->RowValues[Row].bNoMin ? EVisibility::Hidden : EVisibility::Visible;
+				return (FloatRangeColumn->RowValues.IsValidIndex(Row) && FloatRangeColumn->RowValues[Row].bNoMin) ? EVisibility::Hidden : EVisibility::Visible;
 			})
 			.MaxValue_Lambda([FloatRangeColumn, Row]()
 			{
-				return (Row < FloatRangeColumn->RowValues.Num()) ? FloatRangeColumn->RowValues[Row].Max : 0;
+				return (FloatRangeColumn->RowValues.IsValidIndex(Row)) ? FloatRangeColumn->RowValues[Row].Max : 0;
 			})
 			.Value_Lambda([FloatRangeColumn, Row]()
 			{
-				return (Row < FloatRangeColumn->RowValues.Num()) ? FloatRangeColumn->RowValues[Row].Min : 0;
+				return (FloatRangeColumn->RowValues.IsValidIndex(Row)) ? FloatRangeColumn->RowValues[Row].Min : 0;
 			})
 			.OnValueCommitted_Lambda([Chooser, FloatRangeColumn, Row](float NewValue, ETextCommit::Type CommitType)
 			{
-				if (Row < FloatRangeColumn->RowValues.Num())
+				if (FloatRangeColumn->RowValues.IsValidIndex(Row))
 				{
 					const FScopedTransaction Transaction(LOCTEXT("Edit Min Value", "Edit Min Value"));
 					Chooser->Modify(true);
@@ -103,7 +103,7 @@ TSharedRef<SWidget> CreateFloatRangeColumnWidget(UChooserTable* Chooser, FChoose
 			SNew(SEditableTextBox).IsReadOnly(true).Text(LOCTEXT("negativeinfinity","-∞"))
 				.Visibility_Lambda([FloatRangeColumn, Row]
 				{
-					return FloatRangeColumn->RowValues[Row].bNoMin ? EVisibility::Visible : EVisibility::Hidden;
+					return (FloatRangeColumn->RowValues.IsValidIndex(Row) && FloatRangeColumn->RowValues[Row].bNoMin) ? EVisibility::Visible : EVisibility::Hidden;
 				})
 		]
 	]
@@ -119,19 +119,19 @@ TSharedRef<SWidget> CreateFloatRangeColumnWidget(UChooserTable* Chooser, FChoose
 			SNew(SNumericEntryBox<float>)
 			.Visibility_Lambda([FloatRangeColumn, Row]
 			{
-				return FloatRangeColumn->RowValues[Row].bNoMax ? EVisibility::Hidden : EVisibility::Visible;
+				return (FloatRangeColumn->RowValues.IsValidIndex(Row) && FloatRangeColumn->RowValues[Row].bNoMax) ? EVisibility::Hidden : EVisibility::Visible;
 			})
 			.MinValue_Lambda([FloatRangeColumn, Row]()
 			{
-				return (Row < FloatRangeColumn->RowValues.Num()) ? FloatRangeColumn->RowValues[Row].Min : 0;
+				return (FloatRangeColumn->RowValues.IsValidIndex(Row)) ? FloatRangeColumn->RowValues[Row].Min : 0;
 			})
 			.Value_Lambda([FloatRangeColumn, Row]()
 			{
-				return (Row < FloatRangeColumn->RowValues.Num()) ? FloatRangeColumn->RowValues[Row].Max : 0;
+				return (FloatRangeColumn->RowValues.IsValidIndex(Row)) ? FloatRangeColumn->RowValues[Row].Max : 0;
 			})
 			.OnValueCommitted_Lambda([Chooser, FloatRangeColumn, Row](float NewValue, ETextCommit::Type CommitType)
 			{
-				if (Row < FloatRangeColumn->RowValues.Num())
+				if (FloatRangeColumn->RowValues.IsValidIndex(Row))
 				{
 					const FScopedTransaction Transaction(LOCTEXT("Edit Max", "Edit Max Value"));
 					Chooser->Modify(true);
@@ -144,7 +144,7 @@ TSharedRef<SWidget> CreateFloatRangeColumnWidget(UChooserTable* Chooser, FChoose
 			SNew(SEditableTextBox).IsReadOnly(true).Text(LOCTEXT("infinity","∞"))
 				.Visibility_Lambda([FloatRangeColumn, Row]
 				{
-					return FloatRangeColumn->RowValues[Row].bNoMax ? EVisibility::Visible : EVisibility::Hidden;
+					return  (FloatRangeColumn->RowValues.IsValidIndex(Row) && FloatRangeColumn->RowValues[Row].bNoMax) ? EVisibility::Visible : EVisibility::Hidden;
 				})
 		]
 	]
diff --git a/Engine/Plugins/Chooser/Source/ChooserEditor/Private/GameplayTagColumnEditor.cpp b/Engine/Plugins/Chooser/Source/ChooserEditor/Private/GameplayTagColumnEditor.cpp
--- a/Engine/Plugins/Chooser/Source/ChooserEditor/Private/GameplayTagColumnEditor.cpp
+++ b/Engine/Plugins/Chooser/Source/ChooserEditor/Private/GameplayTagColumnEditor.cpp
@@ -60,7 +60,7 @@ TSharedRef<SWidget> CreateGameplayTagColumnWidget(UChooserTable* Chooser, FChoos
 		.ReadOnly(false)
 		.TagContainer_Lambda([GameplayTagColumn, Row]()
 		{
-			if (Row < GameplayTagColumn->RowValues.Num())
+			if (GameplayTagColumn->RowValues.IsValidIndex(Row))
 			{
 				return GameplayTagColumn->RowValues[Row];
 			}
@@ -71,7 +71,7 @@ TSharedRef<SWidget> CreateGameplayTagColumnWidget(UChooserTable* Chooser, FChoos
 		})
 		.OnTagContainerChanged_Lambda([GameplayTagColumn, Row](const FGameplayTagContainer& UpdatedTags)
 		{
-			if (Row < GameplayTagColumn->RowValues.Num())
+			if (GameplayTagColumn->RowValues.IsValidIndex(Row))
 			{
 				GameplayTagColumn->RowValues[Row] = UpdatedTags;
 			}
@@ -103,7 +103,7 @@ TSharedRef<SWidget> CreateOutputGameplayTagColumnWidget(UChooserTable* Chooser,
 		.ReadOnly(false)
 		.TagContainer_Lambda([OutputGameplayTagColumn, Row]()
 		{
-			if (Row < OutputGameplayTagColumn->RowValues.Num())
+			if (OutputGameplayTagColumn->RowValues.IsValidIndex(Row))
 			{
 				return OutputGameplayTagColumn->RowValues[Row];
 			}
@@ -114,10 +114,10 @@ TSharedRef<SWidget> CreateOutputGameplayTagColumnWidget(UChooserTable* Chooser,
 		})
 		.OnTagContainerChanged_Lambda([OutputGameplayTagColumn, Row](const FGameplayTagContainer& UpdatedTags)
 		{
-				if (Row < OutputGameplayTagColumn->RowValues.Num())
-				{
-					OutputGameplayTagColumn->RowValues[Row] = UpdatedTags;
-				}
+			if (OutputGameplayTagColumn->RowValues.IsValidIndex(Row))
+			{
+				OutputGameplayTagColumn->RowValues[Row] = UpdatedTags;
+			}
 		}
 	);
 }
diff --git a/Engine/Plugins/Chooser/Source/ChooserEditor/Private/GameplayTagQueryColumnEditor.cpp b/Engine/Plugins/Chooser/Source/ChooserEditor/Private/GameplayTagQueryColumnEditor.cpp
--- a/Engine/Plugins/Chooser/Source/ChooserEditor/Private/GameplayTagQueryColumnEditor.cpp
+++ b/Engine/Plugins/Chooser/Source/ChooserEditor/Private/GameplayTagQueryColumnEditor.cpp
@@ -65,7 +65,7 @@ TSharedRef<SWidget> CreateGameplayTagQueryColumnWidget(UChooserTable* Chooser, F
 	return SNew(SGameplayTagQueryEntryBox)
 		.TagQuery_Lambda([GameplayTagQueryColumn, Row]()
 		{
-			if (Row < GameplayTagQueryColumn->RowValues.Num())
+			if (GameplayTagQueryColumn->RowValues.IsValidIndex(Row))
 			{
 				return GameplayTagQueryColumn->RowValues[Row];
 			}
@@ -75,7 +75,7 @@ TSharedRef<SWidget> CreateGameplayTagQueryColumnWidget(UChooserTable* Chooser, F
 		.DescriptionMaxWidth(1000.f)
 		.OnTagQueryChanged_Lambda([GameplayTagQueryColumn, Row](const FGameplayTagQuery& UpdatedQuery)
 		{
-			if (Row < GameplayTagQueryColumn->RowValues.Num())
+			if (GameplayTagQueryColumn->RowValues.IsValidIndex(Row))
 			{
 				GameplayTagQueryColumn->RowValues[Row] = UpdatedQuery;
 			}
@@ -105,7 +105,7 @@ TSharedRef<SWidget> CreateOutputGameplayTagQueryColumnWidget(UChooserTable* Choo
 	return SNew(SGameplayTagQueryEntryBox)
 		.TagQuery_Lambda([OutputGameplayTagQueryColumn, Row]()
 		{
-			if (Row < OutputGameplayTagQueryColumn->RowValues.Num())
+			if (OutputGameplayTagQueryColumn->RowValues.IsValidIndex(Row))
 			{
 				return OutputGameplayTagQueryColumn->RowValues[Row];
 			}
@@ -115,7 +115,7 @@ TSharedRef<SWidget> CreateOutputGameplayTagQueryColumnWidget(UChooserTable* Choo
 		.DescriptionMaxWidth(1000.f)
 		.OnTagQueryChanged_Lambda([OutputGameplayTagQueryColumn, Row](const FGameplayTagQuery& UpdatedQuery)
 		{
-			if (Row < OutputGameplayTagQueryColumn->RowValues.Num())
+			if (OutputGameplayTagQueryColumn->RowValues.IsValidIndex(Row))
 			{
 				OutputGameplayTagQueryColumn->RowValues[Row] = UpdatedQuery;
 			}
diff --git a/Engine/Plugins/Chooser/Source/ChooserEditor/Private/OutputFloatColumnEditor.cpp b/Engine/Plugins/Chooser/Source/ChooserEditor/Private/OutputFloatColumnEditor.cpp
--- a/Engine/Plugins/Chooser/Source/ChooserEditor/Private/OutputFloatColumnEditor.cpp
+++ b/Engine/Plugins/Chooser/Source/ChooserEditor/Private/OutputFloatColumnEditor.cpp
@@ -58,11 +58,11 @@ TSharedRef<SWidget> CreateOutputFloatColumnWidget(UChooserTable* Chooser, FChoos
 	return SNew(SNumericEntryBox<double>)
     		.Value_Lambda([OutputFloatColumn, Row]()
     		{
-    			return (Row < OutputFloatColumn->RowValues.Num()) ? OutputFloatColumn->RowValues[Row] : 0;
+    			return (OutputFloatColumn->RowValues.IsValidIndex(Row)) ? OutputFloatColumn->RowValues[Row] : 0;
     		})
     		.OnValueCommitted_Lambda([Chooser, OutputFloatColumn, Row](double NewValue, ETextCommit::Type CommitType)
     		{
-    			if (Row < OutputFloatColumn->RowValues.Num())
+    			if (OutputFloatColumn->RowValues.IsValidIndex(Row))
     			{
     				const FScopedTransaction Transaction(LOCTEXT("Edit Float Value", "Edit Float Value"));
     				Chooser->Modify(true);
diff --git a/Engine/Plugins/Chooser/Source/ChooserEditor/Private/OutputStructColumnEditor.cpp b/Engine/Plugins/Chooser/Source/ChooserEditor/Private/OutputStructColumnEditor.cpp
--- a/Engine/Plugins/Chooser/Source/ChooserEditor/Private/OutputStructColumnEditor.cpp
+++ b/Engine/Plugins/Chooser/Source/ChooserEditor/Private/OutputStructColumnEditor.cpp
@@ -50,23 +50,27 @@ TSharedRef<SWidget> CreateOutputStructColumnWidget(UChooserTable* Chooser, FChoo
 			}
 			else
 			{
-				RowValue = &StructColumn->RowValues[Row];
+				if (StructColumn->RowValues.IsValidIndex(Row))
+				{
+					RowValue = &StructColumn->RowValues[Row];
+				}
 			}
 
-
 			FString Value;
-			if (const UScriptStruct* ScriptStruct = RowValue->GetScriptStruct())
-			{
-				void* DefaultStructMemory = FMemory_Alloca_Aligned(ScriptStruct->GetStructureSize(), ScriptStruct->GetMinAlignment());
-				ScriptStruct->InitializeStruct(DefaultStructMemory);
-				ScriptStruct->ExportText(Value, RowValue->GetMemory(), DefaultStructMemory, nullptr, PPF_ExternalEditor, nullptr);
-				ScriptStruct->DestroyStruct(DefaultStructMemory);
-			}
-			else
+			if (RowValue)
 			{
-				Value = TEXT("()");
+				if (const UScriptStruct* ScriptStruct = RowValue->GetScriptStruct())
+				{
+					void* DefaultStructMemory = FMemory_Alloca_Aligned(ScriptStruct->GetStructureSize(), ScriptStruct->GetMinAlignment());
+					ScriptStruct->InitializeStruct(DefaultStructMemory);
+					ScriptStruct->ExportText(Value, RowValue->GetMemory(), DefaultStructMemory, nullptr, PPF_ExternalEditor, nullptr);
+					ScriptStruct->DestroyStruct(DefaultStructMemory);
+					
+					return FText::FromString(Value);
+				}
 			}
-			return FText::FromString(Value);
+			
+			return FText::FromString(TEXT("()"));
 		});
 
 	TSharedRef<STextBlock> TextBlock = SNew(STextBlock)
