Chooser - Get an Object from a ChooserEvaluation context by class or interface

#Synced-CL 51737212

diff --git a/Engine/Plugins/Chooser/Source/Chooser/Private/ChooserFunctionLibrary.cpp b/Engine/Plugins/Chooser/Source/Chooser/Private/ChooserFunctionLibrary.cpp
--- a/Engine/Plugins/Chooser/Source/Chooser/Private/ChooserFunctionLibrary.cpp
+++ b/Engine/Plugins/Chooser/Source/Chooser/Private/ChooserFunctionLibrary.cpp
@@ -71,6 +71,32 @@ UObject* UChooserFunctionLibrary::GetChooserObjectInput(FChooserEvaluationContex
 	return nullptr;
 }
 
+UObject* UChooserFunctionLibrary::GetChooserObject(FChooserEvaluationContext& Context, TSubclassOf<UObject> ObjectClass)
+{
+	for (FStructView& ContextParam : Context.Params)
+	{
+		if (FChooserEvaluationInputObject* ObjectParam = ContextParam.GetPtr<FChooserEvaluationInputObject>())
+		{
+			if (UObject* ObjectParamObject = ObjectParam->Object)
+			{
+				if (ObjectParamObject->IsA(ObjectClass))
+				{
+					return ObjectParamObject;
+				}
+
+				if (const UClass* ObjectParamObjectClass = ObjectParamObject->GetClass())
+				{
+					if (ObjectParamObjectClass->ImplementsInterface(ObjectClass))
+					{
+						return ObjectParamObject;
+					}
+				}
+			}
+		}
+	}
+	return nullptr;
+}
+
 void UChooserFunctionLibrary::AddChooserStructInput(FChooserEvaluationContext& Inputs, int32 Value)
 {
 	checkNoEntry();
diff --git a/Engine/Plugins/Chooser/Source/Chooser/Public/ChooserFunctionLibrary.h b/Engine/Plugins/Chooser/Source/Chooser/Public/ChooserFunctionLibrary.h
--- a/Engine/Plugins/Chooser/Source/Chooser/Public/ChooserFunctionLibrary.h
+++ b/Engine/Plugins/Chooser/Source/Chooser/Public/ChooserFunctionLibrary.h
@@ -86,6 +86,16 @@ public:
 	*/
 	UFUNCTION(BlueprintPure, meta = (BlueprintThreadSafe ,BlueprintInternalUseOnly = "true", DeterminesOutputType = "ObjectClass"))
 	static UObject* GetChooserObjectInput(UPARAM(Ref) FChooserEvaluationContext& Context, int32 Index, TSubclassOf<UObject> ObjectClass);
+
+	/**
+	* Get an Object from a ChooserEvaluation context by class or interface
+	* Experimental, this feature might be removed without warning, not for production use
+	*
+	* @param Context			(in) A struct reference to the chooser evaluation context
+	* @param ObjectClass		(in) the wanted object class or interface
+	*/
+	UFUNCTION(BlueprintPure, Category = "Animation | Experimental", meta = (BlueprintThreadSafe))
+	static UObject* GetChooserObject(UPARAM(Ref) FChooserEvaluationContext& Context, TSubclassOf<UObject> ObjectClass);
 	
 	/**
 	* Add a Struct to a ChooserEvaluation context
