Chooser - fix for MM col resetting their rows database properties when moving columns

#Synced-CL 52029572

diff --git a/Engine/Plugins/Animation/PoseSearch/Source/Runtime/Private/Chooser/PoseSearchChooserColumn.cpp b/Engine/Plugins/Animation/PoseSearch/Source/Runtime/Private/Chooser/PoseSearchChooserColumn.cpp
--- a/Engine/Plugins/Animation/PoseSearch/Source/Runtime/Private/Chooser/PoseSearchChooserColumn.cpp
+++ b/Engine/Plugins/Animation/PoseSearch/Source/Runtime/Private/Chooser/PoseSearchChooserColumn.cpp
@@ -46,19 +46,23 @@ namespace UE::PoseSearch
 	static const FInstancedStruct& GetResult(const UChooserTable* Chooser, int32 RowIndex)
 	{
 		check(Chooser);
-		if (RowIndex == ChooserColumn_SpecialIndex_Fallback)
-		{
-			return Chooser->FallbackResult;
-		}
 
 #if WITH_EDITORONLY_DATA
 		if (!Chooser->IsCookedData())
 		{
-			return Chooser->ResultsStructs[RowIndex];
+			if (Chooser->ResultsStructs.IsValidIndex(RowIndex))
+			{
+				return Chooser->ResultsStructs[RowIndex];
+			}
+			return Chooser->FallbackResult;
 		}
 #endif // WITH_EDITORONLY_DATA
 		
-		return Chooser->CookedResults[RowIndex];
+		if (Chooser->CookedResults.IsValidIndex(RowIndex))
+		{
+			return Chooser->CookedResults[RowIndex];
+		}
+		return Chooser->FallbackResult;
 	}
 
 #if WITH_EDITOR
@@ -324,13 +328,8 @@ namespace UE::PoseSearch
 
 					FPoseSearchDatabaseAnimationAsset& DatabaseAnimationAsset = TableCell.Column->EditRowValue(TableCell.RowIndex).Data;
 					DatabaseAnimationAsset = AnimationAsset;
-
 					// enforcing RowValue.Data.AnimAsset to be equal to ReferencedObject
-					UObject* ReferencedObject = GetReferencedObject(TableCell.Column->GetChooser(), TableCell.RowIndex);
-					if (ReferencedObject != DatabaseAnimationAsset.AnimAsset)
-					{
-						DatabaseAnimationAsset.AnimAsset = ReferencedObject;
-					}
+					DatabaseAnimationAsset.AnimAsset = GetReferencedObject(TableCell.Column->GetChooser(), TableCell.RowIndex);
 
 					DO_ENSURE_CONSISTENCY(CheckConsistency(TableCell.Column->GetChooser()));
 				}
@@ -471,7 +470,7 @@ const FPoseSearchColumnRow& FPoseSearchColumn::GetRowValue(int32 RowIndex) const
 	{
 		return FallbackValue;
 	}
-	if (ensure(RowValues.IsValidIndex(RowIndex)))
+	if (RowValues.IsValidIndex(RowIndex))
 	{
 		return RowValues[RowIndex];
 	}
@@ -1066,8 +1065,13 @@ void FPoseSearchColumn::Initialize(UChooserTable* OuterChooser)
 	ensure(OuterChooser);
 	FChooserColumnBase::Initialize(OuterChooser);
 
-	RowValues.Reset();
-	FallbackValue = FPoseSearchColumnRow();
+	// if Chooser is null we don't reset the rows, since the initialize has been called on a
+	// newly copied column via FChooserTableViewModel::CopyColumnInternal
+	if (Chooser && Chooser != OuterChooser)
+	{
+		RowValues.Reset();
+		FallbackValue = FPoseSearchColumnRow();
+	}
 
 	Chooser = OuterChooser;
 	FMapping(Chooser->GetRootChooser()).ApplyMappingToAllColumns();
@@ -1149,19 +1153,17 @@ void FPoseSearchColumn::SetFromDetails(FInstancedPropertyBag& PropertyBag, int32
 	if (FStructView* StructView = Result.TryGetValue())
 	{
 		const FPoseSearchColumnRowReflection& RowReflection = StructView->Get<FPoseSearchColumnRowReflection>();
-		const FPoseSearchColumnRow& RowReflectionBase = static_cast<const FPoseSearchColumnRow&>(RowReflection);
+		FPoseSearchColumnRow RowReflectionBase = static_cast<const FPoseSearchColumnRow&>(RowReflection);
+		RowReflectionBase.Data.AnimAsset = GetReferencedObject(GetChooser(), RowIndex);
+
 		FPoseSearchColumnRow& RowValue = EditRowValue(RowIndex);
 
 		ensure(!RowReflectionBase.Database || RowReflectionBase.Database->GetOuter() == Chooser->GetRootChooser());
 
-		UObject* AnimAsset = GetReferencedObject(GetChooser(), RowIndex);
 		if (RowValue != RowReflectionBase)
 		{
 			Chooser->Modify();
-
 			RowValue = RowReflectionBase;
-			RowValue.Data.AnimAsset = GetReferencedObject(GetChooser(), RowIndex);
-
 			FMapping(Chooser->GetRootChooser()).ApplyMappingToAllColumns();
 			Chooser->PostEditChange();
 		}
@@ -1349,7 +1351,10 @@ void FPoseSearchColumn::CopyRow(FChooserColumnBase& SourceColumn, int SourceRowI
 		SourceRow.Database->Modify();
 	}
 
-	EditRowValue(TargetRowIndex) = SourceRow;
+	FPoseSearchColumnRow& TargetRow = EditRowValue(TargetRowIndex);
+	TargetRow = SourceRow;
+	// resynchronizing the TargetRow.RowValue.Data with this column Chooser reference object associated with the RowIndex
+	TargetRow.Data.AnimAsset = GetReferencedObject(Chooser, TargetRowIndex);
 
 	bool bPerformDatabasePostLoad = false;
 	bool bPerformPostEditChange = false;
diff --git a/Engine/Plugins/Chooser/Source/Chooser/Public/IChooserColumn.h b/Engine/Plugins/Chooser/Source/Chooser/Public/IChooserColumn.h
--- a/Engine/Plugins/Chooser/Source/Chooser/Public/IChooserColumn.h
+++ b/Engine/Plugins/Chooser/Source/Chooser/Public/IChooserColumn.h
@@ -173,7 +173,10 @@ enum {  ChooserColumn_SpecialIndex_Fallback = -2 };
 	}\
 	virtual void CopyRow(FChooserColumnBase& SourceColumn, int SourceRowIndex, int TargetRowIndex) override\
 	{\
-		RowValuesProperty[TargetRowIndex]=static_cast<decltype(this)>(&SourceColumn)->RowValuesProperty[SourceRowIndex];\
+		if (RowValuesProperty.IsValidIndex(TargetRowIndex) && static_cast<decltype(this)>(&SourceColumn)->RowValuesProperty.IsValidIndex(SourceRowIndex)) \
+		{\
+			RowValuesProperty[TargetRowIndex]=static_cast<decltype(this)>(&SourceColumn)->RowValuesProperty[SourceRowIndex];\
+		}\
 	}\
 	virtual UScriptStruct* GetInputBaseType() const override { return ParameterType::StaticStruct(); };\
 	virtual const UScriptStruct* GetInputType() const override { return InputValue.IsValid() ? InputValue.GetScriptStruct() : nullptr; };\
diff --git a/Engine/Plugins/Chooser/Source/ChooserEditor/Private/ChooserTableViewModel.cpp b/Engine/Plugins/Chooser/Source/ChooserEditor/Private/ChooserTableViewModel.cpp
--- a/Engine/Plugins/Chooser/Source/ChooserEditor/Private/ChooserTableViewModel.cpp
+++ b/Engine/Plugins/Chooser/Source/ChooserEditor/Private/ChooserTableViewModel.cpp
@@ -325,7 +325,7 @@ int32 FChooserTableViewModel::MoveColumn(int32 SourceIndex, int32 TargetIndex)
 
 	Chooser->Modify(true);
 
-	FInstancedStruct ColumnData = Chooser->ColumnsStructs[SourceIndex];
+	FInstancedStruct ColumnData = MoveTemp(Chooser->ColumnsStructs[SourceIndex]);
 	Chooser->ColumnsStructs.RemoveAt(SourceIndex);
 		
 	if (SourceIndex < TargetIndex)
@@ -342,7 +342,7 @@ int32 FChooserTableViewModel::MoveColumn(int32 SourceIndex, int32 TargetIndex)
 		}
 	}
 
-	Chooser->ColumnsStructs.Insert(ColumnData, TargetIndex);
+	Chooser->ColumnsStructs.Insert(MoveTemp(ColumnData), TargetIndex);
 
 	RefreshAll();
 
@@ -373,7 +373,7 @@ void FChooserTableViewModel::AddColumn(const UScriptStruct* ColumnType)
 		InsertIndex = Chooser->ColumnsStructs.Num();
 		if (InsertIndex == 0 || !Chooser->ColumnsStructs[InsertIndex - 1].Get<FChooserColumnBase>().IsRandomizeColumn())
 		{
-			Chooser->ColumnsStructs.Add(NewColumn);
+			Chooser->ColumnsStructs.Add(MoveTemp(NewColumn));
 		}
 	}
 	else if (NewColumnRef.HasOutputs())
@@ -384,7 +384,7 @@ void FChooserTableViewModel::AddColumn(const UScriptStruct* ColumnType)
 		{
 			InsertIndex--;
 		}
-		Chooser->ColumnsStructs.Insert(NewColumn, InsertIndex);
+		Chooser->ColumnsStructs.Insert(MoveTemp(NewColumn), InsertIndex);
 	}
 	else
 	{
@@ -398,7 +398,7 @@ void FChooserTableViewModel::AddColumn(const UScriptStruct* ColumnType)
 			}
 			InsertIndex++;
 		}
-		Chooser->ColumnsStructs.Insert(NewColumn, InsertIndex);
+		Chooser->ColumnsStructs.Insert(MoveTemp(NewColumn), InsertIndex);
 	}
 
 	UpdateTableColumns();
@@ -1053,9 +1053,21 @@ UChooserTable* FChooserTableViewModel::CopyColumnInternal(int32 ColumnIndex)
 	CopyData->ResultType = RootChooser->ResultType;
 	CopyData->ContextData = RootChooser->ContextData;
 
-	// add selected column including all the cell data
+	// add selected column including all the cell data, using CopyRow to properly transfer
+	// per-row references (e.g. database pointers) into CopyData's NestedObjects so they
+	// survive text serialization to the clipboard.
 	CopyData->ColumnsStructs.Add(Chooser->ColumnsStructs[ColumnIndex]);
-	
+	FChooserColumnBase& SourceColumn = Chooser->ColumnsStructs[ColumnIndex].GetMutable<FChooserColumnBase>();
+	FChooserColumnBase& CopiedColumn = CopyData->ColumnsStructs.Last().GetMutable<FChooserColumnBase>();
+	CopiedColumn.Initialize(CopyData);
+	CopiedColumn.SetNumRows(0);
+	CopiedColumn.SetNumRows(Chooser->ResultsStructs.Num());
+	for (int32 RowIndex = 0; RowIndex < Chooser->ResultsStructs.Num(); RowIndex++)
+	{
+		CopiedColumn.CopyRow(SourceColumn, RowIndex, RowIndex);
+	}
+	CopiedColumn.CopyFallback(SourceColumn);
+
 	return CopyData;
 }
 	
@@ -1115,6 +1127,14 @@ protected:
 		if (CreatedObject->IsA<UChooserTable>())
 		{
 			ClipboardContent = CastChecked<UChooserTable>(CreatedObject);
+
+			for (FInstancedStruct& ColumnData : ClipboardContent->ColumnsStructs)
+			{
+				if (FChooserColumnBase* Column = ColumnData.GetMutablePtr<FChooserColumnBase>())
+				{
+					Column->Initialize(ClipboardContent);
+				}
+			}
 		}
 	}
 };
@@ -1150,8 +1170,24 @@ void FChooserTableViewModel::PasteInternal(UChooserTable* PastedContent, int32 P
 				InsertColumnIndex--;
 			}
 		}
-		Chooser->ColumnsStructs.Insert(PastedContent->ColumnsStructs, InsertColumnIndex);
-		Chooser->ColumnsStructs[InsertColumnIndex].GetMutable<FChooserColumnBase>().Initialize(Chooser);
+
+		int32 CurrentInsertColumnIndex = InsertColumnIndex;
+		for (FInstancedStruct& ColumnData : PastedContent->ColumnsStructs)
+		{
+			FChooserColumnBase& PastedColumn = ColumnData.GetMutable<FChooserColumnBase>();
+			Chooser->ColumnsStructs.Insert(ColumnData, CurrentInsertColumnIndex);
+			FChooserColumnBase& Column = Chooser->ColumnsStructs[CurrentInsertColumnIndex].GetMutable<FChooserColumnBase>();
+			Column.Initialize(Chooser);
+			Column.SetNumRows(0);
+			Column.SetNumRows(Chooser->ResultsStructs.Num());
+			for (int32 RowIndex = 0; RowIndex < Chooser->ResultsStructs.Num(); RowIndex++)
+			{
+				Column.CopyRow(PastedColumn, RowIndex, RowIndex);
+			}
+			Column.CopyFallback(PastedColumn);
+			++CurrentInsertColumnIndex;
+		}
+
 		SelectColumn(Chooser, InsertColumnIndex);
 	}
 	else
