Fix hash determinism

#Synced-CL 52154273

diff --git a/Engine/Plugins/Chooser/Source/Chooser/Private/AnimNode_ChooserPlayer.cpp b/Engine/Plugins/Chooser/Source/Chooser/Private/AnimNode_ChooserPlayer.cpp
--- a/Engine/Plugins/Chooser/Source/Chooser/Private/AnimNode_ChooserPlayer.cpp
+++ b/Engine/Plugins/Chooser/Source/Chooser/Private/AnimNode_ChooserPlayer.cpp
@@ -10,12 +10,25 @@
 #include "Animation/BlendSpace.h"
 #include "BlendStack/AnimNode_BlendStack.h"
 #include "IObjectChooser.h"
+#include "Serialization/ArchiveCrc32.h"
 
 #include UE_INLINE_GENERATED_CPP_BY_NAME(AnimNode_ChooserPlayer)
 
+uint32 GetTypeHash(const FAnimCurveOverride& CurveOverride)
+{
+	return HashCombine(GetTypeHash(CurveOverride.CurveName.ToString()), GetTypeHash(CurveOverride.CurveValue));
+}
+
 void FAnimCurveOverrideList::ComputeHash()
 {
-	Hash = UE::StructUtils::GetStructCrc32(*StaticStruct(), reinterpret_cast<const uint8*>(this));
+	if (Values.Num() == 0)
+	{
+		Hash = 0;
+	}
+	else
+	{
+		Hash = GetArrayHash(&Values[0], Values.Num());
+	}
 }
 
 FAnimNode_ChooserPlayer::FAnimNode_ChooserPlayer()
