Proxy and chooser fixes
- Add support for multiple struct outputs and soft object results to EvaluateProxy blueprint node
- Add multiple struct output support for choosers returning soft object results
- prevent crashing when chooser tables and proxy tables have a null result
- change struct output order in proxy evaluation to match new standard (proxy outputs first, and is overwritten if the chooser has outputs)

#Synced-CL 52174097

#Synced-CL 52174097

diff --git a/Engine/Plugins/Chooser/Source/Chooser/Private/Chooser.cpp b/Engine/Plugins/Chooser/Source/Chooser/Private/Chooser.cpp
--- a/Engine/Plugins/Chooser/Source/Chooser/Private/Chooser.cpp
+++ b/Engine/Plugins/Chooser/Source/Chooser/Private/Chooser.cpp
@@ -737,23 +737,25 @@ FObjectChooserBase::EIteratorStatus UChooserTable::EvaluateChooser(FChooserEvalu
 				}
 				check(ScratchAreaStart == TotalScratchAreaSize);
 			}
-			
-			const FObjectChooserBase& SelectedResult = (*ResultsArray)[SelectedIndexData.Index].Get<FObjectChooserBase>();
-			FObjectChooserBase::EIteratorStatus Status = SelectedResult.ChooseMulti(Context, Callback);
-			
-			if (Status != FObjectChooserBase::EIteratorStatus::Failed)
-			{
-				bAnyRowSucceeded = true;
-			}
 
-			if (Status == FObjectChooserBase::EIteratorStatus::Stop)
+			if (const FObjectChooserBase* SelectedResult = (*ResultsArray)[SelectedIndexData.Index].GetPtr<FObjectChooserBase>())
 			{
-				DeinitializeScratchAreas();
-				TRACE_CHOOSER_EVALUATION(Chooser, Context, SelectedIndexData.Index);
-#if WITH_EDITOR
-				Chooser->SetDebugSelectedRow(SelectedIndexData.Index);
-#endif
-				return FObjectChooserBase::EIteratorStatus::Stop;
+				FObjectChooserBase::EIteratorStatus Status = SelectedResult->ChooseMulti(Context, Callback);
+				
+				if (Status != FObjectChooserBase::EIteratorStatus::Failed)
+				{
+					bAnyRowSucceeded = true;
+				}
+	
+				if (Status == FObjectChooserBase::EIteratorStatus::Stop)
+				{
+					DeinitializeScratchAreas();
+					TRACE_CHOOSER_EVALUATION(Chooser, Context, SelectedIndexData.Index);
+	#if WITH_EDITOR
+					Chooser->SetDebugSelectedRow(SelectedIndexData.Index);
+	#endif
+					return FObjectChooserBase::EIteratorStatus::Stop;
+				}
 			}
 		}
 	}
diff --git a/Engine/Plugins/Chooser/Source/Chooser/Private/ChooserFunctionLibrary.cpp b/Engine/Plugins/Chooser/Source/Chooser/Private/ChooserFunctionLibrary.cpp
--- a/Engine/Plugins/Chooser/Source/Chooser/Private/ChooserFunctionLibrary.cpp
+++ b/Engine/Plugins/Chooser/Source/Chooser/Private/ChooserFunctionLibrary.cpp
@@ -387,9 +387,16 @@ TArray<TSoftObjectPtr<UObject>> UChooserFunctionLibrary::EvaluateObjectChooserBa
 
 	if (const FObjectChooserBase* ObjectChooserPtr = ObjectChooser.GetPtr<FObjectChooserBase>())
 	{
-		ObjectChooserPtr->ChooseMulti(Context, FObjectChooserBase::FObjectChooserSoftObjectIteratorCallback::CreateLambda([&Result, bResultIsClass, ObjectClass](const TSoftObjectPtr<UObject>& InResult)
+		ObjectChooserPtr->ChooseMulti(Context, FObjectChooserBase::FObjectChooserSoftObjectIteratorCallback::CreateLambda([&Result, &Context](const TSoftObjectPtr<UObject>& InResult)
 		{
 			Result.Add(InResult);
+				
+			// todo: move this into a context function
+			for(auto& OutputArray : Context.OutputArrays)
+			{
+				OutputArray.OutputValues.Add(FInstancedStruct(Context.Params[OutputArray.SourceParamIndex]));
+			}	
+				
 			return FObjectChooserBase::EIteratorStatus::Continue;
 		}));
 	}
diff --git a/Engine/Plugins/Chooser/Source/ProxyTable/Internal/LookupProxy.h b/Engine/Plugins/Chooser/Source/ProxyTable/Internal/LookupProxy.h
--- a/Engine/Plugins/Chooser/Source/ProxyTable/Internal/LookupProxy.h
+++ b/Engine/Plugins/Chooser/Source/ProxyTable/Internal/LookupProxy.h
@@ -54,10 +54,12 @@ USTRUCT(meta=(Hidden))
 struct FLookupProxyWithOverrideTable : public FObjectChooserBase
 {
 	GENERATED_BODY()
+	UE_API virtual EIteratorStatus ChooseMulti(FChooserEvaluationContext& ContextData, FObjectChooserSoftObjectIteratorCallback Callback) const final override;
 	UE_API virtual EIteratorStatus ChooseMulti(FChooserEvaluationContext &Context, FObjectChooserIteratorCallback Callback) const final override;
 	UE_API virtual UObject* ChooseObject(FChooserEvaluationContext& Context) const final override;
-	
-	public:
+	UE_API virtual void ChooseObject(FChooserEvaluationContext& Context, TSoftObjectPtr<UObject>& Result) const override;
+
+public:
 	
 	UPROPERTY(EditAnywhere, Category="Parameters")
 	TObjectPtr<UProxyAsset> Proxy;
diff --git a/Engine/Plugins/Chooser/Source/ProxyTable/Private/LookupProxy.cpp b/Engine/Plugins/Chooser/Source/ProxyTable/Private/LookupProxy.cpp
--- a/Engine/Plugins/Chooser/Source/ProxyTable/Private/LookupProxy.cpp
+++ b/Engine/Plugins/Chooser/Source/ProxyTable/Private/LookupProxy.cpp
@@ -65,6 +65,24 @@ UObject* FLookupProxyWithOverrideTable::ChooseObject(FChooserEvaluationContext&
 	return nullptr;
 }
 
+void FLookupProxyWithOverrideTable::ChooseObject(FChooserEvaluationContext& Context, TSoftObjectPtr<>& Result) const
+{
+	if (Proxy && OverrideProxyTable)
+   	{
+   		OverrideProxyTable->FindProxyObject(Proxy->Guid, Context, Result);
+   	}
+}
+
+FObjectChooserBase::EIteratorStatus FLookupProxyWithOverrideTable::ChooseMulti(FChooserEvaluationContext& Context,
+                                                                               FObjectChooserSoftObjectIteratorCallback Callback) const
+{
+	if (Proxy && OverrideProxyTable)
+   	{
+   		return OverrideProxyTable->FindProxyObjectMulti(Proxy->Guid, Context, Callback);
+   	}
+   	return FObjectChooserBase::EIteratorStatus::Failed;
+}
+
 FObjectChooserBase::EIteratorStatus FLookupProxyWithOverrideTable::ChooseMulti(FChooserEvaluationContext& Context, FObjectChooserIteratorCallback Callback) const
 {
 	if (Proxy && OverrideProxyTable)
diff --git a/Engine/Plugins/Chooser/Source/ProxyTable/Private/ProxyTable.cpp b/Engine/Plugins/Chooser/Source/ProxyTable/Private/ProxyTable.cpp
--- a/Engine/Plugins/Chooser/Source/ProxyTable/Private/ProxyTable.cpp
+++ b/Engine/Plugins/Chooser/Source/ProxyTable/Private/ProxyTable.cpp
@@ -243,31 +243,71 @@ UObject* UProxyTable::FindProxyObject(const FGuid& Key, FChooserEvaluationContex
 	if (FoundIndex != INDEX_NONE)
 	{
 		const FRuntimeProxyValue& EntryValueData = RuntimeValues[FoundIndex];
-		const FObjectChooserBase &EntryValue = EntryValueData.Value.Get<const FObjectChooserBase>();
-
-		UObject* Result = EntryValue.ChooseObject(Context);
 
 		OutputStructData(this, EntryValueData, Context);
 
-		return Result;
+		if (const FObjectChooserBase *EntryValue = EntryValueData.Value.GetPtr<const FObjectChooserBase>())
+		{
+			return EntryValue->ChooseObject(Context);
+		}
+
+		return nullptr;
 	}
 	
 	return nullptr;
 }
 
-FObjectChooserBase::EIteratorStatus UProxyTable::FindProxyObjectMulti(const FGuid& Key, FChooserEvaluationContext &Context, FObjectChooserBase::FObjectChooserIteratorCallback Callback) const
+void UProxyTable::FindProxyObject(const FGuid& Key, FChooserEvaluationContext& Context, TSoftObjectPtr<UObject>& Result) const
+{
+	const int FoundIndex = Algo::BinarySearch(Keys, Key);
+	if (FoundIndex != INDEX_NONE)
+	{
+		const FRuntimeProxyValue& EntryValueData = RuntimeValues[FoundIndex];
+		
+		OutputStructData(this, EntryValueData, Context);
+
+		if (const FObjectChooserBase *EntryValue = EntryValueData.Value.GetPtr<const FObjectChooserBase>())
+		{
+			EntryValue->ChooseObject(Context, Result);
+		}
+	}
+}
+
+FObjectChooserBase::EIteratorStatus UProxyTable::FindProxyObjectMulti(const FGuid& Key, FChooserEvaluationContext &Context, FObjectChooserBase::FObjectChooserSoftObjectIteratorCallback Callback) const
 {
 	const int FoundIndex = Algo::BinarySearch(Keys, Key);
     if (FoundIndex != INDEX_NONE)
     {
     	const FRuntimeProxyValue& EntryValueData = RuntimeValues[FoundIndex];
-    	const FObjectChooserBase &EntryValue = EntryValueData.Value.Get<const FObjectChooserBase>();
-    
-    	FObjectChooserBase::EIteratorStatus Result = EntryValue.ChooseMulti(Context, Callback);
+    	
+    	OutputStructData(this, EntryValueData, Context);
+    	
+    	if (const FObjectChooserBase *EntryValue = EntryValueData.Value.GetPtr<const FObjectChooserBase>())
+    	{
+    		return EntryValue->ChooseMulti(Context, Callback);
+    	}
     
+    	return FObjectChooserBase::EIteratorStatus::Failed;
+    }
+    	
+    return FObjectChooserBase::EIteratorStatus::Failed;
+}
+
+FObjectChooserBase::EIteratorStatus UProxyTable::FindProxyObjectMulti(const FGuid& Key, FChooserEvaluationContext &Context, FObjectChooserBase::FObjectChooserIteratorCallback Callback) const
+{
+	const int FoundIndex = Algo::BinarySearch(Keys, Key);
+    if (FoundIndex != INDEX_NONE)
+    {
+    	const FRuntimeProxyValue& EntryValueData = RuntimeValues[FoundIndex];
+    	
 		OutputStructData(this, EntryValueData, Context);
+    	
+    	if (const FObjectChooserBase *EntryValue = EntryValueData.Value.GetPtr<const FObjectChooserBase>())
+    	{
+    		return EntryValue->ChooseMulti(Context, Callback);
+    	}
     
-    	return Result;
+    	return FObjectChooserBase::EIteratorStatus::Failed;
     }
     	
     return FObjectChooserBase::EIteratorStatus::Failed;
diff --git a/Engine/Plugins/Chooser/Source/ProxyTable/Public/ProxyTable.h b/Engine/Plugins/Chooser/Source/ProxyTable/Public/ProxyTable.h
--- a/Engine/Plugins/Chooser/Source/ProxyTable/Public/ProxyTable.h
+++ b/Engine/Plugins/Chooser/Source/ProxyTable/Public/ProxyTable.h
@@ -87,6 +87,9 @@ public:
 
 	UE_API FObjectChooserBase::EIteratorStatus FindProxyObjectMulti(const FGuid& Key, FChooserEvaluationContext &Context, FObjectChooserBase::FObjectChooserIteratorCallback Callback) const;
 	UE_API UObject* FindProxyObject(const FGuid& Key, FChooserEvaluationContext& Context) const;
+	void FindProxyObject(const FGuid& Key, FChooserEvaluationContext& Context, TSoftObjectPtr<UObject>& Result) const;
+	FObjectChooserBase::EIteratorStatus FindProxyObjectMulti(const FGuid& Key, FChooserEvaluationContext& Context,
+	                                                         FObjectChooserBase::FObjectChooserSoftObjectIteratorCallback Callback) const;
 	UE_API FObjectChooserBase::EIteratorStatus IterateProxyObjects(const FGuid& Key, FChooserEvaluationContext& Context, FObjectChooserBase::FObjectChooserIteratorCallback Callback) const;
 	UE_API FObjectChooserBase::EIteratorStatus IterateProxyObjects(const FGuid& Key, FObjectChooserBase::FObjectChooserIteratorCallback Callback) const;
 
diff --git a/Engine/Plugins/Chooser/Source/ProxyTableUncooked/Private/EvaluateProxyNode.cpp b/Engine/Plugins/Chooser/Source/ProxyTableUncooked/Private/EvaluateProxyNode.cpp
--- a/Engine/Plugins/Chooser/Source/ProxyTableUncooked/Private/EvaluateProxyNode.cpp
+++ b/Engine/Plugins/Chooser/Source/ProxyTableUncooked/Private/EvaluateProxyNode.cpp
@@ -296,7 +296,6 @@ FText UK2Node_EvaluateProxy::GetMenuCategory() const
 
 
 
-
 //----------------------------------------------------------------------------------------------
 // UK2Node_EvaluateProxy2
 // New Implementation of EvaluateChooser with support for passing in/out multiple objects and structs
@@ -431,13 +430,19 @@ void UK2Node_EvaluateProxy2::AllocateDefaultPins()
 
 					if (StructContext.Direction == EContextObjectDirection::Write || StructContext.Direction == EContextObjectDirection::ReadWrite)
 					{
+						const bool bArrayOutput = StructOutputMode == EEvaluateProxyStructOutputMode::Default && Mode == EEvaluateProxyMode::AllResults;
 						if (UEdGraphPin* Pin = FindPin(StructContext.Struct.GetFName(), EGPD_Output))
 						{
 							Pin->PinType.PinSubCategoryObject = StructContext.Struct;
+							Pin->PinType.ContainerType = bArrayOutput ? EPinContainerType::Array : EPinContainerType::None;
 						}
 						else
 						{
-							CreatePin(EGPD_Output, UEdGraphSchema_K2::PC_Struct, StructContext.Struct, StructContext.Struct.GetFName(), StructContextPinCreationParams);
+							UEdGraphNode::FCreatePinParams PinParams;
+							PinParams.ContainerType = bArrayOutput ? EPinContainerType::Array : EPinContainerType::None;
+							PinParams.ValueTerminalType.TerminalCategory = UEdGraphSchema_K2::PC_Struct;
+							
+							CreatePin(EGPD_Output, PinParams.ValueTerminalType.TerminalCategory, StructContext.Struct, StructContext.Struct.GetFName(), PinParams);
 						}
 					}
 
@@ -479,7 +484,8 @@ void UK2Node_EvaluateProxy2::AllocateDefaultPins()
 
 	if (UEdGraphPin* ResultPin = FindPin(TEXT("Result"), EGPD_Output))
 	{
-		ResultPin->PinType.PinCategory = bResultIsClass ? UEdGraphSchema_K2::PC_Class : UEdGraphSchema_K2::PC_Object;
+		ResultPin->PinType.PinCategory = bReturnSoftObjectReference ? UEdGraphSchema_K2::PC_SoftObject
+								: bResultIsClass ? UEdGraphSchema_K2::PC_Class : UEdGraphSchema_K2::PC_Object;
 		ResultPin->PinType.PinSubCategoryObject = ResultType;
 		ResultPin->PinType.ContainerType = (Mode == EEvaluateProxyMode::AllResults) ? EPinContainerType::Array : EPinContainerType::None;
 	}
@@ -487,8 +493,8 @@ void UK2Node_EvaluateProxy2::AllocateDefaultPins()
 	{
 		UEdGraphNode::FCreatePinParams PinParams;
 		PinParams.ContainerType = (Mode == EEvaluateProxyMode::AllResults) ? EPinContainerType::Array : EPinContainerType::None;
-		PinParams.ValueTerminalType.TerminalCategory =  bResultIsClass ? UEdGraphSchema_K2::PC_Class : UEdGraphSchema_K2::PC_Object;
-				
+		PinParams.ValueTerminalType.TerminalCategory = bReturnSoftObjectReference ? UEdGraphSchema_K2::PC_SoftObject
+							: bResultIsClass ? UEdGraphSchema_K2::PC_Class : UEdGraphSchema_K2::PC_Object;	
 		CreatePin(EGPD_Output, PinParams.ValueTerminalType.TerminalCategory, ResultType, TEXT("Result"), PinParams);
 	}
 }
@@ -514,13 +520,22 @@ FText UK2Node_EvaluateProxy2::GetPinDisplayName(const UEdGraphPin* Pin) const
 	return FText::FromName(Pin->PinName);
 }
 
+void UK2Node_EvaluateProxy2::PostPlacedNewNode()
+{
+	// Class default value (Single Struct output mode) is there to not break existing content.
+	// for new nodes, set the default to match the object output mode
+	StructOutputMode = EEvaluateProxyStructOutputMode::Default;
+}
+
 void UK2Node_EvaluateProxy2::PostEditChangeProperty(struct FPropertyChangedEvent& PropertyChangedEvent)
 {
-	if (PropertyChangedEvent.Property->GetName() == "Proxy")
+	FString PropertyName = PropertyChangedEvent.Property->GetName();
+	
+	if (PropertyName == "Proxy")
 	{
 		ProxyChanged();
 	}
-	if (PropertyChangedEvent.Property->GetName() == "Mode")
+	if (PropertyName == "Mode" || PropertyName == "StructOutputMode" || PropertyName == "bReturnSoftObjectReference")
 	{
 		ReconstructNode();
 	}
@@ -582,14 +597,28 @@ void UK2Node_EvaluateProxy2::ExpandNode(class FKismetCompilerContext& CompilerCo
 	{
 		UK2Node_CallFunction* CallFunction = CompilerContext.SpawnIntermediateNode<UK2Node_CallFunction>(this, SourceGraph);
 		CompilerContext.MessageLog.NotifyIntermediateObjectCreation(CallFunction, this);
-
-		if (Mode == EEvaluateProxyMode::AllResults)
+		
+		if (bReturnSoftObjectReference)
 		{
-			CallFunction->SetFromFunction(UChooserFunctionLibrary::StaticClass()->FindFunctionByName(GET_MEMBER_NAME_CHECKED(UChooserFunctionLibrary, EvaluateObjectChooserBaseMulti)));
+			if (Mode == EEvaluateProxyMode::AllResults)
+			{
+				CallFunction->SetFromFunction(UChooserFunctionLibrary::StaticClass()->FindFunctionByName(GET_MEMBER_NAME_CHECKED(UChooserFunctionLibrary, EvaluateObjectChooserBaseMultiSoft)));
+			}
+			else
+			{
+				CallFunction->SetFromFunction(UChooserFunctionLibrary::StaticClass()->FindFunctionByName(GET_MEMBER_NAME_CHECKED(UChooserFunctionLibrary, EvaluateObjectChooserBaseSoft)));
+			}
 		}
 		else
 		{
-			CallFunction->SetFromFunction(UChooserFunctionLibrary::StaticClass()->FindFunctionByName(GET_MEMBER_NAME_CHECKED(UChooserFunctionLibrary, EvaluateObjectChooserBase)));
+			if (Mode == EEvaluateProxyMode::AllResults)
+			{
+				CallFunction->SetFromFunction(UChooserFunctionLibrary::StaticClass()->FindFunctionByName(GET_MEMBER_NAME_CHECKED(UChooserFunctionLibrary, EvaluateObjectChooserBaseMulti)));
+			}
+			else
+			{
+				CallFunction->SetFromFunction(UChooserFunctionLibrary::StaticClass()->FindFunctionByName(GET_MEMBER_NAME_CHECKED(UChooserFunctionLibrary, EvaluateObjectChooserBase)));
+			}
 		}
 		CallFunction->AllocateDefaultPins();
 		
@@ -608,8 +637,11 @@ void UK2Node_EvaluateProxy2::ExpandNode(class FKismetCompilerContext& CompilerCo
 		SelfNode->AllocateDefaultPins();
 		UEdGraphPin* SelfPin = SelfNode->FindPin(TEXT("self"));
 		
+		int OutputStructArrayIndex = 0;
+		
 		if (Proxy)
 		{
+			bool bFoundObject = false;
 			int ContextDataCount = Proxy->ContextData.Num();
 			for(int ContextDataIndex = 0; ContextDataIndex < ContextDataCount; ContextDataIndex++)
 			{
@@ -619,6 +651,7 @@ void UK2Node_EvaluateProxy2::ExpandNode(class FKismetCompilerContext& CompilerCo
 					const UScriptStruct* EntryType = ContextDataEntry.GetScriptStruct();
 					if (EntryType == FContextObjectTypeClass::StaticStruct())
 					{
+						bFoundObject = true;
 						const FContextObjectTypeClass& ClassContext = ContextDataEntry.Get<FContextObjectTypeClass>();
 						if (ClassContext.Class)
 						{
@@ -654,7 +687,16 @@ void UK2Node_EvaluateProxy2::ExpandNode(class FKismetCompilerContext& CompilerCo
 						{
 							UK2Node_CallFunction* AddStructFunction = CompilerContext.SpawnIntermediateNode<UK2Node_CallFunction>(this, SourceGraph);
 							CompilerContext.MessageLog.NotifyIntermediateObjectCreation(AddStructFunction, this);
-							AddStructFunction->SetFromFunction(UChooserFunctionLibrary::StaticClass()->FindFunctionByName(GET_MEMBER_NAME_CHECKED(UChooserFunctionLibrary, AddChooserStructInput)));
+							
+							if (StructContext.Direction == EContextObjectDirection::Read || Mode == EEvaluateProxyMode::FirstResult || StructOutputMode == EEvaluateProxyStructOutputMode::SingleOutputStruct)
+							{
+								AddStructFunction->SetFromFunction(UChooserFunctionLibrary::StaticClass()->FindFunctionByName(GET_MEMBER_NAME_CHECKED(UChooserFunctionLibrary, AddChooserStructInput)));
+							}
+							else
+							{
+								AddStructFunction->SetFromFunction(UChooserFunctionLibrary::StaticClass()->FindFunctionByName(GET_MEMBER_NAME_CHECKED(UChooserFunctionLibrary, AddChooserStructInputOutput)));
+							}
+							
 							AddStructFunction->AllocateDefaultPins();
 							ContextStructPin->MakeLinkTo(AddStructFunction->FindPin(FName("Context")));
 
@@ -699,10 +741,20 @@ void UK2Node_EvaluateProxy2::ExpandNode(class FKismetCompilerContext& CompilerCo
 										// set up struct output pin
 										UK2Node_CallFunction* GetStructFunction = CompilerContext.SpawnIntermediateNode<UK2Node_CallFunction>(this, SourceGraph);
 										CompilerContext.MessageLog.NotifyIntermediateObjectCreation(GetStructFunction, this);
-										GetStructFunction->SetFromFunction(UChooserFunctionLibrary::StaticClass()->FindFunctionByName(GET_MEMBER_NAME_CHECKED(UChooserFunctionLibrary, GetChooserStructOutput)));
-										GetStructFunction->AllocateDefaultPins();
-										GetStructFunction->FindPin(FName("Index"))->DefaultValue = FString::FromInt(ContextDataIndex);
-									
+										
+										if (Mode == EEvaluateProxyMode::FirstResult || StructOutputMode == EEvaluateProxyStructOutputMode::SingleOutputStruct)
+										{
+											GetStructFunction->SetFromFunction(UChooserFunctionLibrary::StaticClass()->FindFunctionByName(GET_MEMBER_NAME_CHECKED(UChooserFunctionLibrary, GetChooserStructOutput)));
+											GetStructFunction->AllocateDefaultPins();
+											GetStructFunction->FindPin(FName("Index"))->DefaultValue = FString::FromInt(ContextDataIndex);
+										}
+										else
+										{
+											GetStructFunction->SetFromFunction(UChooserFunctionLibrary::StaticClass()->FindFunctionByName(GET_MEMBER_NAME_CHECKED(UChooserFunctionLibrary, GetChooserStructOutputArray)));
+											GetStructFunction->AllocateDefaultPins();
+											GetStructFunction->FindPin(FName("Index"))->DefaultValue = FString::FromInt(OutputStructArrayIndex);
+										}
+
 										ContextStructPin->MakeLinkTo(GetStructFunction->FindPin(FName("Context")));
 										UEdGraphPin* ValuePin = GetStructFunction->FindPin(FName("Value"));
 										// not sure why this isn't automatically happening:
@@ -712,11 +764,29 @@ void UK2Node_EvaluateProxy2::ExpandNode(class FKismetCompilerContext& CompilerCo
 										CompilerContext.MovePinLinksToIntermediate(*Pin, *ValuePin);
 									}
 								}
+								OutputStructArrayIndex ++;
 							}
 						}
 					}
 				}
 			}
+			
+			if (!bFoundObject)
+			{
+				// add Self reference to the end of the context, for debugging purposes
+				UK2Node_CallFunction* AddObjectFunction = CompilerContext.SpawnIntermediateNode<UK2Node_CallFunction>(this, SourceGraph);
+				CompilerContext.MessageLog.NotifyIntermediateObjectCreation(AddObjectFunction, this);
+				AddObjectFunction->SetFromFunction(UChooserFunctionLibrary::StaticClass()->FindFunctionByName(GET_MEMBER_NAME_CHECKED(UChooserFunctionLibrary, AddChooserObjectInput)));
+				AddObjectFunction->AllocateDefaultPins();
+				ContextStructPin->MakeLinkTo(AddObjectFunction->FindPin(FName("Context")));
+	
+				PreviousNodeExecOutput->MakeLinkTo(AddObjectFunction->GetExecPin());
+				PreviousNodeExecOutput = AddObjectFunction->GetThenPin();
+									
+				UEdGraphPin* AddObjectPin = AddObjectFunction->FindPin(FName("Object"));
+				SelfPin->MakeLinkTo(AddObjectPin);
+			}
+			
 
 			if (PreviousNodeExecOutput)
 			{
diff --git a/Engine/Plugins/Chooser/Source/ProxyTableUncooked/Private/EvaluateProxyNode.h b/Engine/Plugins/Chooser/Source/ProxyTableUncooked/Private/EvaluateProxyNode.h
--- a/Engine/Plugins/Chooser/Source/ProxyTableUncooked/Private/EvaluateProxyNode.h
+++ b/Engine/Plugins/Chooser/Source/ProxyTableUncooked/Private/EvaluateProxyNode.h
@@ -14,6 +14,20 @@
 
 #include "EvaluateProxyNode.generated.h"
 
+UENUM()
+enum class EEvaluateProxyMode
+{
+	FirstResult,
+	AllResults
+};
+
+UENUM()
+enum class EEvaluateProxyStructOutputMode
+{
+	Default,
+	SingleOutputStruct,
+};
+
 UCLASS(MinimalAPI, Hidden)
 class UK2Node_EvaluateProxy : public UK2Node
 {
@@ -67,13 +81,6 @@ private:
 	FText NodeTooltip;
 };
 
-UENUM()
-enum class EEvaluateProxyMode
-{
-	FirstResult,
-	AllResults
-};
-
 //----------------------------------------------------------------------------------------------
 // UK2Node_EvaluateChooser2
 // New Implementation of EvaluateChooser with support for passing in/out multiple objects and structs
@@ -98,6 +105,7 @@ class UK2Node_EvaluateProxy2 : public UK2Node
 	virtual void PinTypeChanged(UEdGraphPin* Pin) override;
 	virtual FText GetTooltipText() const override;
 	virtual FText GetPinDisplayName(const UEdGraphPin* Pin) const override;
+	virtual void PostPlacedNewNode() override;
 	//~ End UEdGraphNode Interface.
 
 	//~ Begin UK2Node Interface.
@@ -125,9 +133,17 @@ private:
 	UPROPERTY(EditAnywhere, Category = "Proxy")
 	TObjectPtr<UProxyAsset> Proxy;
 
-	UPROPERTY(EditAnywhere, Category = "Chooser")
+	UPROPERTY(EditAnywhere, Category = "Proxy")
 	EEvaluateProxyMode Mode;
 
+	// This property exist for backwards compatibility with old content (where only single output struct was supported).
+	// New content should use "Default" setting to stay in sync with the regular output mode.
+	UPROPERTY(EditAnywhere, Category = "Proxy")
+	EEvaluateProxyStructOutputMode StructOutputMode = EEvaluateProxyStructOutputMode::SingleOutputStruct;
+	
+	UPROPERTY(EditAnywhere, Category = "Proxy")
+	bool bReturnSoftObjectReference = false;
+
 	/** Tooltip text for this node. */
 	FText NodeTooltip;
 };
