Fix crashing in ChooserTableEditor details panel edits, and Enum Column details panel edits not being applied

#Synced-CL 52174097

diff --git a/Engine/Plugins/Chooser/Source/Chooser/Private/EnumColumn.cpp b/Engine/Plugins/Chooser/Source/Chooser/Private/EnumColumn.cpp
--- a/Engine/Plugins/Chooser/Source/Chooser/Private/EnumColumn.cpp
+++ b/Engine/Plugins/Chooser/Source/Chooser/Private/EnumColumn.cpp
@@ -121,7 +121,7 @@ void FEnumColumn::Filter(FChooserEvaluationContext& Context, const FChooserIndex
 }
 
 #if WITH_EDITOR
-	void FEnumColumn::AddToDetails (FInstancedPropertyBag& PropertyBag, int32 ColumnIndex, int32 RowIndex)
+	void FEnumColumn::AddToDetails(FInstancedPropertyBag& PropertyBag, int32 ColumnIndex, int32 RowIndex)
 	{
 		FText DisplayName;
 		InputValue.Get<FChooserParameterBase>().GetDisplayName(DisplayName);
@@ -152,10 +152,14 @@ void FEnumColumn::Filter(FChooserEvaluationContext& Context, const FChooserIndex
 		if (FStructView* StructView = Result.TryGetValue())
 		{
 			const UScriptStruct* StructDefinition = StructView->GetScriptStruct();
-			if (const FEnumProperty* ValueProperty = CastField<FEnumProperty>(StructDefinition->FindPropertyByName("Value")))
+			if (const FEnumProperty* ValueEnumProperty = CastField<FEnumProperty>(StructDefinition->FindPropertyByName("Value")))
 			{
-				ValueProperty->GetValue_InContainer(StructView->GetMemory(), &RowValues[RowIndex].Value);
+				ValueEnumProperty->GetValue_InContainer(StructView->GetMemory(), &RowValues[RowIndex].Value);
 			}
+			else if (const FByteProperty* ValueByteProperty = CastField<FByteProperty>(StructDefinition->FindPropertyByName("Value")))
+            {
+            	ValueByteProperty->GetValue_InContainer(StructView->GetMemory(), &RowValues[RowIndex].Value);
+            }
 			
 			if (const FEnumProperty* ComparisonProperty = CastField<FEnumProperty>(StructDefinition->FindPropertyByName("Comparison")))
 			{
diff --git a/Engine/Plugins/Chooser/Source/ChooserEditor/Private/ChooserTableViewModel.cpp b/Engine/Plugins/Chooser/Source/ChooserEditor/Private/ChooserTableViewModel.cpp
--- a/Engine/Plugins/Chooser/Source/ChooserEditor/Private/ChooserTableViewModel.cpp
+++ b/Engine/Plugins/Chooser/Source/ChooserEditor/Private/ChooserTableViewModel.cpp
@@ -125,57 +125,66 @@ void FChooserTableViewModel::OnObjectsTransacted(UObject* Object, const FTransac
 		// if this is for the chooser we're editing
 		if (GetChooser() == RowDetails->Chooser)
 		{
-			if (RowDetails->Chooser->ResultsStructs.IsValidIndex(RowDetails->Row))
+			QueueRowTransaction(RowDetails);
+			
+		}
+	}
+}
+	
+bool FChooserTableViewModel::HandleTicker(float DeltaTime)
+{
+	if (!QueuedRowTransactions.IsEmpty())
+	{
+		for (TWeakObjectPtr<UChooserRowDetails> RowDetails : QueuedRowTransactions)
+		{
+			if (RowDetails.IsValid())
 			{
-				// copy all the values over
-				TValueOrError<FStructView, EPropertyBagResult> Result = RowDetails->Properties.GetValueStruct("Result", FInstancedStruct::StaticStruct());
-				if (Result.IsValid())
+				if (RowDetails->Chooser->ResultsStructs.IsValidIndex(RowDetails->Row))
 				{
-					RowDetails->Chooser->ResultsStructs[RowDetails->Row] = Result.GetValue().Get<FInstancedStruct>();
-				}
+					// copy all the values over
+					TValueOrError<FStructView, EPropertyBagResult> Result = RowDetails->Properties.GetValueStruct("Result", FInstancedStruct::StaticStruct());
+					if (Result.IsValid())
+					{
+						RowDetails->Chooser->ResultsStructs[RowDetails->Row] = Result.GetValue().Get<FInstancedStruct>();
+					}
 
-				int32 ColumnIndex = 0;
-				for (FInstancedStruct& ColumnData : RowDetails->Chooser->ColumnsStructs)
-				{
-					FChooserColumnBase& Column = ColumnData.GetMutable<FChooserColumnBase>();
-					Column.SetFromDetails(RowDetails->Properties, ColumnIndex, RowDetails->Row);
-					ColumnIndex++;
-				}
+					int32 ColumnIndex = 0;
+					for (FInstancedStruct& ColumnData : RowDetails->Chooser->ColumnsStructs)
+					{
+						FChooserColumnBase& Column = ColumnData.GetMutable<FChooserColumnBase>();
+						Column.SetFromDetails(RowDetails->Properties, ColumnIndex, RowDetails->Row);
+						ColumnIndex++;
+					}
 			
-				TValueOrError<bool, EPropertyBagResult> DisabledResult = RowDetails->Properties.GetValueBool("Disabled");
-				if (DisabledResult.IsValid())
-				{
-					RowDetails->Chooser->DisabledRows[RowDetails->Row] = DisabledResult.GetValue();
+					TValueOrError<bool, EPropertyBagResult> DisabledResult = RowDetails->Properties.GetValueBool("Disabled");
+					if (DisabledResult.IsValid())
+					{
+						RowDetails->Chooser->DisabledRows[RowDetails->Row] = DisabledResult.GetValue();
+					}
 				}
-			}
-			else if (RowDetails->Row == ColumnWidget_SpecialIndex_Fallback)
-			{
-				TValueOrError<FStructView, EPropertyBagResult> Result = RowDetails->Properties.GetValueStruct("Result", FInstancedStruct::StaticStruct());
-				if (Result.IsValid())
+				else if (RowDetails->Row == ColumnWidget_SpecialIndex_Fallback)
 				{
-					RowDetails->Chooser->FallbackResult = Result.GetValue().Get<FInstancedStruct>();
-				}
+					TValueOrError<FStructView, EPropertyBagResult> Result = RowDetails->Properties.GetValueStruct("Result", FInstancedStruct::StaticStruct());
+					if (Result.IsValid())
+					{
+						RowDetails->Chooser->FallbackResult = Result.GetValue().Get<FInstancedStruct>();
+					}
 				
-				int32 ColumnIndex = 0;
-				for (FInstancedStruct& ColumnData : RowDetails->Chooser->ColumnsStructs)
-				{
-					FChooserColumnBase& Column = ColumnData.GetMutable<FChooserColumnBase>();
-					Column.SetFromDetails(RowDetails->Properties, ColumnIndex, RowDetails->Row);
-					ColumnIndex++;
+					int32 ColumnIndex = 0;
+					for (FInstancedStruct& ColumnData : RowDetails->Chooser->ColumnsStructs)
+					{
+						FChooserColumnBase& Column = ColumnData.GetMutable<FChooserColumnBase>();
+						Column.SetFromDetails(RowDetails->Properties, ColumnIndex, RowDetails->Row);
+						ColumnIndex++;
+					}
 				}
 			}
-			RowDetails->Chooser->MarkPackageDirty();
-			QueueRefreshAll();
 		}
-	}
-}
-	
-bool FChooserTableViewModel::HandleTicker(float DeltaTime)
-{
-	if (bQueueRefreshAll)
-	{
+		
+		GetChooser()->MarkPackageDirty();
 		RefreshAll();
-		bQueueRefreshAll = false;
+		
+		QueuedRowTransactions.Empty();
 	}
 	return true;
 }
@@ -237,11 +246,6 @@ void FChooserTableViewModel::RefreshAll()
 	RefreshAssetRegistry();
 }
 
-void FChooserTableViewModel::QueueRefreshAll()
-{
-	bQueueRefreshAll = true;
-}
-
 bool FChooserTableViewModel::MatchesContext(const FTransactionContext& InContext, const TArray<TPair<UObject*, FTransactionObjectEvent>>& TransactionObjectContexts) const
 {
 	TArray<UObject*> ContainedObjects;
@@ -354,6 +358,11 @@ void FChooserTableViewModel::UpdateTableColumns()
 	RefreshTableColumnsDelegate.ExecuteIfBound();
 }
 
+void FChooserTableViewModel::QueueRowTransaction(UChooserRowDetails* RowDetails)
+{
+	QueuedRowTransactions.Add(RowDetails);
+}
+
 void FChooserTableViewModel::AddColumn(const UScriptStruct* ColumnType)
 {
 	FSlateApplication::Get().DismissAllMenus();
diff --git a/Engine/Plugins/Chooser/Source/ChooserEditor/Private/ChooserTableViewModel.h b/Engine/Plugins/Chooser/Source/ChooserEditor/Private/ChooserTableViewModel.h
--- a/Engine/Plugins/Chooser/Source/ChooserEditor/Private/ChooserTableViewModel.h
+++ b/Engine/Plugins/Chooser/Source/ChooserEditor/Private/ChooserTableViewModel.h
@@ -85,9 +85,6 @@ namespace UE::ChooserEditor
 		virtual void NotifyPreChange( FProperty* PropertyAboutToChange ) override;
 		virtual void NotifyPostChange( const FPropertyChangedEvent& PropertyChangedEvent, FProperty* PropertyThatChanged) override;
 
-
-		void QueueRefreshAll();
-	
 		int32 GetNumRows() const;
 
 		bool AreRowsDisabled(TConstArrayView<int32> RowIndices) const;
@@ -188,6 +185,7 @@ namespace UE::ChooserEditor
 		void UpdateTableRows();
 		void UpdateTableColumns();
 
+		void QueueRowTransaction(UChooserRowDetails* RowDetails);
 		void OnObjectsTransacted(UObject* Object, const FTransactionObjectEvent& Event);
 
 		/** FEditorUndoClient Interface */
@@ -200,7 +198,6 @@ namespace UE::ChooserEditor
 		
 		bool HandleTicker(float DeltaTime);
 
-		bool bQueueRefreshAll = false;
 		FTSTicker::FDelegateHandle TickDelegateHandle;
 		
 		// the root asset for the chooser editor
@@ -209,7 +206,7 @@ namespace UE::ChooserEditor
 		// the current (possibly nested) chooser table being edited
 		UChooserTable* CurrentChooser = nullptr;
 
-
+		TArray<TWeakObjectPtr<UChooserRowDetails>> QueuedRowTransactions;
 	};
 	
 }
\ No newline at end of file
