[ProxyTable] Adds Blueprint Function to get all ProxyTable Entries, output in a TMap

#Synced-CL 52468934

diff --git a/Engine/Plugins/Chooser/Source/ProxyTable/Private/ProxyTableFunctionLibrary.cpp b/Engine/Plugins/Chooser/Source/ProxyTable/Private/ProxyTableFunctionLibrary.cpp
--- a/Engine/Plugins/Chooser/Source/ProxyTable/Private/ProxyTableFunctionLibrary.cpp
+++ b/Engine/Plugins/Chooser/Source/ProxyTable/Private/ProxyTableFunctionLibrary.cpp
@@ -63,3 +63,35 @@ FInstancedStruct UProxyTableFunctionLibrary::MakeLookupProxyWithOverrideTable(UP
  	return Struct;
 }
 
+#if WITH_EDITORONLY_DATA
+TMap<UProxyAsset*, UObject*> UProxyTableFunctionLibrary::GetProxyTableEntries(const UProxyTable* ProxyTable)
+{
+	TMap<UProxyAsset*, UObject*> Result;
+	if (ProxyTable == nullptr)
+	{
+		return Result;
+	}
+
+	const int32 NumEntries = FMath::Min(ProxyTable->Keys.Num(), ProxyTable->RuntimeValues.Num());
+	for (int32 Index = 0; Index < NumEntries; Index++)
+	{
+		const FRuntimeProxyValue& RuntimeValue = ProxyTable->RuntimeValues[Index];
+		UProxyAsset* ProxyAsset = RuntimeValue.ProxyAsset.Get();
+		if (ProxyAsset == nullptr)
+		{
+			continue;
+		}
+
+		UObject* Value = nullptr;
+		if (const FObjectChooserBase* Chooser = RuntimeValue.Value.GetPtr<FObjectChooserBase>())
+		{
+			Value = Chooser->GetReferencedObject();
+		}
+
+		Result.Add(ProxyAsset, Value);
+	}
+
+	return Result;
+}
+#endif
+
diff --git a/Engine/Plugins/Chooser/Source/ProxyTable/Public/ProxyTableFunctionLibrary.h b/Engine/Plugins/Chooser/Source/ProxyTable/Public/ProxyTableFunctionLibrary.h
--- a/Engine/Plugins/Chooser/Source/ProxyTable/Public/ProxyTableFunctionLibrary.h
+++ b/Engine/Plugins/Chooser/Source/ProxyTable/Public/ProxyTableFunctionLibrary.h
@@ -48,6 +48,17 @@ public:
 	
 	UFUNCTION(BlueprintPure, Category = "Animation", meta = (BlueprintThreadSafe, BlueprintInternalUseOnly="true"))
 	static UE_API FInstancedStruct MakeLookupProxyWithOverrideTable(UProxyAsset* Proxy, UProxyTable* ProxyTable);
+
+#if WITH_EDITORONLY_DATA
+	/**
+	* Get all entries in a ProxyTable as a map from ProxyAsset to its resolved UObject value.
+	*
+	* @param ProxyTable				(in) The ProxyTable asset to read entries from
+	* @return Map of ProxyAsset to resolved UObject value
+	*/
+	UFUNCTION(BlueprintCallable, Category = "Animation", meta = (DevelopmentOnly))
+	static UE_API TMap<UProxyAsset*, UObject*> GetProxyTableEntries(const UProxyTable* ProxyTable);
+#endif
 };
 
 #undef UE_API
