[ProxyTable] Tweak UPROPERTY flags and Editor flags of new GetProxyTableEntries Function

#Synced-CL 52468934

diff --git a/Engine/Plugins/Chooser/Source/ProxyTable/Private/ProxyTableFunctionLibrary.cpp b/Engine/Plugins/Chooser/Source/ProxyTable/Private/ProxyTableFunctionLibrary.cpp
--- a/Engine/Plugins/Chooser/Source/ProxyTable/Private/ProxyTableFunctionLibrary.cpp
+++ b/Engine/Plugins/Chooser/Source/ProxyTable/Private/ProxyTableFunctionLibrary.cpp
@@ -63,7 +63,7 @@ FInstancedStruct UProxyTableFunctionLibrary::MakeLookupProxyWithOverrideTable(UP
  	return Struct;
 }
 
-#if WITH_EDITORONLY_DATA
+#if WITH_EDITOR
 TMap<UProxyAsset*, UObject*> UProxyTableFunctionLibrary::GetProxyTableEntries(const UProxyTable* ProxyTable)
 {
 	TMap<UProxyAsset*, UObject*> Result;
diff --git a/Engine/Plugins/Chooser/Source/ProxyTable/Public/ProxyTableFunctionLibrary.h b/Engine/Plugins/Chooser/Source/ProxyTable/Public/ProxyTableFunctionLibrary.h
--- a/Engine/Plugins/Chooser/Source/ProxyTable/Public/ProxyTableFunctionLibrary.h
+++ b/Engine/Plugins/Chooser/Source/ProxyTable/Public/ProxyTableFunctionLibrary.h
@@ -49,14 +49,15 @@ public:
 	UFUNCTION(BlueprintPure, Category = "Animation", meta = (BlueprintThreadSafe, BlueprintInternalUseOnly="true"))
 	static UE_API FInstancedStruct MakeLookupProxyWithOverrideTable(UProxyAsset* Proxy, UProxyTable* ProxyTable);
 
-#if WITH_EDITORONLY_DATA
+#if WITH_EDITOR
 	/**
 	* Get all entries in a ProxyTable as a map from ProxyAsset to its resolved UObject value.
+	* Editor-only: uses editor-only APIs (GetReferencedObject) that are not available in cooked builds.
 	*
 	* @param ProxyTable				(in) The ProxyTable asset to read entries from
 	* @return Map of ProxyAsset to resolved UObject value
 	*/
-	UFUNCTION(BlueprintCallable, Category = "Animation", meta = (DevelopmentOnly))
+	UFUNCTION(BlueprintCallable, Category = "Animation", meta = (CallInEditor = "true"))
 	static UE_API TMap<UProxyAsset*, UObject*> GetProxyTableEntries(const UProxyTable* ProxyTable);
 #endif
 };
